Chaos Cloth - Move the new Cloth Asset modules out of the Restricted area to the Engine/Plugins/Experimental folder (part 2, add files)

diff --git a/Engine/Plugins/Experimental/ChaosClothAsset/Source/ChaosClothAssetEngine/Private/ChaosClothAsset/ClothComponent.cpp b/Engine/Plugins/Experimental/ChaosClothAsset/Source/ChaosClothAssetEngine/Private/ChaosClothAsset/ClothComponent.cpp
new file mode 100644
--- /dev/null
+++ b/Engine/Plugins/Experimental/ChaosClothAsset/Source/ChaosClothAssetEngine/Private/ChaosClothAsset/ClothComponent.cpp
@@ -0,0 +1,141 @@
+// Copyright Epic Games, Inc. All Rights Reserved.
+
+#include "ChaosClothAsset/ClothComponent.h"
+#include "ChaosClothAsset/ClothAsset.h"
+#include "SkeletalRenderPublic.h"
+#include "PhysicsEngine/PhysicsAsset.h"
+
+DEFINE_LOG_CATEGORY_STATIC(LogChaosClothComponent, Log, All);
+
+#include UE_INLINE_GENERATED_CPP_BY_NAME(ClothComponent)
+
+DECLARE_STATS_GROUP(TEXT("ChaosClothComponent"), STATGROUP_ChaosClothComponent, STATCAT_Advanced);
+DECLARE_CYCLE_STAT(TEXT("SendRenderDynamicData_Concurrent"), STAT_ChaosClothComponent_SendRenderDynamicData_Concurrent, STATGROUP_ChaosClothComponent);
+DECLARE_CYCLE_STAT(TEXT("MeshObjectUpdate"), STAT_ChaosClothComponent_MeshObjectUpdate, STATGROUP_ChaosClothComponent);
+DECLARE_CYCLE_STAT(TEXT("CalcBounds"), STAT_ChaosClothComponent_CalcBounds, STATGROUP_ChaosClothComponent);
+
+UChaosClothComponent::UChaosClothComponent(const FObjectInitializer& ObjectInitializer)
+	: Super(ObjectInitializer)
+{
+	bAutoActivate = true;
+	PrimaryComponentTick.bCanEverTick = true;
+	PrimaryComponentTick.TickGroup = TG_PrePhysics;
+
+	//VisibilityBasedAnimTickOption = EVisibilityBasedAnimTickOption::AlwaysTickPoseAndRefreshBones;
+
+	SetCollisionProfileName(UCollisionProfile::NoCollision_ProfileName);
+
+	StreamingDistanceMultiplier = 1.0f;
+	//bCanHighlightSelectedSections = false;
+	CanCharacterStepUpOn = ECB_Owner;
+#if WITH_EDITORONLY_DATA
+	//SectionIndexPreview = -1;
+	//MaterialIndexPreview = -1;
+
+	//SelectedEditorSection = INDEX_NONE;
+	//SelectedEditorMaterial = INDEX_NONE;
+#endif // WITH_EDITORONLY_DATA
+	bCastCapsuleDirectShadow = false;
+	bCastCapsuleIndirectShadow = false;
+	CapsuleIndirectShadowMinVisibility = .1f;
+
+	bDoubleBufferedComponentSpaceTransforms = true;
+	//LastStreamerUpdateBoundsRadius = -1.0;
+	CurrentEditableComponentTransforms = 0;
+	CurrentReadComponentTransforms = 1;
+	bNeedToFlipSpaceBaseBuffers = false;
+	bBoneVisibilityDirty = false;
+
+	//bUpdateDeformerAtNextTick = false;
+
+	bCanEverAffectNavigation = false;
+	//MasterBoneMapCacheCount = 0;
+	bSyncAttachParentLOD = true;
+	//bIgnoreMasterPoseComponentLOD = false;
+
+	CurrentBoneTransformRevisionNumber = 0;
+
+	//ExternalInterpolationAlpha = 0.0f;
+	//ExternalDeltaTime = 0.0f;
+	//ExternalTickRate = 1;
+	//bExternalInterpolate = false;
+	//bExternalUpdate = false;
+	//bExternalEvaluationRateLimited = false;
+	//bExternalTickRateControlled = false;
+
+	bMipLevelCallbackRegistered = false;
+
+#if !(UE_BUILD_SHIPPING || UE_BUILD_TEST)
+	//bDrawDebugSkeleton = false;
+#endif
+
+	//CurrentSkinWeightProfileName = NAME_None;
+}
+
+void UChaosClothComponent::OnRegister()
+{
+	LLM_SCOPE(ELLMTag::Chaos);
+
+	Super::OnRegister();
+
+	if (GetClothAsset())
+	{
+		FSkeletalMeshLODRenderData& LODData = GetClothAsset()->GetResourceForRendering()->LODRenderData[GetPredictedLODLevel()];
+		GetClothAsset()->FillComponentSpaceTransforms(GetClothAsset()->GetRefSkeleton().GetRefBonePose(), LODData.RequiredBones, GetEditableComponentSpaceTransforms());
+
+		bNeedToFlipSpaceBaseBuffers = true; // Have updated space bases so need to flip
+		FlipEditableSpaceBases();
+		bHasValidBoneTransform = true;
+	}
+}
+
+void UChaosClothComponent::RefreshBoneTransforms(FActorComponentTickFunction* TickFunction)
+{
+	// TODO: RefreshBoneTransforms
+}
+
+void UChaosClothComponent::PostLoad()
+{
+	Super::PostLoad();
+
+#if WITH_EDITORONLY_DATA
+	PRAGMA_DISABLE_DEPRECATION_WARNINGS
+	ClothAsset = GetClothAsset();
+	PRAGMA_ENABLE_DEPRECATION_WARNINGS
+#endif
+}
+
+void UChaosClothComponent::SetClothAsset(UChaosClothAsset* InClothAsset)
+{
+	PRAGMA_DISABLE_DEPRECATION_WARNINGS
+#if WITH_EDITORONLY_DATA
+	ClothAsset = InClothAsset;
+#endif
+	SetSkinnedAsset(InClothAsset);
+	PRAGMA_ENABLE_DEPRECATION_WARNINGS
+}
+
+UChaosClothAsset* UChaosClothComponent::GetClothAsset() const
+{
+	return Cast<UChaosClothAsset>(GetSkinnedAsset());
+}
+
+#if WITH_EDITOR
+void UChaosClothComponent::PostEditChangeProperty(FPropertyChangedEvent& PropertyChangedEvent)
+{
+	// Set the skinned asset pointer with the alias pointer (must happen before the call to Super::PostEditChangeProperty)
+	if (const FProperty* const Property = PropertyChangedEvent.Property)
+	{
+		PRAGMA_DISABLE_DEPRECATION_WARNINGS
+		if (Property->GetFName() == GET_MEMBER_NAME_CHECKED(UChaosClothComponent, ClothAsset))
+		{
+			SetSkinnedAsset(ClothAsset);
+		}
+		PRAGMA_ENABLE_DEPRECATION_WARNINGS
+	}
+
+	Super::PostEditChangeProperty(PropertyChangedEvent);
+}
+#endif // WITH_EDITOR
+
+
diff --git a/Engine/Plugins/Experimental/ChaosClothAsset/Source/ChaosClothAssetEngine/Private/ChaosClothAsset/ClothPreset.cpp b/Engine/Plugins/Experimental/ChaosClothAsset/Source/ChaosClothAssetEngine/Private/ChaosClothAsset/ClothPreset.cpp
new file mode 100644
--- /dev/null
+++ b/Engine/Plugins/Experimental/ChaosClothAsset/Source/ChaosClothAssetEngine/Private/ChaosClothAsset/ClothPreset.cpp
@@ -0,0 +1,222 @@
+// Copyright Epic Games, Inc. All Rights Reserved.
+
+#include "ChaosClothAsset/ClothPreset.h"
+
+#include UE_INLINE_GENERATED_CPP_BY_NAME(ClothPreset)
+
+template<>
+EChaosClothPresetPropertyType FChaosClothPresetPropertyDescriptor::GetType<bool>()
+{
+	return EChaosClothPresetPropertyType::Boolean;
+}
+
+template<>
+EChaosClothPresetPropertyType FChaosClothPresetPropertyDescriptor::GetType<int32>()
+{
+	return EChaosClothPresetPropertyType::Integer;
+}
+
+template<>
+EChaosClothPresetPropertyType FChaosClothPresetPropertyDescriptor::GetType<float>()
+{
+	return EChaosClothPresetPropertyType::Float;
+}
+
+template<>
+EChaosClothPresetPropertyType FChaosClothPresetPropertyDescriptor::GetType<FIntVector3>()
+{
+	return EChaosClothPresetPropertyType::Vector3Integer;
+}
+
+template<>
+EChaosClothPresetPropertyType FChaosClothPresetPropertyDescriptor::GetType<FVector3f>()
+{
+	return EChaosClothPresetPropertyType::Vector3Float;
+}
+
+template<>
+EChaosClothPresetPropertyType FChaosClothPresetPropertyDescriptor::GetType<FString>()
+{
+	return EChaosClothPresetPropertyType::String;
+}
+
+UChaosClothPreset::UChaosClothPreset(const FObjectInitializer& ObjectInitializer)
+	: Super(ObjectInitializer)
+	, ClothPresetCollection(MakeShared<UE::Chaos::ClothAsset::FClothPresetCollection>())
+{
+}
+
+void UChaosClothPreset::Serialize(FArchive& Ar)
+{
+	Super::Serialize(Ar);
+
+	ClothPresetCollection->Serialize(Ar);
+}
+
+bool UChaosClothPreset::AddProperty(const FName& Name, EChaosClothPresetPropertyType Type)
+{
+	// Return a nullptr if the property already exists
+	if (PropertyDescriptors.FindByPredicate([&Name](const FChaosClothPresetPropertyDescriptor& PropertyDescriptor)->bool { return PropertyDescriptor.Name == Name; }))
+	{
+		return false;
+	}
+
+	// Add the new property
+	PropertyDescriptors.Emplace(Name, Type);
+
+	// Re-sync the collection
+	SyncCollection();
+
+	return true;
+}
+
+void UChaosClothPreset::SetProperty(const FName& Name, const FString& StringValue, bool bEnable, bool bAnimatable)
+{
+	const FChaosClothPresetPropertyDescriptor* const PropertyDescriptor = GetPropertyDescriptor(Name);
+	if (ensure(PropertyDescriptor) && ensure(PropertyDescriptor->Type == FChaosClothPresetPropertyDescriptor::GetType<FString>()))
+	{
+		SetPropertyInternal(*PropertyDescriptor, FVector3f::ZeroVector, FVector3f::ZeroVector, StringValue, bEnable, bAnimatable);
+	}
+}
+
+void UChaosClothPreset::SetPropertyStringValue(const FName& Name, const FString& StringValue)
+{
+	const FChaosClothPresetPropertyDescriptor* const PropertyDescriptor = GetPropertyDescriptor(Name);
+	if (ensure(PropertyDescriptor))
+	{
+		SyncProperty(*PropertyDescriptor);
+
+		const int32 Index = GetPropertyElementIndex(*PropertyDescriptor);
+		check(Index != INDEX_NONE);
+
+		ClothPresetCollection->StringValue[Index] = StringValue;
+	}
+}
+
+void UChaosClothPreset::SetPropertyEnable(const FName& Name, bool bEnable)
+{
+	const FChaosClothPresetPropertyDescriptor* const PropertyDescriptor = GetPropertyDescriptor(Name);
+	if (ensure(PropertyDescriptor))
+	{
+		SyncProperty(*PropertyDescriptor);
+
+		const int32 Index = GetPropertyElementIndex(*PropertyDescriptor);
+		check(Index != INDEX_NONE);
+
+		ClothPresetCollection->Enable[Index] = bEnable;
+	}
+}
+
+void UChaosClothPreset::SetPropertyAnimatable(const FName& Name, bool bAnimatable)
+{
+	const FChaosClothPresetPropertyDescriptor* const PropertyDescriptor = GetPropertyDescriptor(Name);
+	if (ensure(PropertyDescriptor))
+	{
+		SyncProperty(*PropertyDescriptor);
+
+		const int32 Index = GetPropertyElementIndex(*PropertyDescriptor);
+		check(Index != INDEX_NONE);
+
+		ClothPresetCollection->Animatable[Index] = bAnimatable;
+	}
+}
+
+const FChaosClothPresetPropertyDescriptor* UChaosClothPreset::GetPropertyDescriptor(const FName& Name) const
+{
+	return PropertyDescriptors.FindByPredicate(
+		[&Name](const FChaosClothPresetPropertyDescriptor& PropertyDescriptor)
+		{
+			return PropertyDescriptor.Name == Name;
+		});
+}
+
+int32 UChaosClothPreset::GetPropertyElementIndex(const FChaosClothPresetPropertyDescriptor& PropertyDescriptor) const
+{
+	return ClothPresetCollection->Name.Find(PropertyDescriptor.Name.ToString());
+}
+
+void UChaosClothPreset::SetPropertyInternal(
+	const FChaosClothPresetPropertyDescriptor& PropertyDescriptor,
+	const FVector3f& LowValue,
+	const FVector3f& HighValue,
+	const FString& StringValue,
+	bool bEnable,
+	bool bAnimatable)
+{
+	SyncProperty(PropertyDescriptor);
+
+	const int32 Index = GetPropertyElementIndex(PropertyDescriptor);
+	check(Index != INDEX_NONE);
+
+	ClothPresetCollection->LowValue[Index] = LowValue;
+	ClothPresetCollection->HighValue[Index] = HighValue;
+	ClothPresetCollection->StringValue[Index] = StringValue;
+	ClothPresetCollection->Enable[Index] = bEnable;
+	ClothPresetCollection->Animatable[Index] = bAnimatable;
+}
+
+void UChaosClothPreset::SetPropertyValueInternal(const FChaosClothPresetPropertyDescriptor& PropertyDescriptor, const FVector3f& Value)
+{
+	SyncProperty(PropertyDescriptor);
+
+	const int32 Index = GetPropertyElementIndex(PropertyDescriptor);
+	check(Index != INDEX_NONE);
+
+	ClothPresetCollection->LowValue[Index] = Value;
+}
+
+void UChaosClothPreset::SetPropertyValuesInternal(const FChaosClothPresetPropertyDescriptor& PropertyDescriptor, const FVector3f& LowValue, const FVector3f& HighValue)
+{
+	SyncProperty(PropertyDescriptor);
+
+	const int32 Index = GetPropertyElementIndex(PropertyDescriptor);
+	check(Index != INDEX_NONE);
+
+	ClothPresetCollection->LowValue[Index] = LowValue;
+	ClothPresetCollection->HighValue[Index] = HighValue;
+}
+
+void UChaosClothPreset::SyncProperty(const FChaosClothPresetPropertyDescriptor& PropertyDescriptor)
+{
+	using namespace UE::Chaos::ClothAsset;
+
+	int32 Index = GetPropertyElementIndex(PropertyDescriptor);
+	if (Index == INDEX_NONE)
+	{
+		// Add this property to the preset collection
+		Index = ClothPresetCollection->AddElements(1, FClothPresetCollection::PropertyGroup);
+
+		// Write the new property name and default values to the preset collection
+		ClothPresetCollection->Name[Index] = PropertyDescriptor.Name.ToString();
+		ClothPresetCollection->LowValue[Index] = PropertyDescriptor.DefaultValue;
+		ClothPresetCollection->HighValue[Index] = PropertyDescriptor.DefaultValue;
+		ClothPresetCollection->StringValue[Index] = PropertyDescriptor.DefaultString;
+		ClothPresetCollection->Enable[Index] = PropertyDescriptor.bDefaultEnable;
+		ClothPresetCollection->Animatable[Index] = PropertyDescriptor.bDefaultAnimatable;
+	}
+}
+
+void UChaosClothPreset::SyncCollection()
+{
+	using namespace UE::Chaos::ClothAsset;
+
+	// Cleanup the redundant collection elements that are no longer in the descriptor array
+	const int32 NumElements = ClothPresetCollection->NumElements(FClothPresetCollection::PropertyGroup);
+	TArray<int32> DeletionList;
+	for (int32 Index = 0; Index < NumElements; ++Index)
+	{
+		const FName Name(*ClothPresetCollection->Name[Index]);
+		if (!GetPropertyDescriptor(Name))
+		{
+			DeletionList.Add(Index);
+		}
+	}
+	ClothPresetCollection->RemoveElements(FClothPresetCollection::PropertyGroup, DeletionList);
+
+	// Add any new elements that may have been added since the last sync
+	for (const FChaosClothPresetPropertyDescriptor& PropertyDescriptor : PropertyDescriptors)
+	{
+		SyncProperty(PropertyDescriptor);
+	}
+}
+
diff --git a/Engine/Plugins/Experimental/ChaosClothAsset/Source/ChaosClothAssetEngine/Public/ChaosClothAsset/ClothAsset.h b/Engine/Plugins/Experimental/ChaosClothAsset/Source/ChaosClothAssetEngine/Public/ChaosClothAsset/ClothAsset.h
new file mode 100644
--- /dev/null
+++ b/Engine/Plugins/Experimental/ChaosClothAsset/Source/ChaosClothAssetEngine/Public/ChaosClothAsset/ClothAsset.h
@@ -0,0 +1,236 @@
+// Copyright Epic Games, Inc. All Rights Reserved.
+
+#pragma once
+
+#include "CoreMinimal.h"
+#include "ChaosClothAsset/ClothCollection.h"
+#include "ChaosClothAsset/ClothPreset.h"
+#include "Engine/SkeletalMesh.h"  // For FSkeletalMeshLODInfo
+#include "Engine/SkinnedAsset.h"
+#include "ReferenceSkeleton.h"
+#include "RenderCommandFence.h"
+#include "ClothAsset.generated.h"
+
+class FSkeletalMeshRenderData;
+class FSkeletalMeshModel;
+
+
+UENUM()
+enum class EClothAssetAsyncProperties : uint64
+{
+	None = 0,
+	RenderData = 1 << 0,
+	All = MAX_uint64
+};
+ENUM_CLASS_FLAGS(EClothAssetAsyncProperties);
+
+
+/**
+ * Tailored cloth simulation asset.
+ */
+UCLASS(hidecategories = Object, BlueprintType)
+class CHAOSCLOTHASSETENGINE_API UChaosClothAsset : public USkinnedAsset
+{
+	GENERATED_BODY()
+public:
+	UChaosClothAsset(const FObjectInitializer& ObjectInitializer);
+	//UChaosClothAsset(FVTableHelper& Helper);  // This is declared so we can use TUniquePtr<FSkeletalMeshRenderData> with just a forward declare of that class
+	~UChaosClothAsset() = default;
+
+	//~ Begin UObject interface
+	virtual void Serialize(FArchive& Ar) override;
+	virtual void BeginDestroy() override;
+	virtual bool IsReadyForFinishDestroy() override;
+	virtual void PostLoad() override;
+	//~ End UObject interface
+
+	//~ Begin USkinnedAsset interface
+	virtual FReferenceSkeleton& GetRefSkeleton()								{ return RefSkeleton; }
+	virtual const FReferenceSkeleton& GetRefSkeleton() const					{ return RefSkeleton; }
+	virtual FSkeletalMeshLODInfo* GetLODInfo(int32 Index)						{ return LODInfo.IsValidIndex(Index) ? &LODInfo[Index] : nullptr; }
+	virtual const FSkeletalMeshLODInfo* GetLODInfo(int32 Index) const			{ return LODInfo.IsValidIndex(Index) ? &LODInfo[Index] : nullptr; }
+	UFUNCTION(BlueprintGetter)
+	virtual class UPhysicsAsset* GetShadowPhysicsAsset() const override			{ return ShadowPhysicsAsset; }
+	virtual FMatrix GetComposedRefPoseMatrix(FName InBoneName) const override;
+	virtual FMatrix GetComposedRefPoseMatrix(int32 InBoneIndex) const override	{ return CachedComposedRefPoseMatrices[InBoneIndex]; }
+	virtual const FMeshUVChannelInfo* GetUVChannelData(int32 MaterialIndex) const override;
+	virtual bool GetSupportRayTracing() const override							{ return bSupportRayTracing; }
+	virtual int32 GetRayTracingMinLOD() const override							{ return RayTracingMinLOD; }
+	virtual TArray<FMatrix44f>& GetRefBasesInvMatrix() override					{ return RefBasesInvMatrix; }
+	virtual const TArray<FMatrix44f>& GetRefBasesInvMatrix() const override		{ return RefBasesInvMatrix; }
+	virtual TArray<FSkeletalMeshLODInfo>& GetLODInfoArray() override			{ return LODInfo; }
+	virtual const TArray<FSkeletalMeshLODInfo>& GetLODInfoArray() const override { return LODInfo; }
+	virtual FSkeletalMeshRenderData* GetResourceForRendering() const override;
+	virtual int32 GetDefaultMinLod() const										{ return 0; }
+	virtual UPhysicsAsset* GetPhysicsAsset() const								{ return nullptr; }
+	virtual TArray<FSkeletalMaterial>& GetMaterials() override					{ return Materials; }
+	virtual const TArray<FSkeletalMaterial>& GetMaterials() const override		{ return Materials; }
+	virtual int32 GetLODNum() const override									{ return LODInfo.Num(); }
+	virtual bool IsMaterialUsed(int32 MaterialIndex) const override				{ return true; }
+	virtual FBoxSphereBounds GetBounds() const override							{ return Bounds; }
+	virtual TArray<class USkeletalMeshSocket*> GetActiveSocketList() const override { static TArray<class USkeletalMeshSocket*> Dummy; return Dummy; }
+	virtual USkeletalMeshSocket* FindSocket(FName InSocketName) const override	{ return nullptr; }
+	virtual USkeletalMeshSocket* FindSocketInfo(FName InSocketName, FTransform& OutTransform, int32& OutBoneIndex, int32& OutIndex) const override { return nullptr; }
+	virtual USkeleton* GetSkeleton() override									{ return Skeleton; }
+	virtual const USkeleton* GetSkeleton() const override						{ return Skeleton; }
+	virtual void SetSkeleton(USkeleton* InSkeleton) override					{ Skeleton = InSkeleton; }
+	virtual UMeshDeformer* GetDefaultMeshDeformer() const override				{ return nullptr; }
+	virtual bool IsValidLODIndex(int32 Index) const override					{ return LODInfo.IsValidIndex(Index); }
+	virtual int32 GetMinLodIdx(bool bForceLowestLODIdx = false) const override	{ return 0; }
+	virtual bool NeedCPUData(int32 LODIndex) const override						{ return false; }
+	virtual bool GetHasVertexColors() const override							{ return false; }
+	virtual int32 GetPlatformMinLODIdx(const ITargetPlatform* TargetPlatform) const override { return 0; }
+	virtual const FPerPlatformBool& GetDisableBelowMinLodStripping() const override { return DisableBelowMinLodStripping; }
+	virtual const FPerPlatformInt& GetMinLod() const override					{ return MinLod; }
+#if WITH_EDITOR
+	virtual FString BuildDerivedDataKey(const ITargetPlatform* TargetPlatform) override;
+	virtual bool IsInitialBuildDone() const override;
+	virtual bool GetEnableLODStreaming(const class ITargetPlatform* TargetPlatform) const override	{ return false; }
+	virtual int32 GetMaxNumStreamedLODs(const class ITargetPlatform* TargetPlatform) const override	{ return 0; }
+	virtual int32 GetMaxNumOptionalLODs(const ITargetPlatform* TargetPlatform) const override { return 0; }
+#endif
+#if WITH_EDITORONLY_DATA
+	virtual FSkeletalMeshModel* GetImportedModel() const override				{ return MeshModel.Get(); }
+#endif
+	//~ End USkinnedAsset interface
+
+	/** Return the enclosed Cloth Collection object. */
+	TSharedPtr<UE::Chaos::ClothAsset::FClothCollection> GetClothCollection() { return ClothCollection; }
+
+	/** Return the enclosed Cloth Collection object, const version. */
+	const TSharedPtr<UE::Chaos::ClothAsset::FClothCollection> GetClothCollection() const { return ClothCollection; }
+
+	/** Build this asset static render and simulation data. This needs to be done every time the asset has changed. */
+	void Build();
+
+	/**
+	 * Copy the draped simulation mesh patterns into the render mesh data.
+	 * This is useful to visualize the simulation mesh, or when the simulation mesh can be used for both simulation and rendering.
+	 * @param MaterialIndex The index of the Materials array.
+	 */
+	void CopySimMeshToRenderMesh(int32 MaterialIndex);
+
+private:
+	//~ Begin USkinnedAsset interface
+	/** Initial step for the Post Load process - Can't be done in parallel. */
+	virtual void BeginPostLoadInternal(FSkinnedAssetPostLoadContext& Context) override;
+	/** Thread-safe part of the Post Load */
+	virtual void ExecutePostLoadInternal(FSkinnedAssetPostLoadContext& Context) override;
+	/** Complete the postload process - Can't be done in parallel. */
+	virtual void FinishPostLoadInternal(FSkinnedAssetPostLoadContext& Context) override;
+	/** Convert async property from enum value to string. */
+	virtual FString GetAsyncPropertyName(uint64 Property) const override;
+	//~ End USkinnedAsset interface
+
+	/** 
+	 * Wait for the asset to finish compilation to protect internal skinned asset data from race conditions during async build.
+	 * This should be called before accessing all async accessible properties.
+	 */
+	void WaitUntilAsyncPropertyReleased(EClothAssetAsyncProperties AsyncProperties, ESkinnedAssetAsyncPropertyLockType LockType = ESkinnedAssetAsyncPropertyLockType::ReadWrite) const;
+
+	/** Pre-calculate refpose-to-local transforms. */
+	void CalculateInvRefMatrices();
+
+	/** Re-calculate the bounds for this asset. */
+	void CalculateBounds();
+
+#if WITH_EDITORONLY_DATA
+	/** Build the SkeletalMeshLODModel for this asset. */
+	void BuildModel();
+#endif
+
+	/** Initialize all render resources. */
+	void InitResources();
+
+	/** Safely release the render data. */
+	void ReleaseResources();
+
+	/** Set render data. */
+	void SetResourceForRendering(TUniquePtr<FSkeletalMeshRenderData>&& InSkeletalMeshRenderData);
+
+#if WITH_EDITOR
+	/** Load render data from DDC if the data is cached, otherwise generate render data and save into DDC */
+	void CacheDerivedData(FSkinnedAssetPostLoadContext* Context);
+#endif
+
+	/** List of cloth presets for this cloth asset. */
+	UPROPERTY(EditAnywhere, AssetRegistrySearchable, Category = ClothPresets)
+	TArray<TObjectPtr<UChaosClothPreset>> ClothPresets;
+
+	// TODO: determine if it should be serialized or transient
+	/** List of materials for this cloth asset. */
+	UPROPERTY(EditAnywhere, Category = Materials)
+	TArray<FSkeletalMaterial> Materials;
+
+	FBoxSphereBounds Bounds;
+
+	/** Skeleton. */
+	UPROPERTY(EditAnywhere, AssetRegistrySearchable, Category = Skeleton)
+	TObjectPtr<USkeleton> Skeleton;
+
+	/** Struct containing information for each LOD level, such as materials to use, and when use the LOD. */
+	UPROPERTY(EditAnywhere, EditFixedSize, Category = LevelOfDetails)
+	TArray<FSkeletalMeshLODInfo> LODInfo;
+
+	UPROPERTY(EditAnywhere, Category = LODSettings)
+	FPerPlatformBool DisableBelowMinLodStripping;
+
+	UPROPERTY(EditAnywhere, Category = LODSettings, meta = (DisplayName = "Minimum LOD"))
+	FPerPlatformInt MinLod;
+
+	/** Enable raytracing for this asset. */
+	UPROPERTY(EditAnywhere, Category = RayTracing)
+	uint8 bSupportRayTracing : 1;
+
+	/** Minimum raytracing LOD for this asset. */
+	UPROPERTY(EditAnywhere, Category = RayTracing)
+	int32 RayTracingMinLOD;
+
+	/** Whether to blend positions between the skinned/simulated transitions of the cloth render mesh. */
+	UPROPERTY(EditAnywhere, Category = RayTracing)
+	bool bSmoothTransition = true;
+
+	/** Whether to use multiple triangle influences on the proxy wrap deformer to help smoothen deformations. */
+	UPROPERTY(EditAnywhere, Category = RayTracing)
+	bool bUseMultipleInfluences = false;
+
+	/** The radius from which to get the multiple triangle influences from the simulated proxy mesh. */
+	UPROPERTY(EditAnywhere, Category = RayTracing)
+	float SkinningKernelRadius = 30.f;
+
+	/**
+	 * Physics asset whose shapes will be used for shadowing when components have bCastCharacterCapsuleDirectShadow or bCastCharacterCapsuleIndirectShadow enabled.
+	 * Only spheres and sphyl shapes in the physics asset can be supported.  The more shapes used, the higher the cost of the capsule shadows will be.
+	 */
+	UPROPERTY(EditAnywhere, AssetRegistrySearchable, BlueprintGetter = GetShadowPhysicsAsset, Category = Lighting)
+	TObjectPtr<class UPhysicsAsset> ShadowPhysicsAsset;
+
+	/** A unique identifier as used by the section rendering code. */
+	UPROPERTY()
+	FGuid AssetGuid;
+
+	/** mesh-space ref pose, where parent matrices are applied to ref pose matrices */
+	TArray<FMatrix> CachedComposedRefPoseMatrices;
+
+	/** Cloth Collection containing this asset data. */
+	TSharedPtr<UE::Chaos::ClothAsset::FClothCollection> ClothCollection;
+
+	/** Reference skeleton created from the provided skeleton asset. */
+	FReferenceSkeleton RefSkeleton;
+
+	/** Reference skeleton precomputed bases. */
+	TArray<FMatrix44f> RefBasesInvMatrix;
+
+	/** Rendering data. */
+	TUniquePtr<FSkeletalMeshRenderData> SkeletalMeshRenderData;
+
+	/** A fence which is used to keep track of the rendering thread releasing the static mesh resources. */
+	FRenderCommandFence ReleaseResourcesFence;
+
+#if WITH_EDITORONLY_DATA
+	/** Source mesh geometry information (not used at runtime). */
+	TSharedPtr<FSkeletalMeshModel> MeshModel;
+#endif
+
+	friend class UClothAssetBuilderEditor;
+};
diff --git a/Engine/Plugins/Experimental/ChaosClothAsset/Source/ChaosClothAssetEngine/Public/ChaosClothAsset/ClothAssetBuilder.h b/Engine/Plugins/Experimental/ChaosClothAsset/Source/ChaosClothAssetEngine/Public/ChaosClothAsset/ClothAssetBuilder.h
new file mode 100644
--- /dev/null
+++ b/Engine/Plugins/Experimental/ChaosClothAsset/Source/ChaosClothAssetEngine/Public/ChaosClothAsset/ClothAssetBuilder.h
@@ -0,0 +1,35 @@
+// Copyright Epic Games, Inc. All Rights Reserved.
+
+#pragma once
+
+#include "ChaosClothAsset/ClothAsset.h"
+#include "Features/IModularFeature.h"
+
+#include "ClothAssetBuilder.generated.h"
+
+class UClothAssetBuilder;
+
+/**
+ * A modular interface to provide ways to build a SkeletalMesh LODModel from a cloth asset.
+ * This cannot be done in a non Editor build due to the required dependencies,
+ * and instead is exposed as a modular feature in order to be called from the ClothAsset engine class.
+ */
+class CHAOSCLOTHASSETENGINE_API IClothAssetBuilderClassProvider : public IModularFeature
+{
+public:
+	inline static const FName FeatureName = TEXT("IClothAssetBuilderClassProvider");
+
+	virtual TSubclassOf<UClothAssetBuilder> GetClothAssetBuilderClass() const = 0;
+};
+
+/** Modular builder base class. */
+UCLASS(Abstract)
+class CHAOSCLOTHASSETENGINE_API UClothAssetBuilder : public UObject
+{
+	GENERATED_BODY()
+
+public:
+	/** Build a FSkeletalMeshLODModel out of the cloth asset for the specified LOD index. */
+	virtual void BuildLod(FSkeletalMeshLODModel& LODModel, const UChaosClothAsset& ClothAsset, int32 LodIndex) const
+	PURE_VIRTUAL(UClothAssetBuilder::BuildLod, );
+};
diff --git a/Engine/Plugins/Experimental/ChaosClothAsset/Source/ChaosClothAssetEngine/Public/ChaosClothAsset/ClothComponent.h b/Engine/Plugins/Experimental/ChaosClothAsset/Source/ChaosClothAssetEngine/Public/ChaosClothAsset/ClothComponent.h
new file mode 100644
--- /dev/null
+++ b/Engine/Plugins/Experimental/ChaosClothAsset/Source/ChaosClothAssetEngine/Public/ChaosClothAsset/ClothComponent.h
@@ -0,0 +1,52 @@
+// Copyright Epic Games, Inc. All Rights Reserved.
+
+#pragma once
+
+#include "Components/SkinnedMeshComponent.h"
+
+#include "ClothComponent.generated.h"
+
+class UChaosClothAsset;
+
+/**
+ * Cloth simulation component.
+ */
+UCLASS(ClassGroup = Physics, hidecategories = (Object, "Mesh|SkeletalAsset"), meta = (BlueprintSpawnableComponent, ToolTip = "Chaos cloth simulation component."), DisplayName = "Cloth Simulation", HideCategories = (Physics, Constraints, Advanced, Cooking, Collision))
+class CHAOSCLOTHASSETENGINE_API UChaosClothComponent : public USkinnedMeshComponent
+{
+	GENERATED_BODY()
+public:
+	UChaosClothComponent(const FObjectInitializer& ObjectInitializer);
+
+	/** Set the cloth asset used by this component. */
+	UFUNCTION(BlueprintCallable, Category = "Components|ClothAsset")
+	void SetClothAsset(UChaosClothAsset* InClothAsset);
+
+	/** Get the cloth asset used by this component. */
+	UFUNCTION(BlueprintPure, Category = "Components|ClothAsset")
+	UChaosClothAsset* GetClothAsset() const;
+
+protected:
+	//~ Begin UObject Interface
+	virtual void PostLoad() override;
+#if WITH_EDITOR
+	virtual void PostEditChangeProperty(FPropertyChangedEvent& PropertyChangedEvent) override;
+#endif // WITH_EDITOR
+	//~ End UObject Interface
+
+	//~ Begin UActorComponent Interface
+	virtual void OnRegister() override;
+	//~ End UActorComponent Interface
+
+	//~ Begin USkinClothComponent Interface
+	virtual void RefreshBoneTransforms(FActorComponentTickFunction* TickFunction = nullptr) override;
+	//~ End USkinClothComponent Interface
+
+private:
+#if WITH_EDITORONLY_DATA
+	/** Cloth asset used by this component. */
+	UE_DEPRECATED(5.1, "This property isn't deprecated, but getter and setter must be used at all times to preserve correct operations.")
+	UPROPERTY(EditAnywhere, Transient, BlueprintSetter = SetClothAsset, BlueprintGetter = GetClothAsset, Category = Mesh)
+	TObjectPtr<UChaosClothAsset> ClothAsset;
+#endif
+};
diff --git a/Engine/Plugins/Experimental/ChaosClothAsset/Source/ChaosClothAssetEngine/Public/ChaosClothAsset/ClothPreset.h b/Engine/Plugins/Experimental/ChaosClothAsset/Source/ChaosClothAssetEngine/Public/ChaosClothAsset/ClothPreset.h
new file mode 100644
--- /dev/null
+++ b/Engine/Plugins/Experimental/ChaosClothAsset/Source/ChaosClothAssetEngine/Public/ChaosClothAsset/ClothPreset.h
@@ -0,0 +1,248 @@
+// Copyright Epic Games, Inc. All Rights Reserved.
+
+#pragma once
+
+#include "CoreMinimal.h"
+#include "ChaosClothAsset/ClothPresetCollection.h"
+
+#include "ClothPreset.generated.h"
+
+/**
+ * Preset property categories.
+ */
+UENUM()
+enum class EChaosClothPresetPropertyCategory : uint8
+{
+	Mass,
+	Constraint,
+	Collision,
+	Environmental,
+	Animation,
+	Experimental,
+	Custom
+};
+
+/**
+ * Preset property types.
+ */
+UENUM()
+enum class EChaosClothPresetPropertyType : uint8
+{
+	Boolean,
+	Integer,
+	Float,
+	Vector3Integer,
+	Vector3Float,
+	String
+};
+
+template <typename T> struct TIsClothPresetPropertyNonStringType { enum { Value = false }; };
+template<> struct TIsClothPresetPropertyNonStringType<bool> { enum { Value = true }; };
+template<> struct TIsClothPresetPropertyNonStringType<int32> { enum { Value = true }; };
+template<> struct TIsClothPresetPropertyNonStringType<float> { enum { Value = true }; };
+template<> struct TIsClothPresetPropertyNonStringType<FIntVector3> { enum { Value = true }; };
+template<> struct TIsClothPresetPropertyNonStringType<FVector3f> { enum { Value = true }; };
+
+template <typename T> struct TIsClothPresetPropertyType { enum { Value = TIsClothPresetPropertyNonStringType<T>::Value }; };
+template<> struct TIsClothPresetPropertyType<FString> { enum { Value = true }; };
+
+/**
+ * Structure used to describe simulation properties.
+ */
+USTRUCT()
+struct FChaosClothPresetPropertyDescriptor
+{
+	GENERATED_BODY()
+
+	static const FName MassCategory;
+	static const FName ConstraintCategory;
+	static const FName CollisionCategory;
+	static const FName EnvironmentalCategory;
+	static const FName AnimationCategory;
+	static const FName CustomCategory;
+
+	UPROPERTY(EditAnywhere, Category = "Property Descriptor")
+	FName Name;
+
+	UPROPERTY(EditAnywhere, Category = "Property Descriptor")
+	EChaosClothPresetPropertyType Type = EChaosClothPresetPropertyType::Vector3Float;
+
+	UPROPERTY(EditAnywhere, Category = "Property Descriptor")
+	EChaosClothPresetPropertyCategory Category = EChaosClothPresetPropertyCategory::Custom;
+
+	UPROPERTY(EditAnywhere, Category = "Property Descriptor")
+	FText Description;
+
+	UPROPERTY(EditAnywhere, Category = "Property Descriptor")
+	FVector3f DefaultValue = FVector3f::ZeroVector;  // Can only be used for continuous integer values from -16777216 to 16777216
+
+	UPROPERTY(EditAnywhere, Category = "Property Descriptor")
+	FString DefaultString;  // Only used for string type
+
+	UPROPERTY(EditAnywhere, Category = "Property Descriptor")
+	bool bDefaultEnable = true;
+
+	UPROPERTY(EditAnywhere, Category = "Property Descriptor")
+	bool bDefaultAnimatable = false;
+
+	UPROPERTY(EditAnywhere, Category = "Property Descriptor")
+	float UIMin = 0.f;
+
+	UPROPERTY(EditAnywhere, Category = "Property Descriptor")
+	float UIMax = 1000.f;
+
+	UPROPERTY(EditAnywhere, Category = "Property Descriptor")
+	float ClampMin = TNumericLimits<float>::Min();
+
+	UPROPERTY(EditAnywhere, Category = "Property Descriptor")
+	float ClampMax = TNumericLimits<float>::Max();
+
+	UPROPERTY(EditAnywhere, Category = "Property Descriptor")
+	FString EditCondition;
+
+	FChaosClothPresetPropertyDescriptor() = default;
+
+	/** Minimal viable constructor for use in UChaosClothPreset. */
+	FChaosClothPresetPropertyDescriptor(
+		const FName& InName,
+		EChaosClothPresetPropertyType InType)
+		: Name(InName)
+		, Type(InType)
+	{}
+
+	template<typename T>
+	static typename TEnableIf<TIsClothPresetPropertyType<T>::Value, EChaosClothPresetPropertyType>::Type GetType();
+};
+
+/**
+ * Class used to generate custom list of simulation properties.
+ */
+UCLASS()
+class UChaosClothPresetPropertyGenerator : public UClass
+{
+	GENERATED_BODY()
+
+public:
+
+};
+
+/**
+ * Tailored cloth simulation preset.
+ */
+UCLASS(hidecategories = Object, BlueprintType)
+class CHAOSCLOTHASSETENGINE_API UChaosClothPreset : public UObject
+{
+	GENERATED_BODY()
+public:
+	UChaosClothPreset(const FObjectInitializer& ObjectInitializer);
+
+	void Serialize(FArchive& Ar);
+
+	/** Add a new property, return false if the property already exists or true otherwise. */
+	bool AddProperty(const FName& Name, EChaosClothPresetPropertyType Type);
+
+	/** Set non string properties. */
+	template<typename T>
+	inline typename TEnableIf<TIsClothPresetPropertyNonStringType<T>::Value, void>::Type SetProperty(
+		const FName& Name,
+		T LowValue,
+		T HighValue,
+		const FString& StringValue,
+		bool bEnable,
+		bool bAnimatable);
+
+	/** Set string properties. */
+	void SetProperty(const FName& Name, const FString& StringValue, bool bEnable, bool bAnimatable);
+
+	/** Set (low) value. */
+	template<typename T>
+	inline typename TEnableIf<TIsClothPresetPropertyNonStringType<T>::Value, void>::Type SetPropertyValue(
+		const FName& Name,
+		T Value);
+
+	/** Set low and high value. */
+	template<typename T>
+	inline typename TEnableIf<TIsClothPresetPropertyNonStringType<T>::Value, void>::Type SetPropertyValues(
+		const FName& Name,
+		T LowValue,
+		T HighValue);
+
+	/** Set string value. */
+	void SetPropertyStringValue(const FName& Name, const FString& StringValue);
+
+	/** Set property enabled. */
+	void SetPropertyEnable(const FName& Name, bool bEnable);
+
+	/** Set property animatable. */
+	void SetPropertyAnimatable(const FName& Name, bool bAnimatable);
+
+private:
+	const FChaosClothPresetPropertyDescriptor* GetPropertyDescriptor(const FName& Name) const;
+
+	int32 GetPropertyElementIndex(const FChaosClothPresetPropertyDescriptor& PropertyDescriptor) const;
+
+	void SetPropertyInternal(
+		const FChaosClothPresetPropertyDescriptor& PropertyDescriptor,
+		const FVector3f& LowValue,
+		const FVector3f& HighValue,
+		const FString& StringValue,
+		bool bEnable,
+		bool bAnimatable);
+
+	void SetPropertyValueInternal(const FChaosClothPresetPropertyDescriptor& PropertyDescriptor, const FVector3f& Value);
+	void SetPropertyValuesInternal(const FChaosClothPresetPropertyDescriptor& PropertyDescriptor, const FVector3f& LowValue, const FVector3f& HighValue);
+
+	void SyncProperty(const FChaosClothPresetPropertyDescriptor& PropertyDescriptor);
+	void SyncCollection();
+
+	UPROPERTY(EditAnywhere, Category = "PropertyDescriptor Editor")
+	TArray<FChaosClothPresetPropertyDescriptor> PropertyDescriptors;
+
+	TSharedPtr<UE::Chaos::ClothAsset::FClothPresetCollection, ESPMode::ThreadSafe> ClothPresetCollection;
+};
+
+template<typename T>
+typename TEnableIf<TIsClothPresetPropertyNonStringType<T>::Value, void>::Type UChaosClothPreset::SetProperty(
+	const FName& Name,
+	T LowValue,
+	T HighValue,
+	const FString& StringValue,
+	bool bEnable,
+	bool bAnimatable)
+{
+	const FChaosClothPresetPropertyDescriptor* const PropertyDescriptor = GetPropertyDescriptor(Name);
+	if (ensure(PropertyDescriptor) && ensure(PropertyDescriptor->Type == FChaosClothPresetPropertyDescriptor::GetType<T>()))
+	{
+		const FVector3f LowValueClamped = ClampVector(FVector3f(LowValue), FVector3f(PropertyDescriptor->ClampMin), FVector3f(PropertyDescriptor->ClampMax));
+		const FVector3f HighValueClamped = ClampVector(FVector3f(HighValue), FVector3f(PropertyDescriptor->ClampMin), FVector3f(PropertyDescriptor->ClampMax));
+		SetPropertyInternal(*PropertyDescriptor, LowValueClamped, HighValueClamped, StringValue, bEnable, bAnimatable);
+	}
+}
+
+template<typename T>
+typename TEnableIf<TIsClothPresetPropertyNonStringType<T>::Value, void>::Type UChaosClothPreset::SetPropertyValue(
+	const FName& Name,
+	T Value)
+{
+	const FChaosClothPresetPropertyDescriptor* const PropertyDescriptor = GetPropertyDescriptor(Name);
+	if (ensure(PropertyDescriptor) && ensure(PropertyDescriptor->Type == FChaosClothPresetPropertyDescriptor::GetType<T>()))
+	{
+		const FVector3f ValueClamped = ClampVector(FVector3f(Value), FVector3f(PropertyDescriptor->ClampMin), FVector3f(PropertyDescriptor->ClampMax));
+		SetPropertyValueInternal(*PropertyDescriptor, ValueClamped);
+	}
+}
+
+template<typename T>
+typename TEnableIf<TIsClothPresetPropertyNonStringType<T>::Value, void>::Type UChaosClothPreset::SetPropertyValues(
+	const FName& Name,
+	T LowValue,
+	T HighValue)
+{
+	const FChaosClothPresetPropertyDescriptor* const PropertyDescriptor = GetPropertyDescriptor(Name);
+	if (ensure(PropertyDescriptor) && ensure(PropertyDescriptor->Type == FChaosClothPresetPropertyDescriptor::GetType<T>()))
+	{
+		const FVector3f LowValueClamped = ClampVector(FVector3f(LowValue), FVector3f(PropertyDescriptor->ClampMin), FVector3f(PropertyDescriptor->ClampMax));
+		const FVector3f HighValueClamped = ClampVector(FVector3f(HighValue), FVector3f(PropertyDescriptor->ClampMin), FVector3f(PropertyDescriptor->ClampMax));
+		SetPropertyValuesInternal(*PropertyDescriptor, LowValueClamped, HighValueClamped);
+	}
+}
diff --git a/Engine/Plugins/Experimental/ChaosClothAssetEditor/ChaosClothAssetEditor.uplugin b/Engine/Plugins/Experimental/ChaosClothAssetEditor/ChaosClothAssetEditor.uplugin
new file mode 100644
--- /dev/null
+++ b/Engine/Plugins/Experimental/ChaosClothAssetEditor/ChaosClothAssetEditor.uplugin
@@ -0,0 +1,61 @@
+{
+	"FileVersion": 1,
+	"Version": 1,
+	"VersionName": "0.1",
+	"FriendlyName": "Chaos Cloth Asset Editor",
+	"Description": "Editor for modifying cloth assets",
+	"Category": "Other",
+	"CreatedBy": "Epic Games, Inc.",
+	"CreatedByURL": "http://epicgames.com",
+	"DocsURL": "",
+	"MarketplaceURL": "",
+	"SupportURL": "",
+	"CanContainContent": true,
+	"IsExperimentalVersion": true,
+	"Installed": false,
+	"Modules": [
+		{
+			"Name": "ChaosClothAssetEditor",
+			"Type": "Editor",
+			"LoadingPhase": "Default",
+			"PlatformAllowList": [ "Win64" ]
+		},
+		{
+			"Name": "ChaosClothAssetEditorTools",
+			"Type": "Editor",
+			"LoadingPhase": "Default",
+			"PlatformAllowList": [ "Win64" ]
+		},
+		{
+			"Name": "ChaosClothAssetTools",
+			"Type": "UncookedOnly",
+			"LoadingPhase": "Default",
+			"PlatformAllowList": [ "Win64" ]
+		}
+	],
+	"Plugins": [
+		{
+			"Name": "GeometryProcessing",
+			"Enabled": true
+		},
+		{
+			"Name": "ChaosClothAsset",
+			"Enabled": true
+		},
+		{
+			"Name": "ChaosCloth",
+			"Enabled": true
+		},
+		{
+			"Name": "MeshModelingToolsetExp",
+			"Enabled": true,
+			"TargetAllowList": [ "Editor" ]
+		},
+		{
+			"Name": "BaseCharacterFXEditor",
+			"Enabled": true,
+			"TargetAllowList": [ "Editor" ]
+		}
+
+	]
+}
\ No newline at end of file
diff --git a/Engine/Plugins/Experimental/ChaosClothAssetEditor/Source/ChaosClothAssetEditor/Private/ChaosClothAsset/AssetTypeActions_ClothAsset.cpp b/Engine/Plugins/Experimental/ChaosClothAssetEditor/Source/ChaosClothAssetEditor/Private/ChaosClothAsset/AssetTypeActions_ClothAsset.cpp
new file mode 100644
--- /dev/null
+++ b/Engine/Plugins/Experimental/ChaosClothAssetEditor/Source/ChaosClothAssetEditor/Private/ChaosClothAsset/AssetTypeActions_ClothAsset.cpp
@@ -0,0 +1,62 @@
+// Copyright Epic Games, Inc. All Rights Reserved.
+
+#include "AssetTypeActions_ClothAsset.h"
+#include "ThumbnailRendering/SceneThumbnailInfo.h"
+#include "ToolMenus.h"
+#include "ChaosClothAsset/ClothEditor.h"
+#include "ChaosClothAsset/ClothAsset.h"
+
+namespace UE::Chaos::ClothAsset
+{
+	FText FAssetTypeActions_ClothAsset::GetName() const
+	{
+		return NSLOCTEXT("AssetTypeActions", "AssetTypeActions_ClothAsset", "Cloth Asset");
+	}
+
+	void FAssetTypeActions_ClothAsset::GetActions(const TArray<UObject*>& Objects, FToolMenuSection& Section)
+	{
+		FAssetTypeActions_Base::GetActions(Objects, Section);
+	}
+
+	FColor FAssetTypeActions_ClothAsset::GetTypeColor() const
+	{
+		return FColor(180, 120, 110);
+	}
+
+	UClass* FAssetTypeActions_ClothAsset::GetSupportedClass() const
+	{
+		return UChaosClothAsset::StaticClass();
+	}
+
+	void FAssetTypeActions_ClothAsset::OpenAssetEditor(const TArray<UObject*>& Objects, TSharedPtr<IToolkitHost> EditWithinLevelEditor)
+	{
+		TArray<TObjectPtr<UObject>> ClothObjects;
+
+		for (UObject* const Object : Objects)
+		{
+			if (Cast<UChaosClothAsset>(Object))
+			{
+				ClothObjects.Add(Object);
+			}
+		}
+
+		// TODO:
+		if (ClothObjects.Num() > 0)
+		{
+			UAssetEditorSubsystem* const AssetEditorSubsystem = GEditor->GetEditorSubsystem<UAssetEditorSubsystem>();
+			UChaosClothAssetEditor* const AssetEditor = NewObject<UChaosClothAssetEditor>(AssetEditorSubsystem, NAME_None, RF_Transient);
+			AssetEditor->Initialize(ClothObjects);
+		}
+	}
+
+	uint32 FAssetTypeActions_ClothAsset::GetCategories()
+	{
+		return EAssetTypeCategories::Physics;
+	}
+
+	UThumbnailInfo* FAssetTypeActions_ClothAsset::GetThumbnailInfo(UObject* Asset) const
+	{
+		check(Cast<UChaosClothAsset>(Asset));
+		return NewObject<USceneThumbnailInfo>(Asset, NAME_None, RF_Transactional);
+	}
+} // End namespace UE::Chaos::ClothAsset
diff --git a/Engine/Plugins/Experimental/ChaosClothAssetEditor/Source/ChaosClothAssetEditor/Private/ChaosClothAsset/AssetTypeActions_ClothAsset.h b/Engine/Plugins/Experimental/ChaosClothAssetEditor/Source/ChaosClothAssetEditor/Private/ChaosClothAsset/AssetTypeActions_ClothAsset.h
new file mode 100644
--- /dev/null
+++ b/Engine/Plugins/Experimental/ChaosClothAssetEditor/Source/ChaosClothAssetEditor/Private/ChaosClothAsset/AssetTypeActions_ClothAsset.h
@@ -0,0 +1,22 @@
+// Copyright Epic Games, Inc. All Rights Reserved.
+
+#pragma once
+
+#include "CoreMinimal.h"
+#include "AssetTypeActions_Base.h"
+
+namespace UE::Chaos::ClothAsset
+{
+	class FAssetTypeActions_ClothAsset : public FAssetTypeActions_Base
+	{
+	public:
+		// IAssetTypeActions Implementation
+		virtual FText GetName() const override;
+		virtual void GetActions(const TArray<UObject*>& InObjects, struct FToolMenuSection& Section) override;
+		virtual FColor GetTypeColor() const override;
+		virtual UClass* GetSupportedClass() const override;
+		virtual void OpenAssetEditor(const TArray<UObject*>& InObjects, TSharedPtr<class IToolkitHost> EditWithinLevelEditor = TSharedPtr<IToolkitHost>()) override;
+		virtual uint32 GetCategories() override;
+		virtual UThumbnailInfo* GetThumbnailInfo(UObject* Asset) const override;
+	};
+}  // End namespace UE::Chaos::ClothAsset
\ No newline at end of file
diff --git a/Engine/Plugins/Experimental/ChaosClothAssetEditor/Source/ChaosClothAssetEditor/Private/ChaosClothAsset/AssetTypeActions_ClothPreset.cpp b/Engine/Plugins/Experimental/ChaosClothAssetEditor/Source/ChaosClothAssetEditor/Private/ChaosClothAsset/AssetTypeActions_ClothPreset.cpp
new file mode 100644
--- /dev/null
+++ b/Engine/Plugins/Experimental/ChaosClothAssetEditor/Source/ChaosClothAssetEditor/Private/ChaosClothAsset/AssetTypeActions_ClothPreset.cpp
@@ -0,0 +1,45 @@
+// Copyright Epic Games, Inc. All Rights Reserved.
+
+#include "AssetTypeActions_ClothPreset.h"
+#include "ThumbnailRendering/SceneThumbnailInfo.h"
+#include "ToolMenus.h"
+#include "ChaosClothAsset/ClothPreset.h"
+
+namespace UE::Chaos::ClothAsset
+{
+	FText FAssetTypeActions_ClothPreset::GetName() const
+	{
+		return NSLOCTEXT("AssetTypeActions", "AssetTypeActions_ClothPreset", "Cloth Preset");
+	}
+
+	void FAssetTypeActions_ClothPreset::GetActions(const TArray<UObject*>& Objects, FToolMenuSection& Section)
+	{
+		FAssetTypeActions_Base::GetActions(Objects, Section);
+	}
+
+	FColor FAssetTypeActions_ClothPreset::GetTypeColor() const
+	{
+		return FColor(180, 120, 110);
+	}
+
+	UClass* FAssetTypeActions_ClothPreset::GetSupportedClass() const
+	{
+		return UChaosClothPreset::StaticClass();
+	}
+
+	void FAssetTypeActions_ClothPreset::OpenAssetEditor(const TArray<UObject*>& Objects, TSharedPtr<IToolkitHost> EditWithinLevelEditor)
+	{
+		// TODO: open the cloth editor
+	}
+
+	uint32 FAssetTypeActions_ClothPreset::GetCategories()
+	{
+		return EAssetTypeCategories::Physics;
+	}
+
+	UThumbnailInfo* FAssetTypeActions_ClothPreset::GetThumbnailInfo(UObject* Asset) const
+	{
+		check(Cast<UChaosClothPreset>(Asset));
+		return NewObject<USceneThumbnailInfo>(Asset, NAME_None, RF_Transactional);
+	}
+} // End namespace UE::Chaos::ClothAsset
diff --git a/Engine/Plugins/Experimental/ChaosClothAssetEditor/Source/ChaosClothAssetEditor/Private/ChaosClothAsset/AssetTypeActions_ClothPreset.h b/Engine/Plugins/Experimental/ChaosClothAssetEditor/Source/ChaosClothAssetEditor/Private/ChaosClothAsset/AssetTypeActions_ClothPreset.h
new file mode 100644
--- /dev/null
+++ b/Engine/Plugins/Experimental/ChaosClothAssetEditor/Source/ChaosClothAssetEditor/Private/ChaosClothAsset/AssetTypeActions_ClothPreset.h
@@ -0,0 +1,22 @@
+// Copyright Epic Games, Inc. All Rights Reserved.
+
+#pragma once
+
+#include "CoreMinimal.h"
+#include "AssetTypeActions_Base.h"
+
+namespace UE::Chaos::ClothAsset
+{
+	class FAssetTypeActions_ClothPreset : public FAssetTypeActions_Base
+	{
+	public:
+		// IAssetTypeActions Implementation
+		virtual FText GetName() const override;
+		virtual void GetActions(const TArray<UObject*>& InObjects, struct FToolMenuSection& Section) override;
+		virtual FColor GetTypeColor() const override;
+		virtual UClass* GetSupportedClass() const override;
+		virtual void OpenAssetEditor(const TArray<UObject*>& InObjects, TSharedPtr<class IToolkitHost> EditWithinLevelEditor = TSharedPtr<IToolkitHost>()) override;
+		virtual uint32 GetCategories() override;
+		virtual UThumbnailInfo* GetThumbnailInfo(UObject* Asset) const override;
+	};
+}  // End namespace UE::Chaos::ClothAsset
diff --git a/Engine/Plugins/Experimental/ChaosClothAssetEditor/Source/ChaosClothAssetEditor/Private/ChaosClothAsset/ChaosClothAssetEditorModule.cpp b/Engine/Plugins/Experimental/ChaosClothAssetEditor/Source/ChaosClothAssetEditor/Private/ChaosClothAsset/ChaosClothAssetEditorModule.cpp
new file mode 100644
--- /dev/null
+++ b/Engine/Plugins/Experimental/ChaosClothAssetEditor/Source/ChaosClothAssetEditor/Private/ChaosClothAsset/ChaosClothAssetEditorModule.cpp
@@ -0,0 +1,60 @@
+// Copyright Epic Games, Inc. All Rights Reserved.
+
+#include "ChaosClothAsset/ChaosClothAssetEditorModule.h"
+#include "ChaosClothAsset/ClothEditorMode.h"
+#include "ChaosClothAsset/ClothEditorCommands.h"
+#include "ChaosClothAsset/ClothEditorStyle.h"
+#include "ContentBrowserMenuContexts.h"
+#include "EditorModeRegistry.h"
+#include "Selection.h"
+#include "ToolMenus.h"
+#include "Editor.h"
+#include "Styling/AppStyle.h"
+#include "AssetToolsModule.h"
+#include "AssetTypeActions_ClothAsset.h"
+#include "AssetTypeActions_ClothPreset.h"
+
+#define LOCTEXT_NAMESPACE "FChaosClothAssetEditorModule"
+
+using namespace UE::Chaos::ClothAsset;
+
+void FChaosClothAssetEditorModule::StartupModule()
+{
+	FChaosClothAssetEditorStyle::Get(); // Causes the constructor to be called
+
+	FChaosClothAssetEditorCommands::Register();
+
+	// Register asset actions
+	FAssetToolsModule& AssetToolsModule = FAssetToolsModule::GetModule();
+	IAssetTools& AssetTools = AssetToolsModule.Get();
+
+	AssetTypeActions_ClothAsset = new FAssetTypeActions_ClothAsset();
+	AssetTools.RegisterAssetTypeActions(MakeShareable(AssetTypeActions_ClothAsset));
+
+	AssetTypeActions_ClothPreset = new FAssetTypeActions_ClothPreset();
+	AssetTools.RegisterAssetTypeActions(MakeShareable(AssetTypeActions_ClothPreset));
+
+	// TODO: Register details view customizations
+}
+
+void FChaosClothAssetEditorModule::ShutdownModule()
+{
+	FChaosClothAssetEditorCommands::Unregister();
+
+	FEditorModeRegistry::Get().UnregisterMode(UChaosClothAssetEditorMode::EM_ChaosClothAssetEditorModeId);
+
+	if (UObjectInitialized())
+	{
+		// Unregister asset actions
+		FAssetToolsModule& AssetToolsModule = FAssetToolsModule::GetModule();
+		IAssetTools& AssetTools = AssetToolsModule.Get();
+		AssetTools.UnregisterAssetTypeActions(AssetTypeActions_ClothAsset->AsShared());
+		AssetTools.UnregisterAssetTypeActions(AssetTypeActions_ClothPreset->AsShared());
+	}
+
+	// TODO: Unregister details view customizations
+}
+
+#undef LOCTEXT_NAMESPACE
+	
+IMPLEMENT_MODULE(FChaosClothAssetEditorModule, ChaosClothAssetEditor)
diff --git a/Engine/Plugins/Experimental/ChaosClothAssetEditor/Source/ChaosClothAssetEditor/Private/ChaosClothAsset/ClothEditor3DViewportClient.cpp b/Engine/Plugins/Experimental/ChaosClothAssetEditor/Source/ChaosClothAssetEditor/Private/ChaosClothAsset/ClothEditor3DViewportClient.cpp
new file mode 100644
--- /dev/null
+++ b/Engine/Plugins/Experimental/ChaosClothAssetEditor/Source/ChaosClothAssetEditor/Private/ChaosClothAsset/ClothEditor3DViewportClient.cpp
@@ -0,0 +1,44 @@
+// Copyright Epic Games, Inc. All Rights Reserved.
+
+#include "ChaosClothAsset/ClothEditor3DViewportClient.h"
+#include "ChaosClothAsset/ClothEditorMode.h"
+#include "EditorModeManager.h"
+#include "Components/SkeletalMeshComponent.h"
+
+FChaosClothAssetEditor3DViewportClient::FChaosClothAssetEditor3DViewportClient(FEditorModeTools* InModeTools,
+	FPreviewScene* InPreviewScene, 
+	const TWeakPtr<SEditorViewport>& InEditorViewportWidget)
+	: FEditorViewportClient(InModeTools, InPreviewScene, InEditorViewportWidget)
+{
+	// We want our near clip plane to be quite close so that we can zoom in further.
+	OverrideNearClipPlane(KINDA_SMALL_NUMBER);
+
+	// Call this once with the default value to get everything in a consistent state
+	EnableRenderMeshWireframe(bRenderMeshWireframe);
+}
+
+
+void FChaosClothAssetEditor3DViewportClient::EnableRenderMeshWireframe(bool bEnable)
+{
+	bRenderMeshWireframe = bEnable;
+
+	for (TObjectPtr<USkeletalMeshComponent> SkeletalMeshComponent : SkeletalMeshComponents)
+	{
+		if (SkeletalMeshComponent)
+		{
+			SkeletalMeshComponent->SetForceWireframe(bRenderMeshWireframe);
+		}
+	}
+}
+
+void FChaosClothAssetEditor3DViewportClient::Draw(const FSceneView* View, FPrimitiveDrawInterface* PDI)
+{
+	if (!bSimMeshWireframe)
+	{
+		return;
+	}
+
+	// TODO: Draw simulation mesh wireframe
+
+}
+
diff --git a/Engine/Plugins/Experimental/ChaosClothAssetEditor/Source/ChaosClothAssetEditor/Private/ChaosClothAsset/ClothEditorCommands.cpp b/Engine/Plugins/Experimental/ChaosClothAssetEditor/Source/ChaosClothAssetEditor/Private/ChaosClothAsset/ClothEditorCommands.cpp
new file mode 100644
--- /dev/null
+++ b/Engine/Plugins/Experimental/ChaosClothAssetEditor/Source/ChaosClothAssetEditor/Private/ChaosClothAsset/ClothEditorCommands.cpp
@@ -0,0 +1,42 @@
+// Copyright Epic Games, Inc. All Rights Reserved.
+
+#include "ChaosClothAsset/ClothEditorCommands.h"
+#include "ChaosClothAsset/ClothEditorStyle.h"
+
+#define LOCTEXT_NAMESPACE "FChaosClothAssetEditorCommands"
+
+const FString FChaosClothAssetEditorCommands::BeginRemeshToolIdentifier = TEXT("BeginRemeshTool");
+const FString FChaosClothAssetEditorCommands::BeginAttributeEditorToolIdentifier = TEXT("BeginAttributeEditorTool");
+const FString FChaosClothAssetEditorCommands::BeginWeightMapPaintToolIdentifier = TEXT("BeginWeightMapPaintTool");
+const FString FChaosClothAssetEditorCommands::ToggleSimMeshWireframeIdentifier = TEXT("ToggleSimMeshWireframe");
+const FString FChaosClothAssetEditorCommands::ToggleRenderMeshWireframeIdentifier = TEXT("ToggleRenderMeshWireframe");
+const FString FChaosClothAssetEditorCommands::TogglePatternModeIdentifier = TEXT("TogglePatternMode");
+
+
+
+FChaosClothAssetEditorCommands::FChaosClothAssetEditorCommands()
+	: TBaseCharacterFXEditorCommands<FChaosClothAssetEditorCommands>("ChaosClothAssetEditor",
+		LOCTEXT("ContextDescription", "Cloth Editor"), 
+		NAME_None, // Parent
+		FChaosClothAssetEditorStyle::Get().GetStyleSetName())
+{
+}
+
+void FChaosClothAssetEditorCommands::RegisterCommands()
+{
+	TBaseCharacterFXEditorCommands::RegisterCommands();
+
+	UI_COMMAND(OpenClothEditor, "Cloth Editor", "Open the Cloth Editor window", EUserInterfaceActionType::Button, FInputChord());
+
+	UI_COMMAND(BeginRemeshTool, "Remesh", "Remesh the selected mesh", EUserInterfaceActionType::Button, FInputChord());
+	UI_COMMAND(BeginAttributeEditorTool, "AttrEd", "Edit/configure mesh attributes", EUserInterfaceActionType::Button, FInputChord());
+	UI_COMMAND(BeginWeightMapPaintTool, "MapPnt", "Paint Weight Maps on the mesh", EUserInterfaceActionType::Button, FInputChord());
+
+	UI_COMMAND(TogglePatternMode, "TogglePatternMode", "Toggle pattern mode", EUserInterfaceActionType::Button, FInputChord());
+
+	UI_COMMAND(ToggleSimMeshWireframe, "ToggleSimMeshWireframe", "Toggle simulation mesh wireframe", EUserInterfaceActionType::Button, FInputChord());
+	UI_COMMAND(ToggleRenderMeshWireframe, "ToggleRenderMeshWireframe", "Toggle render mesh wireframe", EUserInterfaceActionType::Button, FInputChord());
+}
+
+
+#undef LOCTEXT_NAMESPACE
diff --git a/Engine/Plugins/Experimental/ChaosClothAssetEditor/Source/ChaosClothAssetEditor/Private/ChaosClothAsset/ClothEditorMode.cpp b/Engine/Plugins/Experimental/ChaosClothAssetEditor/Source/ChaosClothAssetEditor/Private/ChaosClothAsset/ClothEditorMode.cpp
new file mode 100644
--- /dev/null
+++ b/Engine/Plugins/Experimental/ChaosClothAssetEditor/Source/ChaosClothAssetEditor/Private/ChaosClothAsset/ClothEditorMode.cpp
@@ -0,0 +1,736 @@
+// Copyright Epic Games, Inc. All Rights Reserved.
+
+#include "ChaosClothAsset/ClothEditorMode.h"
+#include "ChaosClothAsset/ClothEditorCommands.h"
+#include "ChaosClothAsset/ClothEditorModeToolkit.h"
+#include "ChaosClothAsset/ClothAsset.h"
+#include "ChaosClothAsset/ClothComponent.h"
+#include "ChaosClothAsset/ClothCollection.h"
+#include "ChaosClothAsset/ClothAdapter.h"
+#include "ChaosClothAsset/ClothEditorRestSpaceViewportClient.h"
+#include "ChaosClothAsset/ClothPatternToDynamicMesh.h"
+#include "AssetEditorModeManager.h"
+#include "Drawing/MeshElementsVisualizer.h"
+#include "EditorViewportClient.h"
+#include "EdModeInteractiveToolsContext.h"
+#include "Framework/Commands/UICommandList.h"
+#include "InteractiveTool.h"
+#include "ModelingToolTargetUtil.h"
+#include "PreviewScene.h"
+#include "TargetInterfaces/AssetBackedTarget.h"
+#include "TargetInterfaces/DynamicMeshCommitter.h"
+#include "TargetInterfaces/DynamicMeshProvider.h"
+#include "TargetInterfaces/MaterialProvider.h"
+#include "ToolTargetManager.h"
+#include "Toolkits/BaseToolkit.h"
+#include "ToolTargets/ToolTarget.h"
+#include "ToolTargets/StaticMeshComponentToolTarget.h"
+#include "ToolTargets/VolumeComponentToolTarget.h"
+#include "ToolTargets/DynamicMeshComponentToolTarget.h"
+#include "ToolTargets/SkeletalMeshComponentToolTarget.h"
+#include "Engine/Selection.h"
+#include "Misc/ConfigCacheIni.h"
+#include "Modules/ModuleManager.h"
+#include "RemeshMeshTool.h"
+#include "AttributeEditorTool.h"
+#include "MeshAttributePaintTool.h"
+#include "ChaosClothAsset/ClothWeightMapPaintTool.h"
+#include "DynamicMesh/DynamicMeshAttributeSet.h"
+#include "DynamicMesh/DynamicVertexSkinWeightsAttribute.h"
+#include "DynamicMeshEditor.h"
+#include "Settings/EditorStyleSettings.h"
+#include "ClothingSystemEditorInterfaceModule.h"
+#include "ClothingAssetFactoryInterface.h"
+#include "ClothingAssetFactory.h"
+#include "Engine/SkeletalMesh.h"
+#include "Rendering/SkeletalMeshLODModel.h"
+#include "Rendering/SkeletalMeshModel.h"
+#include "Utils/ClothingMeshUtils.h"
+#include "Widgets/Notifications/SNotificationList.h"
+#include "BoneWeights.h"
+#include "ToolSetupUtil.h"
+
+#include UE_INLINE_GENERATED_CPP_BY_NAME(ClothEditorMode)
+
+#define LOCTEXT_NAMESPACE "UChaosClothAssetEditorMode"
+
+const FEditorModeID UChaosClothAssetEditorMode::EM_ChaosClothAssetEditorModeId = TEXT("EM_ChaosClothAssetEditorMode");
+
+namespace ChaosClothAssetEditorModeLocals
+{
+	void LogAndToastWarning(const FText& Error)
+	{
+		FNotificationInfo Info(Error);
+		Info.ExpireDuration = 5.0f;
+		FSlateNotificationManager::Get().AddNotification(Info);
+
+		// TODO: use a non-temp category
+		UE_LOG(LogTemp, Warning, TEXT("%s"), *Error.ToString());
+	}
+
+	FBoxSphereBounds RenderMeshBounds(const UE::Chaos::ClothAsset::FClothPatternConstAdapter& ClothPatternAdapter)
+	{
+		const TSharedPtr<const UE::Chaos::ClothAsset::FClothCollection> ClothCollection = ClothPatternAdapter.GetClothCollection();
+		const int32 PatternElementIndex = ClothPatternAdapter.GetElementIndex();
+
+		const int32 RenderVerticesStart = ClothCollection->RenderVerticesStart[PatternElementIndex];
+		const int32 RenderVerticesEnd = ClothCollection->RenderVerticesEnd[PatternElementIndex];
+
+		const FVector3f* RawArray = &ClothCollection->RenderPosition[0] + RenderVerticesStart;
+		const int32 NumPoints = RenderVerticesEnd - RenderVerticesStart;
+		const FBoxSphereBounds3f Bounds(RawArray, NumPoints);
+
+		return FBoxSphereBounds(Bounds);
+	}
+
+}		// namespace ChaosClothAssetEditorModeLocals
+
+const FToolTargetTypeRequirements& UChaosClothAssetEditorMode::GetToolTargetRequirements()
+{
+	static const FToolTargetTypeRequirements ToolTargetRequirements =
+		FToolTargetTypeRequirements({
+			UMaterialProvider::StaticClass(),
+			UDynamicMeshCommitter::StaticClass(),
+			UDynamicMeshProvider::StaticClass()
+			});
+	return ToolTargetRequirements;
+}
+
+UChaosClothAssetEditorMode::UChaosClothAssetEditorMode()
+{
+	Info = FEditorModeInfo(
+		EM_ChaosClothAssetEditorModeId,
+		LOCTEXT("ChaosClothAssetEditorModeName", "Cloth"),
+		FSlateIcon(),
+		false);
+}
+
+void UChaosClothAssetEditorMode::Enter()
+{
+	UBaseCharacterFXEditorMode::Enter();
+
+	SelectionModifiedEventHandle = USelection::SelectionChangedEvent.AddLambda([this](UObject*)
+	{
+		if (!GetModeManager() || !GetModeManager()->GetSelectedComponents())
+		{
+			return;
+		}
+
+		constexpr FColor UnselectedColor = FColor(128, 128, 128);
+		const FColor SelectedColor = GetDefault<UEditorStyleSettings>()->SelectionColor.ToFColor(true);
+
+		const USelection* const SelectedComponents = GetModeManager()->GetSelectedComponents();
+
+		check(DynamicMeshComponents.Num() == WireframesToTick.Num());
+
+		for (int MeshIndex = 0; MeshIndex < WireframesToTick.Num(); ++MeshIndex)
+		{
+			const TObjectPtr<UMeshElementsVisualizer> WireframeDisplay = WireframesToTick[MeshIndex];
+			if (SelectedComponents->IsSelected(DynamicMeshComponents[MeshIndex]))
+			{
+				WireframeDisplay->Settings->WireframeColor = SelectedColor;
+			}
+			else
+			{
+				WireframeDisplay->Settings->WireframeColor = UnselectedColor;
+			}
+		}
+	});
+
+}
+
+void UChaosClothAssetEditorMode::AddToolTargetFactories()
+{
+	GetInteractiveToolsContext()->TargetManager->AddTargetFactory(NewObject<UStaticMeshComponentToolTargetFactory>(GetToolManager()));
+	GetInteractiveToolsContext()->TargetManager->AddTargetFactory(NewObject<UVolumeComponentToolTargetFactory>(GetToolManager()));
+	GetInteractiveToolsContext()->TargetManager->AddTargetFactory(NewObject<UDynamicMeshComponentToolTargetFactory>(GetToolManager()));
+}
+
+void UChaosClothAssetEditorMode::RegisterTools()
+{
+	const FChaosClothAssetEditorCommands& CommandInfos = FChaosClothAssetEditorCommands::Get();
+
+	// Note that the identifiers below need to match the command names so that the tool icons can 
+	// be easily retrieved from the active tool name in ChaosClothAssetEditorModeToolkit::OnToolStarted. Otherwise
+	// we would need to keep some other mapping from tool identifier to tool icon.
+
+	RegisterTool(CommandInfos.BeginRemeshTool, FChaosClothAssetEditorCommands::BeginRemeshToolIdentifier, NewObject<URemeshMeshToolBuilder>());
+	RegisterTool(CommandInfos.BeginAttributeEditorTool, FChaosClothAssetEditorCommands::BeginAttributeEditorToolIdentifier, NewObject<UAttributeEditorToolBuilder>());
+	RegisterTool(CommandInfos.BeginWeightMapPaintTool, FChaosClothAssetEditorCommands::BeginWeightMapPaintToolIdentifier, NewObject<UClothEditorWeightMapPaintToolBuilder>());
+}
+
+bool UChaosClothAssetEditorMode::ShouldToolStartBeAllowed(const FString& ToolIdentifier) const
+{
+	// For now we've decided to disallow switch-away on accept/cancel tools in the Cloth editor.
+	if (GetInteractiveToolsContext()->ActiveToolHasAccept())
+	{
+		return false;
+	}
+	
+	return Super::ShouldToolStartBeAllowed(ToolIdentifier);
+}
+
+void UChaosClothAssetEditorMode::CreateToolkit()
+{
+	Toolkit = MakeShared<FChaosClothAssetEditorModeToolkit>();
+}
+
+void UChaosClothAssetEditorMode::OnToolStarted(UInteractiveToolManager* Manager, UInteractiveTool* Tool)
+{
+	bCanTogglePattern2DMode = false;
+}
+
+void UChaosClothAssetEditorMode::OnToolEnded(UInteractiveToolManager* Manager, UInteractiveTool* Tool)
+{
+	bCanTogglePattern2DMode = true;
+
+	UpdateSimulationMeshes();
+	ReinitializeDynamicMeshComponents();
+}
+
+void UChaosClothAssetEditorMode::BindCommands()
+{
+	const FChaosClothAssetEditorCommands& CommandInfos = FChaosClothAssetEditorCommands::Get();
+	const TSharedRef<FUICommandList>& CommandList = Toolkit->GetToolkitCommands();
+
+	// Hook up to Enter/Esc key presses
+	CommandList->MapAction(
+		CommandInfos.AcceptOrCompleteActiveTool,
+		FExecuteAction::CreateLambda([this]() { AcceptActiveToolActionOrTool(); }),
+		FCanExecuteAction::CreateLambda([this]() {
+			return GetInteractiveToolsContext()->CanAcceptActiveTool() || GetInteractiveToolsContext()->CanCompleteActiveTool();
+		}),
+		FGetActionCheckState(),
+		FIsActionButtonVisible(),
+		EUIActionRepeatMode::RepeatDisabled);
+
+	CommandList->MapAction(
+		CommandInfos.CancelOrCompleteActiveTool,
+		FExecuteAction::CreateLambda([this]() { CancelActiveToolActionOrTool(); }),
+		FCanExecuteAction::CreateLambda([this]() {
+			return GetInteractiveToolsContext()->CanCompleteActiveTool() || GetInteractiveToolsContext()->CanCancelActiveTool();
+		}),
+		FGetActionCheckState(),
+		FIsActionButtonVisible(),
+		EUIActionRepeatMode::RepeatDisabled);
+}
+
+
+void UChaosClothAssetEditorMode::Exit()
+{
+	USelection::SelectionChangedEvent.Remove(SelectionModifiedEventHandle);
+
+	for (TObjectPtr<UMeshElementsVisualizer> WireframeDisplay : WireframesToTick)
+	{
+		WireframeDisplay->Disconnect();
+	}
+	WireframesToTick.Reset();
+	DynamicMeshComponents.Reset();
+	PropertyObjectsToTick.Empty();
+	PreviewWorld = nullptr;
+
+	Super::Exit();
+}
+
+void UChaosClothAssetEditorMode::InitializeContexts(FEditorViewportClient& LivePreviewViewportClient, FAssetEditorModeManager& LivePreviewModeManager)
+{
+	PreviewWorld = LivePreviewModeManager.GetPreviewScene()->GetWorld();
+
+	// UContextObjectStore* ContextStore = GetInteractiveToolsContext()->ToolManager->GetContextObjectStore();
+
+	// TODO: If we are going to support tools that operate on the 3D sim mesh, we'll need something like this to pass into those tools:
+	//UClothToolLivePreviewAPI* LivePreviewAPI = NewObject<UClothToolLivePreviewAPI>();
+	//LivePreviewAPI->Initialize(
+	//	LivePreviewWorld,
+	//	LivePreviewModeManager.GetInteractiveToolsContext()->InputRouter,
+	//	[this, LivePreviewViewportClientPtr = &LivePreviewViewportClient](FViewCameraState& CameraStateOut) {
+	//		GetCameraState(*LivePreviewViewportClientPtr, CameraStateOut); 
+	//	});
+	//ContextStore->AddContextObject(LivePreviewAPI);
+
+	// TODO: If we add buttons to the viewport that the tools will need to interact with (e.g. selection of triangles, edges, etc.), we will need something like this:
+	//ContextStore->AddContextObject(&ViewportButtonsAPI);
+}
+
+void UChaosClothAssetEditorMode::GetRestSpaceMesh(UToolTarget* ToolTarget, UE::Geometry::FDynamicMesh3& RestSpaceMesh)
+{
+	// TODO: Get the actual rest-space mesh. NOTE: Currently full 3D mesh (not 2D projection)
+	RestSpaceMesh = UE::ToolTarget::GetDynamicMeshCopy(ToolTarget);
+}
+
+
+void UChaosClothAssetEditorMode::UpdateSimulationMeshes()
+{
+	const UChaosClothAsset* ChaosClothAsset = ClothComponent->GetClothAsset();
+	UE::Chaos::ClothAsset::FClothAdapter ClothAdapter(ChaosClothAsset->GetClothCollection());
+	
+	check(DynamicMeshSourceInfos.Num() == DynamicMeshComponents.Num());
+	
+	for (int DynamicMeshIndex = 0; DynamicMeshIndex < DynamicMeshComponents.Num(); ++DynamicMeshIndex)
+	{
+		const int32 LodIndex = DynamicMeshSourceInfos[DynamicMeshIndex].LodIndex;
+		const int32 PatternIndex = DynamicMeshSourceInfos[DynamicMeshIndex].PatternIndex;
+
+		UE::Chaos::ClothAsset::FClothLodAdapter ClothLodAdapter = ClothAdapter.GetLod(LodIndex);
+		UE::Chaos::ClothAsset::FClothPatternAdapter ClothPatternAdapter = ClothLodAdapter.GetPattern(PatternIndex);
+
+		const UDynamicMeshComponent& DynamicMeshComponent = *DynamicMeshComponents[DynamicMeshIndex];
+		const UE::Geometry::FDynamicMesh3* DynamicMesh = DynamicMeshComponent.GetMesh();
+
+		//
+		// TODO: Set vertices and triangles on the pattern, e.g.:
+		// void Initialize(const TArray<FVector2f>& Positions, const TArray<FVector3f>& RestPositions, const TArray<uint32>& Indices)
+		//
+
+		//
+		// Weight maps
+		//
+
+		if (!DynamicMesh->HasAttributes() || DynamicMesh->Attributes()->NumWeightLayers() == 0)
+		{
+			continue;
+		}
+
+		const int NumInputWeightElements = DynamicMesh->MaxVertexID();
+
+		for (int32 DynamicMeshWeightMapIndex = 0; DynamicMeshWeightMapIndex < DynamicMesh->Attributes()->NumWeightLayers(); ++DynamicMeshWeightMapIndex)
+		{
+			const UE::Geometry::FDynamicMeshWeightAttribute* WeightMapAttribute = DynamicMesh->Attributes()->GetWeightLayer(DynamicMeshWeightMapIndex);
+			const FName WeightMapName = WeightMapAttribute->GetName();
+
+			ClothAdapter.AddWeightMap(WeightMapName);		// Does nothing if weight map already exists
+
+			TArrayView<float> PatternWeights = ClothPatternAdapter.GetWeightMap(WeightMapName);
+			for (int VertexID = 0; VertexID < NumInputWeightElements; ++VertexID)
+			{
+				float Val;
+				WeightMapAttribute->GetValue(VertexID, &Val);
+				PatternWeights[VertexID] = Val;
+			}
+		}
+	}
+}
+
+void UChaosClothAssetEditorMode::CreateToolTargets(const TArray<TObjectPtr<UObject>>& AssetsIn) 
+{
+	// TODO: When we have a cloth component tool target, create it here
+
+}
+
+
+void UChaosClothAssetEditorMode::ReinitializeDynamicMeshComponents()
+{
+	// Clean up any existing DynamicMeshComponents
+	// Save indices of selected mesh components
+	TArray<int32> PreviouslySelectedDynamicMeshComponents;
+	const int32 TotalNumberExistingDynamicMeshComponents = DynamicMeshComponents.Num();
+
+	USelection* SelectedComponents = GetModeManager()->GetSelectedComponents();
+	for (int32 DynamicMeshComponentIndex = 0; DynamicMeshComponentIndex < DynamicMeshComponents.Num(); ++DynamicMeshComponentIndex)
+	{
+		// TODO: Check for any outstanding changes on the dynamic mesh component?
+
+		TObjectPtr<UDynamicMeshComponent> DynamicMeshComp = DynamicMeshComponents[DynamicMeshComponentIndex];
+		DynamicMeshComp->UnregisterComponent();
+
+		if (SelectedComponents->IsSelected(DynamicMeshComp))
+		{
+			PreviouslySelectedDynamicMeshComponents.Add(DynamicMeshComponentIndex);
+			SelectedComponents->Deselect(DynamicMeshComp);
+		}
+	}
+
+	for (const TObjectPtr<UMeshElementsVisualizer> Wireframe : WireframesToTick)
+	{
+		Wireframe->Disconnect();
+	}
+
+	DynamicMeshComponents.Empty();
+	DynamicMeshComponentParentActors.Empty();
+	PropertyObjectsToTick.Empty();	// TODO: We only want to empty the wireframe display properties. Is anything else using this array?
+	WireframesToTick.Empty();
+	DynamicMeshSourceInfos.Empty();
+
+	if (!ClothComponent)
+	{
+		return;
+	}
+
+	const UChaosClothAsset* ChaosClothAsset = ClothComponent->GetClothAsset();
+
+	if (!ChaosClothAsset)
+	{
+		return;
+	}
+
+	const UE::Chaos::ClothAsset::FClothConstAdapter ClothAdapter(ChaosClothAsset->GetClothCollection());
+
+	for (int32 LodIndex = 0; LodIndex < ClothAdapter.GetNumLods(); ++LodIndex)
+	{
+		const UE::Chaos::ClothAsset::FClothLodConstAdapter ClothLodAdapter = ClothAdapter.GetLod(LodIndex);
+
+		if (bCombineAllPatterns)
+		{
+			// TODO: Enable a mode where all patterns are edited as one dynamic mesh component. Will require some fun bookkeeping.
+			check(0);
+		}
+		else
+		{
+			const int32 NumComponents = ClothLodAdapter.GetNumPatterns();
+			DynamicMeshComponents.SetNum(NumComponents);
+			DynamicMeshComponentParentActors.SetNum(NumComponents);
+
+			for (int32 PatternIndex = 0; PatternIndex < ClothLodAdapter.GetNumPatterns(); ++PatternIndex)
+			{
+				UE::Geometry::FDynamicMesh3 PatternMesh;
+				FClothPatternToDynamicMesh Converter;
+				Converter.Convert(ChaosClothAsset, LodIndex, PatternIndex, bPattern2DMode, PatternMesh);
+
+				// We only need an actor to allow use of HHitProxy for selection
+				const FRotator Rotation(0.0f, 0.0f, 0.0f);
+				const FActorSpawnParameters SpawnInfo;
+				const TObjectPtr<AActor> ParentActor = this->GetWorld()->SpawnActor<AActor>(FVector::ZeroVector, Rotation, SpawnInfo);
+				DynamicMeshComponentParentActors[PatternIndex] = ParentActor;
+
+				TObjectPtr<UDynamicMeshComponent> PatternMeshComponent = NewObject<UDynamicMeshComponent>(ParentActor);
+				PatternMeshComponent->SetMesh(MoveTemp(PatternMesh));
+				PatternMeshComponent->RegisterComponentWithWorld(this->GetWorld());
+				DynamicMeshComponents[PatternIndex] = PatternMeshComponent;
+				DynamicMeshSourceInfos.Add(FDynamicMeshSourceInfo{ LodIndex, PatternIndex });
+			}
+		}
+
+		// If we found any patterns, use this LOD
+		// TODO: Give the user the ability to specify which LOD to display
+		if (ClothLodAdapter.GetNumPatterns() > 0)
+		{
+			break;
+		}
+	}
+
+	for (TObjectPtr<UDynamicMeshComponent> RestSpaceMeshComponent : DynamicMeshComponents)
+	{
+		// Set up the wireframe display of the rest space mesh.
+		TObjectPtr<UMeshElementsVisualizer> WireframeDisplay = NewObject<UMeshElementsVisualizer>(this);
+		WireframeDisplay->CreateInWorld(GetWorld(), FTransform::Identity);
+
+		WireframeDisplay->Settings->DepthBias = 2.0;
+		WireframeDisplay->Settings->bAdjustDepthBiasUsingMeshSize = false;
+		WireframeDisplay->Settings->bShowWireframe = true;
+		WireframeDisplay->Settings->bShowBorders = true;
+		WireframeDisplay->Settings->bShowUVSeams = false;
+		WireframeDisplay->Settings->bShowNormalSeams = false;
+
+		// These are not exposed at the visualizer level yet
+		// TODO: Should they be?
+		WireframeDisplay->WireframeComponent->BoundaryEdgeThickness = 2;
+
+		WireframeDisplay->SetMeshAccessFunction([RestSpaceMeshComponent](UMeshElementsVisualizer::ProcessDynamicMeshFunc ProcessFunc)
+		{
+			ProcessFunc(*RestSpaceMeshComponent->GetMesh());
+		});
+
+		RestSpaceMeshComponent->OnMeshChanged.Add(
+			FSimpleMulticastDelegate::FDelegate::CreateLambda([WireframeDisplay]()
+		{
+			WireframeDisplay->NotifyMeshChanged();
+		}));
+
+		// The settings object and wireframe are not part of a tool, so they won't get ticked like they
+		// are supposed to (to enable property watching), unless we add this here.
+		PropertyObjectsToTick.Add(WireframeDisplay->Settings);
+		int32 WireframeIndex = WireframesToTick.Add(WireframeDisplay);
+
+		RestSpaceMeshComponent->MarkRenderStateDirtyEvent.AddLambda( [WireframeIndex, this](UActorComponent& ActorComponent)
+		{
+			if (WireframeIndex >= WireframesToTick.Num() || WireframeIndex >= DynamicMeshComponents.Num())
+			{
+				return;
+			}
+
+			TObjectPtr<UMeshElementsVisualizer> WireframeDisplay = WireframesToTick[WireframeIndex];
+			TObjectPtr<UDynamicMeshComponent> RestSpaceMesh = DynamicMeshComponents[WireframeIndex];
+
+			if (!WireframeDisplay || !RestSpaceMesh)
+			{
+				return;
+			}
+
+			bool bRestSpaceMeshVisible = RestSpaceMesh->GetVisibleFlag();
+			WireframeDisplay->SetAllVisible(bRestSpaceMeshVisible);
+			WireframeDisplay->Settings->bVisible = bRestSpaceMeshVisible;
+		});
+	}
+
+	const bool bNumberDynamicMeshesUnchanged = (TotalNumberExistingDynamicMeshComponents == DynamicMeshComponents.Num());
+	if (bNumberDynamicMeshesUnchanged && PreviouslySelectedDynamicMeshComponents.Num() > 0)
+	{
+		SelectedComponents->Modify();
+		SelectedComponents->BeginBatchSelectOperation();
+		for (int32 DynamicMeshComponentIndex : PreviouslySelectedDynamicMeshComponents)
+		{
+			if (DynamicMeshComponentIndex < DynamicMeshComponents.Num())
+			{
+				SelectedComponents->Select(DynamicMeshComponents[DynamicMeshComponentIndex]);
+			}
+		}
+		SelectedComponents->EndBatchSelectOperation();
+	}
+
+}
+
+void UChaosClothAssetEditorMode::RefocusRestSpaceViewportClient()
+{
+	TSharedPtr<FEditorViewportClient, ESPMode::ThreadSafe> PinnedVC = RestSpaceViewportClient.Pin();
+	if (PinnedVC.IsValid())
+	{
+		// This will happen in FocusViewportOnBox anyways; do it now to get a consistent end result
+		PinnedVC->ToggleOrbitCamera(false);
+
+		const FBox SceneBounds = SceneBoundingBox();
+		if (IsPattern2DModeActive())
+		{
+			// 2D pattern
+			PinnedVC->SetInitialViewTransform(ELevelViewportType::LVT_Perspective, FVector(0, 150, 100), FRotator(-90, 90, 0), DEFAULT_ORTHOZOOM);
+		}
+		else
+		{
+			// 3D rest space
+			PinnedVC->SetInitialViewTransform(ELevelViewportType::LVT_Perspective, FVector(0, 150, 200), FRotator(0, 0, 0), DEFAULT_ORTHOZOOM);
+		}
+
+		constexpr bool bInstant = true;
+		PinnedVC->FocusViewportOnBox(SceneBounds, bInstant);
+	}
+}
+
+void UChaosClothAssetEditorMode::InitializeTargets(const TArray<TObjectPtr<UObject>>& AssetsIn)
+{
+	// InitializeContexts needs to have been called first so that we have the 3d preview world ready.
+	check(PreviewWorld);
+
+	OriginalObjectsToEdit = AssetsIn;
+
+	// TODO: If we have multiple objects to initialize, position them relative to each other somehow
+	Transforms.Init(FTransform::Identity, AssetsIn.Num());
+
+	CreateToolTargets(AssetsIn);
+
+	// NOTE: This ensure will fire until we are able to build a tool target associated with ClothComponent
+	//if (!ensure(AssetsIn.Num() == ToolTargets.Num()))
+	//{
+	//	return;
+	//}
+
+	for (TObjectPtr<UObject> Asset : AssetsIn)
+	{
+		if (UChaosClothAsset* ChaosClothAsset = Cast<UChaosClothAsset>(Asset))
+		{
+			ClothComponent = NewObject<UChaosClothComponent>();
+			ClothComponent->SetClothAsset(ChaosClothAsset);
+			ClothComponent->RegisterComponentWithWorld(PreviewWorld);
+
+			break;
+		}
+	}
+
+	ReinitializeDynamicMeshComponents();
+	RefocusRestSpaceViewportClient();
+
+	// Prep things for layer/channel selection
+	AssetNames.Reset(ToolTargets.Num());
+	for (const TObjectPtr<UToolTarget>& ToolTarget : ToolTargets)
+	{
+		AssetNames.Add(UE::ToolTarget::GetHumanReadableName(ToolTarget));
+	}
+}
+
+bool UChaosClothAssetEditorMode::HaveUnappliedChanges()
+{
+	return ModifiedAssetIDs.Num() > 0;
+}
+
+void UChaosClothAssetEditorMode::GetAssetsWithUnappliedChanges(TArray<TObjectPtr<UObject>> UnappliedAssetsOut)
+{
+	for (int32 AssetID : ModifiedAssetIDs)
+	{
+		// The asset ID corresponds to the index into OriginalObjectsToEdit
+		UnappliedAssetsOut.Add(OriginalObjectsToEdit[AssetID]);
+	}
+}
+
+void UChaosClothAssetEditorMode::ApplyChanges()
+{
+	GetToolManager()->BeginUndoTransaction(LOCTEXT("ChaosClothAssetEditorApplyChangesTransaction", "Cloth Editor Apply Changes"));
+	GetToolManager()->EndUndoTransaction();
+}
+
+void UChaosClothAssetEditorMode::ModeTick(float DeltaTime)
+{
+	Super::ModeTick(DeltaTime);
+	
+	for (TObjectPtr<UInteractiveToolPropertySet>& Propset : PropertyObjectsToTick)
+	{
+		if (Propset)
+		{
+			if (Propset->IsPropertySetEnabled())
+			{
+				Propset->CheckAndUpdateWatched();
+			}
+			else
+			{
+				Propset->SilentUpdateWatched();
+			}
+		}
+	}
+
+	for (TWeakObjectPtr<UMeshElementsVisualizer> WireframeDisplay : WireframesToTick)
+	{
+		if (WireframeDisplay.IsValid())
+		{
+			WireframeDisplay->OnTick(DeltaTime);
+		}
+	}
+
+	if (PreviewWorld)
+	{
+		PreviewWorld->Tick(ELevelTick::LEVELTICK_All, DeltaTime);
+	}
+}
+
+
+FBox UChaosClothAssetEditorMode::SceneBoundingBox() const
+{
+	FBoxSphereBounds TotalBounds(ForceInitToZero);
+	
+	bool bFoundMeshBounds = false;
+	for (const TObjectPtr<UDynamicMeshComponent> DynamicMeshComponent : DynamicMeshComponents)
+	{
+		const FBoxSphereBounds& CurrentBounds = DynamicMeshComponent->Bounds;
+
+		if (!bFoundMeshBounds)
+		{
+			TotalBounds = CurrentBounds;
+			bFoundMeshBounds = true;
+		}
+		else
+		{
+			TotalBounds = TotalBounds + CurrentBounds;
+		}
+	}
+
+	return TotalBounds.GetBox();
+}
+
+FBox UChaosClothAssetEditorMode::SelectionBoundingBox() const
+{
+	FBoxSphereBounds TotalBounds;
+
+	const USelection* const SelectedComponents = GetModeManager()->GetSelectedComponents();
+
+	bool bFoundSelected = false;
+	for (const TObjectPtr<UDynamicMeshComponent> DynamicMeshComponent : DynamicMeshComponents)
+	{
+		if (SelectedComponents->IsSelected(DynamicMeshComponent))
+		{
+			const FBoxSphereBounds& CurrentBounds = DynamicMeshComponent->Bounds;
+
+			if (!bFoundSelected)
+			{
+				TotalBounds = CurrentBounds;
+				bFoundSelected = true;
+			}
+			else
+			{
+				TotalBounds = TotalBounds + CurrentBounds;
+			}
+			
+		}
+	}
+
+	if (bFoundSelected)
+	{
+		return TotalBounds.GetBox();
+	}
+	else
+	{
+		// Nothing selected, return the whole scene
+		return SceneBoundingBox();
+	}
+
+}
+
+
+FBox UChaosClothAssetEditorMode::PreviewBoundingBox() const
+{
+	FBoxSphereBounds TotalBounds(ForceInit);
+
+	bool bAnyValid = false;
+
+	if (ClothComponent && ClothComponent->GetClothAsset())
+	{
+		const UChaosClothAsset* ChaosClothAsset = ClothComponent->GetClothAsset();
+		const UE::Chaos::ClothAsset::FClothConstAdapter ClothAdapter(ChaosClothAsset->GetClothCollection());
+
+		for (int32 LodIndex = 0; LodIndex < ClothAdapter.GetNumLods(); ++LodIndex)
+		{
+			const UE::Chaos::ClothAsset::FClothLodConstAdapter ClothLodAdapter = ClothAdapter.GetLod(LodIndex);
+
+			for (int32 PatternIndex = 0; PatternIndex < ClothLodAdapter.GetNumPatterns(); ++PatternIndex)
+			{
+				const UE::Chaos::ClothAsset::FClothPatternConstAdapter ClothPatternAdapter = ClothLodAdapter.GetPattern(PatternIndex);
+				FBoxSphereBounds BoxSphere = ChaosClothAssetEditorModeLocals::RenderMeshBounds(ClothPatternAdapter);
+				TotalBounds = TotalBounds + BoxSphere;
+				bAnyValid = true;
+			}
+		}
+	}
+
+	if (bAnyValid)
+	{
+		return TotalBounds.GetBox();
+	}
+	else
+	{
+		return FBox(ForceInitToZero);
+	}
+}
+
+bool UChaosClothAssetEditorMode::IsPattern2DModeActive() const
+{
+	return bPattern2DMode;
+}
+
+void UChaosClothAssetEditorMode::TogglePatternMode()
+{
+	bPattern2DMode = !bPattern2DMode;
+	ReinitializeDynamicMeshComponents();
+
+	TSharedPtr<FChaosClothEditorRestSpaceViewportClient> VC = RestSpaceViewportClient.Pin();
+	if (VC.IsValid())
+	{
+		VC->Set2DMode(bPattern2DMode);
+	}
+
+	RefocusRestSpaceViewportClient();
+}
+
+bool UChaosClothAssetEditorMode::CanTogglePatternMode() const
+{
+	return bCanTogglePattern2DMode;
+}
+
+
+void UChaosClothAssetEditorMode::SetRestSpaceViewportClient(TWeakPtr<FChaosClothEditorRestSpaceViewportClient, ESPMode::ThreadSafe> InViewportClient)
+{
+	RestSpaceViewportClient = InViewportClient;
+
+	TSharedPtr<FChaosClothEditorRestSpaceViewportClient> VC = RestSpaceViewportClient.Pin();
+	if (VC.IsValid())
+	{
+		VC->Set2DMode(bPattern2DMode);
+	}
+}
+
+#undef LOCTEXT_NAMESPACE
+
diff --git a/Engine/Plugins/Experimental/ChaosClothAssetEditor/Source/ChaosClothAssetEditor/Private/ChaosClothAsset/ClothEditorModeToolkit.cpp b/Engine/Plugins/Experimental/ChaosClothAssetEditor/Source/ChaosClothAssetEditor/Private/ChaosClothAsset/ClothEditorModeToolkit.cpp
new file mode 100644
--- /dev/null
+++ b/Engine/Plugins/Experimental/ChaosClothAssetEditor/Source/ChaosClothAssetEditor/Private/ChaosClothAsset/ClothEditorModeToolkit.cpp
@@ -0,0 +1,38 @@
+// Copyright Epic Games, Inc. All Rights Reserved.
+
+#include "ChaosClothAsset/ClothEditorModeToolkit.h"
+#include "ChaosClothAsset/ClothEditorCommands.h"
+#include "ChaosClothAsset/ClothEditorStyle.h"
+#include "Framework/MultiBox/MultiBoxBuilder.h"
+
+#define LOCTEXT_NAMESPACE "FChaosClothAssetEditorModeToolkit"
+
+FName FChaosClothAssetEditorModeToolkit::GetToolkitFName() const
+{
+	return FName("ChaosClothAssetEditorMode");
+}
+
+FText FChaosClothAssetEditorModeToolkit::GetBaseToolkitName() const
+{
+	return NSLOCTEXT("ChaosClothAssetEditorModeToolkit", "DisplayName", "ChaosClothAssetEditorMode");
+}
+
+void FChaosClothAssetEditorModeToolkit::BuildToolPalette(FName PaletteIndex, class FToolBarBuilder& ToolbarBuilder)
+{
+	const FChaosClothAssetEditorCommands& Commands = FChaosClothAssetEditorCommands::Get();
+
+	if (PaletteIndex == FBaseCharacterFXEditorModeToolkit::ToolsTabName)
+	{
+		ToolbarBuilder.AddToolBarButton(Commands.BeginRemeshTool);
+		ToolbarBuilder.AddToolBarButton(Commands.BeginAttributeEditorTool);
+		ToolbarBuilder.AddToolBarButton(Commands.BeginWeightMapPaintTool);
+	}
+}
+
+const FSlateBrush* FChaosClothAssetEditorModeToolkit::GetActiveToolIcon(const FString& ActiveToolIdentifier) const
+{
+	FName ActiveToolIconName = ISlateStyle::Join(FChaosClothAssetEditorCommands::Get().GetContextName(), TCHAR_TO_ANSI(*ActiveToolIdentifier));
+	return FChaosClothAssetEditorStyle::Get().GetOptionalBrush(ActiveToolIconName);
+}
+
+#undef LOCTEXT_NAMESPACE
diff --git a/Engine/Plugins/Experimental/ChaosClothAssetEditor/Source/ChaosClothAssetEditor/Private/ChaosClothAsset/ClothEditorRestSpaceViewportClient.cpp b/Engine/Plugins/Experimental/ChaosClothAssetEditor/Source/ChaosClothAssetEditor/Private/ChaosClothAsset/ClothEditorRestSpaceViewportClient.cpp
new file mode 100644
--- /dev/null
+++ b/Engine/Plugins/Experimental/ChaosClothAssetEditor/Source/ChaosClothAssetEditor/Private/ChaosClothAsset/ClothEditorRestSpaceViewportClient.cpp
@@ -0,0 +1,212 @@
+// Copyright Epic Games, Inc. All Rights Reserved.
+
+#include "ChaosClothAsset/ClothEditorRestSpaceViewportClient.h"
+
+#include "BaseBehaviors/ClickDragBehavior.h"
+#include "BaseBehaviors/MouseWheelBehavior.h"
+#include "EngineUtils.h"
+#include "EditorModeManager.h"
+#include "Engine/Selection.h"
+#include "Behaviors/2DViewportBehaviorTargets.h"
+#include "EdModeInteractiveToolsContext.h"
+#include "CameraController.h"
+#include "Components/DynamicMeshComponent.h"
+
+constexpr float CameraFarPlaneWorldZ = -10.0f;
+constexpr float CameraNearPlaneProportionZ = 0.8f;
+constexpr float DefaultNearClipPlaneDistance = KINDA_SMALL_NUMBER;
+constexpr float DefaultFarClipPlaneDistance = UE_FLOAT_HUGE_DISTANCE / 4.0f;
+
+FChaosClothEditorRestSpaceViewportClient::FChaosClothEditorRestSpaceViewportClient(FEditorModeTools* InModeTools,
+	FPreviewScene* InPreviewScene,
+	const TWeakPtr<SEditorViewport>& InEditorViewportWidget) :
+	FEditorViewportClient(InModeTools, InPreviewScene, InEditorViewportWidget)
+{
+	OverrideNearClipPlane(DefaultNearClipPlaneDistance);
+	OverrideFarClipPlane(DefaultFarClipPlaneDistance);
+
+	BehaviorSet = NewObject<UInputBehaviorSet>();
+
+	// We'll have the priority of our viewport manipulation behaviors be lower (i.e. higher
+	// numerically) than both the gizmo default and the tool default.
+	static constexpr int DEFAULT_VIEWPORT_BEHAVIOR_PRIORITY = 150;
+
+	ScrollBehaviorTarget = MakeUnique<FEditor2DScrollBehaviorTarget>(this);
+	UClickDragInputBehavior* ScrollBehavior = NewObject<UClickDragInputBehavior>();
+	ScrollBehavior->Initialize(ScrollBehaviorTarget.Get());
+	ScrollBehavior->SetDefaultPriority(DEFAULT_VIEWPORT_BEHAVIOR_PRIORITY);
+	ScrollBehavior->SetUseRightMouseButton();
+	BehaviorSet->Add(ScrollBehavior);
+	SavedBehaviors.Add(ScrollBehavior);
+
+	ZoomBehaviorTarget = MakeUnique<FEditor2DMouseWheelZoomBehaviorTarget>(this);
+	ZoomBehaviorTarget->SetCameraFarPlaneWorldZ(CameraFarPlaneWorldZ);
+	ZoomBehaviorTarget->SetCameraNearPlaneProportionZ(CameraNearPlaneProportionZ);
+	ZoomBehaviorTarget->SetZoomLimits(0.001, 100000);
+	UMouseWheelInputBehavior* ZoomBehavior = NewObject<UMouseWheelInputBehavior>();
+	ZoomBehavior->Initialize(ZoomBehaviorTarget.Get());
+	ZoomBehavior->SetDefaultPriority(DEFAULT_VIEWPORT_BEHAVIOR_PRIORITY);
+	BehaviorSet->Add(ZoomBehavior);
+	SavedBehaviors.Add(ZoomBehavior);
+
+	ModeTools->GetInteractiveToolsContext()->InputRouter->RegisterSource(this);
+}
+
+void FChaosClothEditorRestSpaceViewportClient::Set2DMode(bool In2DMode)
+{
+	b2DMode = In2DMode;
+
+	BehaviorSet->RemoveAll();
+	ModeTools->GetInteractiveToolsContext()->InputRouter->DeregisterSource(this);
+
+	if (b2DMode)
+	{
+		for (UInputBehavior* Behavior : SavedBehaviors)
+		{
+			BehaviorSet->Add(Behavior);
+		}
+		ModeTools->GetInteractiveToolsContext()->InputRouter->RegisterSource(this);
+	}
+	else
+	{
+		OverrideNearClipPlane(DefaultNearClipPlaneDistance);
+		OverrideFarClipPlane(DefaultFarClipPlaneDistance);
+	}
+
+}
+
+const UInputBehaviorSet* FChaosClothEditorRestSpaceViewportClient::GetInputBehaviors() const
+{
+	return BehaviorSet;
+}
+
+// Collects UObjects that we don't want the garbage collecter to throw away under us
+void FChaosClothEditorRestSpaceViewportClient::AddReferencedObjects(FReferenceCollector& Collector)
+{
+	FEditorViewportClient::AddReferencedObjects(Collector);
+	Collector.AddReferencedObject(BehaviorSet);
+	Collector.AddReferencedObjects(SavedBehaviors);
+}
+
+bool FChaosClothEditorRestSpaceViewportClient::ShouldOrbitCamera() const
+{
+	return !b2DMode;
+}
+
+bool FChaosClothEditorRestSpaceViewportClient::InputKey(const FInputKeyEventArgs& EventArgs)
+{
+	if (!b2DMode)
+	{
+		return FEditorViewportClient::InputKey(EventArgs);
+	}
+
+
+	// We'll support disabling input like our base class, even if it does not end up being used.
+	if (bDisableInput)
+	{
+		return true;
+	}
+
+	// Our viewport manipulation is placed in the input router that ModeTools manages
+	if (ModeTools->InputKey(this, EventArgs.Viewport, EventArgs.Key, EventArgs.Event))
+	{
+		return true;
+	}
+	
+	//
+	// Recreate the click handling from FEditorViewportClient::Internal_InputKey so that we can still have 
+	// selection work, but skip the code dealing with camera controls
+	//
+
+	bool bHandled = false;
+
+	const FInputEventState InputState(EventArgs.Viewport, EventArgs.Key, EventArgs.Event);
+	const FViewport* InViewport = EventArgs.Viewport;
+	const EInputEvent Event = EventArgs.Event;
+
+	const bool bWasCursorVisible = InViewport->IsCursorVisible();
+	const bool bWasSoftwareCursorVisible = InViewport->IsSoftwareCursorVisible();
+
+	const bool AltDown = InputState.IsAltButtonPressed();
+	const bool ShiftDown = InputState.IsShiftButtonPressed();
+	const bool ControlDown = InputState.IsCtrlButtonPressed();
+
+	RequiredCursorVisibiltyAndAppearance.bDontResetCursor = false;
+	UpdateRequiredCursorVisibility();
+
+	if (bWasCursorVisible != RequiredCursorVisibiltyAndAppearance.bHardwareCursorVisible || bWasSoftwareCursorVisible != RequiredCursorVisibiltyAndAppearance.bSoftwareCursorVisible)
+	{
+		bHandled = true;
+	}
+
+	// Compute a view.
+	FSceneViewFamilyContext ViewFamily(FSceneViewFamily::ConstructionValues(
+		InViewport,
+		GetScene(),
+		EngineShowFlags)
+		.SetRealtimeUpdate(IsRealtime()));
+	FSceneView* View = CalcSceneView(&ViewFamily);
+
+	// Start tracking if any mouse button is down and it was a tracking event (MouseButton/Ctrl/Shift/Alt):
+	if (InputState.IsAnyMouseButtonDown()
+		&& (Event == IE_Pressed || Event == IE_Released)
+		&& (InputState.IsMouseButtonEvent() || InputState.IsCtrlButtonEvent() || InputState.IsAltButtonEvent() || InputState.IsShiftButtonEvent()))
+	{
+		StartTrackingDueToInput(InputState, *View);
+		return true;
+	}
+
+	// If we are tracking and no mouse button is down and this input event released the mouse button stop tracking and process any clicks if necessary
+	if (bIsTracking && !InputState.IsAnyMouseButtonDown() && InputState.IsMouseButtonEvent())
+	{
+		// Handle possible mouse click viewport
+		ProcessClickInViewport(InputState, *View);
+
+		// Stop tracking if no mouse button is down
+		StopTracking();
+
+		bHandled |= true;
+	}
+
+	// apply the visibility and set the cursor positions
+	ApplyRequiredCursorVisibility(true);
+
+	return bHandled;
+}
+
+
+void FChaosClothEditorRestSpaceViewportClient::ProcessClick(FSceneView& View, HHitProxy* HitProxy, FKey Key, EInputEvent Event, uint32 HitX, uint32 HitY)
+{
+	FEditorViewportClient::ProcessClick(View, HitProxy, Key, Event, HitX, HitY);
+
+	// TODO: Add/modify selection if modifier keys are pressed
+	USelection* SelectedComponents = ModeTools->GetSelectedComponents();
+	SelectedComponents->Modify();
+	SelectedComponents->BeginBatchSelectOperation();
+	SelectedComponents->DeselectAll();
+
+	if (HitProxy && HitProxy->IsA(HActor::StaticGetType()))
+	{
+		const HActor* ActorProxy = static_cast<HActor*>(HitProxy);
+		if (ActorProxy && ActorProxy->Actor )
+		{
+			const AActor* Actor = ActorProxy->Actor;
+			const TSet<UActorComponent*>& OwnedComponents = Actor->GetComponents();
+			for (UActorComponent* Component : OwnedComponents)
+			{
+				if (UDynamicMeshComponent* DynamicMeshComp = Cast<UDynamicMeshComponent>(Component))
+				{
+					SelectedComponents->Select(DynamicMeshComp);
+				}
+			}
+		}
+	}
+
+	SelectedComponents->EndBatchSelectOperation();
+
+}
+
+void FChaosClothEditorRestSpaceViewportClient::SetEditorViewportWidget(TWeakPtr<SEditorViewport> InEditorViewportWidget)
+{
+	EditorViewportWidget = InEditorViewportWidget;
+}
diff --git a/Engine/Plugins/Experimental/ChaosClothAssetEditor/Source/ChaosClothAssetEditor/Private/ChaosClothAsset/ClothEditorStyle.cpp b/Engine/Plugins/Experimental/ChaosClothAssetEditor/Source/ChaosClothAssetEditor/Private/ChaosClothAsset/ClothEditorStyle.cpp
new file mode 100644
--- /dev/null
+++ b/Engine/Plugins/Experimental/ChaosClothAssetEditor/Source/ChaosClothAssetEditor/Private/ChaosClothAsset/ClothEditorStyle.cpp
@@ -0,0 +1,74 @@
+// Copyright Epic Games, Inc. All Rights Reserved.
+
+#include "ChaosClothAsset/ClothEditorStyle.h"
+#include "ChaosClothAsset/ClothEditorCommands.h"
+#include "Styling/SlateTypes.h"
+#include "Styling/CoreStyle.h"
+#include "Styling/SlateStyleMacros.h"
+#include "Interfaces/IPluginManager.h"
+#include "Misc/Paths.h"
+
+const FName FChaosClothAssetEditorStyle::StyleName("ClothStyle");
+
+FString FChaosClothAssetEditorStyle::InContent(const FString& RelativePath, const ANSICHAR* Extension)
+{
+	static FString ContentDir = IPluginManager::Get().FindPlugin(TEXT("ChaosClothAssetEditor"))->GetContentDir();
+	FString Path = (ContentDir / RelativePath) + Extension;
+	return Path;
+}
+
+FChaosClothAssetEditorStyle::FChaosClothAssetEditorStyle()
+	: FSlateStyleSet(StyleName)
+{
+	// Used FUVEditorStyle and FModelingToolsEditorModeStyle and as models
+
+	// Some standard icon sizes used elsewhere in the editor
+	const FVector2D Icon8x8(8.0f, 8.0f);
+	const FVector2D Icon16x16(16.0f, 16.0f);
+	const FVector2D Icon20x20(20.0f, 20.0f);
+	const FVector2D Icon28x28(28.0f, 28.0f);
+	const FVector2D Icon40x40(40.0f, 40.0f);
+	const FVector2D Icon120(120.0f, 120.0f);
+
+	// Icon sizes used in this style set
+	const FVector2D ViewportToolbarIconSize = Icon16x16;
+	const FVector2D ToolbarIconSize = Icon20x20;
+
+	// TODO: Set this once we get our own UI content
+	FString ContentDir = IPluginManager::Get().FindPlugin(TEXT("ModelingToolsEditorMode"))->GetContentDir();
+	SetContentRoot(ContentDir);
+	SetCoreContentRoot(FPaths::EngineContentDir() / TEXT("Slate"));
+
+	FString PropertyNameString = "ChaosClothAssetEditor." + FChaosClothAssetEditorCommands::BeginRemeshToolIdentifier;
+	Set(*PropertyNameString, new FSlateImageBrush(FChaosClothAssetEditorStyle::InContent("Icons/Remesh_40x", ".png"), ToolbarIconSize ));
+	
+	PropertyNameString = "ChaosClothAssetEditor." + FChaosClothAssetEditorCommands::BeginAttributeEditorToolIdentifier;
+	Set(*PropertyNameString, new FSlateImageBrush(FChaosClothAssetEditorStyle::InContent("Icons/AttributeEditor_40x", ".png"), ToolbarIconSize));
+
+	PropertyNameString = "ChaosClothAssetEditor." + FChaosClothAssetEditorCommands::BeginWeightMapPaintToolIdentifier;
+	Set(*PropertyNameString, new FSlateImageBrush(FChaosClothAssetEditorStyle::InContent("Icons/ModelingAttributePaint_x40", ".png"), ToolbarIconSize));
+
+	PropertyNameString = "ChaosClothAssetEditor." + FChaosClothAssetEditorCommands::ToggleSimMeshWireframeIdentifier;
+	Set(*PropertyNameString, new FSlateImageBrush(FChaosClothAssetEditorStyle::InContent("Icons/icon_ViewMode_BrushWireframe_16px", ".png"), ViewportToolbarIconSize));		// Placeholder
+
+	PropertyNameString = "ChaosClothAssetEditor." + FChaosClothAssetEditorCommands::ToggleRenderMeshWireframeIdentifier;
+	Set(*PropertyNameString, new FSlateImageBrush(FChaosClothAssetEditorStyle::InContent("Icons/icon_ViewMode_BrushWireframe_16px", ".png"), ViewportToolbarIconSize));		// Placeholder
+
+	PropertyNameString = "ChaosClothAssetEditor." + FChaosClothAssetEditorCommands::TogglePatternModeIdentifier;
+	Set(*PropertyNameString, new FSlateImageBrush(FChaosClothAssetEditorStyle::InContent("Icons/TogglePatternMode_40x", ".png"), ViewportToolbarIconSize));
+
+	FSlateStyleRegistry::RegisterSlateStyle(*this);
+}
+
+FChaosClothAssetEditorStyle::~FChaosClothAssetEditorStyle()
+{
+	FSlateStyleRegistry::UnRegisterSlateStyle(*this);
+}
+
+FChaosClothAssetEditorStyle& FChaosClothAssetEditorStyle::Get()
+{
+	static FChaosClothAssetEditorStyle Inst;
+	return Inst;
+}
+
+
diff --git a/Engine/Plugins/Experimental/ChaosClothAssetEditor/Source/ChaosClothAssetEditor/Private/ChaosClothAsset/SClothEditor3DViewport.cpp b/Engine/Plugins/Experimental/ChaosClothAssetEditor/Source/ChaosClothAssetEditor/Private/ChaosClothAsset/SClothEditor3DViewport.cpp
new file mode 100644
--- /dev/null
+++ b/Engine/Plugins/Experimental/ChaosClothAssetEditor/Source/ChaosClothAssetEditor/Private/ChaosClothAsset/SClothEditor3DViewport.cpp
@@ -0,0 +1,44 @@
+// Copyright Epic Games, Inc. All Rights Reserved.
+
+#include "ChaosClothAsset/SClothEditor3DViewport.h"
+#include "ChaosClothAsset/ClothEditor3DViewportClient.h"
+#include "ChaosClothAsset/SClothEditor3DViewportToolBar.h"
+#include "ChaosClothAsset/ClothEditorCommands.h"
+#include "ChaosClothAsset/ClothEditorMode.h"
+#include "EditorModeTools.h"
+
+#define LOCTEXT_NAMESPACE "SChaosClothAssetEditor3DViewport"
+
+void SChaosClothAssetEditor3DViewport::BindCommands()
+{
+	SAssetEditorViewport::BindCommands();
+	const FChaosClothAssetEditorCommands& CommandInfos = FChaosClothAssetEditorCommands::Get();
+
+	CommandList->MapAction(
+		CommandInfos.ToggleSimMeshWireframe,
+		FExecuteAction::CreateLambda([this]() 
+		{
+			TSharedPtr<FChaosClothAssetEditor3DViewportClient> ClothViewportClient = StaticCastSharedPtr<FChaosClothAssetEditor3DViewportClient>(Client);
+			ClothViewportClient->EnableSimMeshWireframe(!ClothViewportClient->SimMeshWireframeEnabled());
+		}),
+		FCanExecuteAction::CreateLambda([this]() { return true; }),
+		FIsActionChecked::CreateLambda([this]() { return StaticCastSharedPtr<FChaosClothAssetEditor3DViewportClient>(Client)->SimMeshWireframeEnabled(); }));
+
+	CommandList->MapAction(
+		CommandInfos.ToggleRenderMeshWireframe,
+		FExecuteAction::CreateLambda([this]()
+		{
+			TSharedPtr<FChaosClothAssetEditor3DViewportClient> ClothViewportClient = StaticCastSharedPtr<FChaosClothAssetEditor3DViewportClient>(Client);
+			ClothViewportClient->EnableRenderMeshWireframe(!ClothViewportClient->RenderMeshWireframeEnabled());
+		}),
+		FCanExecuteAction::CreateLambda([this]() { return true; }),
+		FIsActionChecked::CreateLambda([this]() { return StaticCastSharedPtr<FChaosClothAssetEditor3DViewportClient>(Client)->RenderMeshWireframeEnabled(); }));
+}
+
+TSharedPtr<SWidget> SChaosClothAssetEditor3DViewport::MakeViewportToolbar()
+{
+	return SNew(SChaosClothAssetEditor3DViewportToolBar, SharedThis(this))
+		.CommandList(CommandList);
+}
+
+#undef LOCTEXT_NAMESPACE
diff --git a/Engine/Plugins/Experimental/ChaosClothAssetEditor/Source/ChaosClothAssetEditor/Private/ChaosClothAsset/SClothEditor3DViewportToolBar.cpp b/Engine/Plugins/Experimental/ChaosClothAssetEditor/Source/ChaosClothAssetEditor/Private/ChaosClothAsset/SClothEditor3DViewportToolBar.cpp
new file mode 100644
--- /dev/null
+++ b/Engine/Plugins/Experimental/ChaosClothAssetEditor/Source/ChaosClothAssetEditor/Private/ChaosClothAsset/SClothEditor3DViewportToolBar.cpp
@@ -0,0 +1,94 @@
+// Copyright Epic Games, Inc. All Rights Reserved.
+
+#include "ChaosClothAsset/SClothEditor3DViewportToolBar.h"
+#include "ChaosClothAsset/SClothEditor3DViewport.h"
+#include "ChaosClothAsset/ClothEditorCommands.h"
+#include "EditorViewportCommands.h"
+#include "Styling/AppStyle.h"
+#include "Framework/MultiBox/MultiBoxBuilder.h"
+#include "Widgets/Layout/SBorder.h"
+#include "Widgets/SBoxPanel.h"
+#include "Editor/UnrealEd/Public/SEditorViewportToolBarMenu.h"
+#include "SEditorViewportViewMenu.h"
+
+
+#define LOCTEXT_NAMESPACE "SChaosClothAssetEditor3DViewportToolBar"
+
+void SChaosClothAssetEditor3DViewportToolBar::Construct(const FArguments& InArgs, TSharedPtr<class SChaosClothAssetEditor3DViewport> InChaosClothAssetEditor3DViewport)
+{
+	ChaosClothAssetEditor3DViewportPtr = InChaosClothAssetEditor3DViewport;
+	CommandList = InArgs._CommandList;
+
+	const FMargin ToolbarSlotPadding(4.0f, 1.0f);
+	TSharedPtr<SHorizontalBox> MainBoxPtr;
+
+	ChildSlot
+	[
+		SNew( SBorder )
+		.BorderImage(FAppStyle::Get().GetBrush("EditorViewportToolBar.Background"))
+		.Cursor(EMouseCursor::Default)
+		[
+			SAssignNew( MainBoxPtr, SHorizontalBox )
+		]
+	];
+
+	MainBoxPtr->AddSlot()
+		.Padding(ToolbarSlotPadding)
+		.HAlign(HAlign_Left)
+		[
+			MakeDisplayToolBar(InArgs._Extenders)
+		];
+
+	MainBoxPtr->AddSlot()
+		.Padding(ToolbarSlotPadding)
+		.HAlign(HAlign_Right)
+		[
+			MakeToolBar(InArgs._Extenders)
+		];
+
+	SViewportToolBar::Construct(SViewportToolBar::FArguments());
+}
+
+TSharedRef<SWidget> SChaosClothAssetEditor3DViewportToolBar::MakeDisplayToolBar(const TSharedPtr<FExtender> InExtenders)
+{
+	TSharedRef<SEditorViewport> ViewportRef = StaticCastSharedPtr<SEditorViewport>(ChaosClothAssetEditor3DViewportPtr.Pin()).ToSharedRef();
+
+	return SNew(SEditorViewportViewMenu, ViewportRef, SharedThis(this))
+		.Cursor(EMouseCursor::Default)
+		.MenuExtenders(InExtenders);
+}
+
+TSharedRef<SWidget> SChaosClothAssetEditor3DViewportToolBar::MakeToolBar(const TSharedPtr<FExtender> InExtenders)
+{
+	// The following is modeled after portions of STransformViewportToolBar, which gets 
+	// used in SCommonEditorViewportToolbarBase.
+
+	// The buttons are hooked up to actual functions via command bindings in SChaosClothAssetEditor3DViewport::BindCommands(),
+	// and the toolbar gets built in SChaosClothAssetEditor3DViewport::MakeViewportToolbar().
+
+	FSlimHorizontalToolBarBuilder ToolbarBuilder(CommandList, FMultiBoxCustomization::None, InExtenders);
+
+	// Use a custom style
+	FName ToolBarStyle = "EditorViewportToolBar";
+	ToolbarBuilder.SetStyle(&FAppStyle::Get(), ToolBarStyle);
+	ToolbarBuilder.SetLabelVisibility(EVisibility::Collapsed);
+
+	ToolbarBuilder.BeginSection("SimControls");
+	{
+		ToolbarBuilder.BeginBlockGroup();
+	    
+		static FName ToggleSimMeshWireframeName = FName(TEXT("SimMeshWireframe"));
+		ToolbarBuilder.AddToolBarButton(FChaosClothAssetEditorCommands::Get().ToggleSimMeshWireframe);
+
+		static FName ToggleRenderMeshWireframeName = FName(TEXT("RenderMeshWireframe"));
+		ToolbarBuilder.AddToolBarButton(FChaosClothAssetEditorCommands::Get().ToggleRenderMeshWireframe);
+
+		ToolbarBuilder.EndBlockGroup();
+	}
+
+	ToolbarBuilder.EndSection();
+
+	return ToolbarBuilder.MakeWidget();
+}
+
+#undef LOCTEXT_NAMESPACE
