Chaos Cloth - Move the new Cloth Asset modules out of the Restricted area to the Engine/Plugins/Experimental folder (part 2, add files)

diff --git a/Engine/Plugins/Experimental/ChaosClothAssetEditor/Source/ChaosClothAssetEditor/Private/ChaosClothAsset/SClothEditorRestSpaceViewport.cpp b/Engine/Plugins/Experimental/ChaosClothAssetEditor/Source/ChaosClothAssetEditor/Private/ChaosClothAsset/SClothEditorRestSpaceViewport.cpp
new file mode 100644
--- /dev/null
+++ b/Engine/Plugins/Experimental/ChaosClothAssetEditor/Source/ChaosClothAssetEditor/Private/ChaosClothAsset/SClothEditorRestSpaceViewport.cpp
@@ -0,0 +1,66 @@
+// Copyright Epic Games, Inc. All Rights Reserved.
+
+#include "ChaosClothAsset/SClothEditorRestSpaceViewport.h"
+#include "SViewportToolBar.h"
+#include "ChaosClothAsset/SClothEditorRestSpaceViewportToolBar.h"
+#include "ChaosClothAsset/ClothEditorCommands.h"
+#include "ChaosClothAsset/ClothEditorMode.h"
+#include "EditorModeManager.h"
+
+// TODO evntually:
+//#include "ClothEditorRestSpaceViewportClient.h"
+
+#define LOCTEXT_NAMESPACE "SChaosClothAssetEditorRestSpaceViewport"
+
+void SChaosClothAssetEditorRestSpaceViewport::BindCommands()
+{
+	SAssetEditorViewport::BindCommands();
+
+	const FChaosClothAssetEditorCommands& CommandInfos = FChaosClothAssetEditorCommands::Get();
+
+	CommandList->MapAction(
+		CommandInfos.TogglePatternMode,
+		FExecuteAction::CreateLambda([this]()
+	{
+		const FEditorModeTools* const EditorModeTools = Client->GetModeTools();
+		UChaosClothAssetEditorMode* const ClothEdMode = Cast<UChaosClothAssetEditorMode>(EditorModeTools->GetActiveScriptableMode(UChaosClothAssetEditorMode::EM_ChaosClothAssetEditorModeId));
+
+		if (ClothEdMode)
+		{
+			ClothEdMode->TogglePatternMode();
+		}
+	}),
+	FCanExecuteAction::CreateLambda([this]() 
+	{ 
+		const FEditorModeTools* const EditorModeTools = Client->GetModeTools();
+		const UChaosClothAssetEditorMode* const ClothEdMode = Cast<UChaosClothAssetEditorMode>(EditorModeTools->GetActiveScriptableMode(UChaosClothAssetEditorMode::EM_ChaosClothAssetEditorModeId));
+
+		if (ClothEdMode)
+		{
+			return ClothEdMode->CanTogglePatternMode();
+		}
+
+		return false; 
+	}),
+	EUIActionRepeatMode::RepeatDisabled);
+}
+
+TSharedPtr<SWidget> SChaosClothAssetEditorRestSpaceViewport::MakeViewportToolbar()
+{
+	return SNew(SChaosClothAssetEditorRestSpaceViewportToolBar, SharedThis(this))
+		.CommandList(CommandList);
+}
+
+
+void SChaosClothAssetEditorRestSpaceViewport::OnFocusViewportToSelection()
+{
+	const FEditorModeTools* const EditorModeTools = Client->GetModeTools();
+	const UChaosClothAssetEditorMode* const ClothEdMode = Cast<UChaosClothAssetEditorMode>(EditorModeTools->GetActiveScriptableMode(UChaosClothAssetEditorMode::EM_ChaosClothAssetEditorModeId));
+
+	if (ClothEdMode)
+	{
+		Client->FocusViewportOnBox(ClothEdMode->SelectionBoundingBox());
+	}
+}
+
+#undef LOCTEXT_NAMESPACE
diff --git a/Engine/Plugins/Experimental/ChaosClothAssetEditor/Source/ChaosClothAssetEditor/Private/ChaosClothAsset/SClothEditorRestSpaceViewportToolBar.cpp b/Engine/Plugins/Experimental/ChaosClothAssetEditor/Source/ChaosClothAssetEditor/Private/ChaosClothAsset/SClothEditorRestSpaceViewportToolBar.cpp
new file mode 100644
--- /dev/null
+++ b/Engine/Plugins/Experimental/ChaosClothAssetEditor/Source/ChaosClothAssetEditor/Private/ChaosClothAsset/SClothEditorRestSpaceViewportToolBar.cpp
@@ -0,0 +1,90 @@
+// Copyright Epic Games, Inc. All Rights Reserved.
+
+#include "ChaosClothAsset/SClothEditorRestSpaceViewportToolBar.h"
+#include "ChaosClothAsset/SClothEditorRestSpaceViewport.h"
+#include "ChaosClothAsset/ClothEditorCommands.h"
+#include "EditorViewportCommands.h"
+#include "Styling/AppStyle.h"
+#include "Framework/MultiBox/MultiBoxBuilder.h"
+#include "Widgets/Layout/SBorder.h"
+#include "Widgets/SBoxPanel.h"
+#include "Editor/UnrealEd/Public/SEditorViewportToolBarMenu.h"
+#include "SEditorViewportViewMenu.h"
+
+#define LOCTEXT_NAMESPACE "SChaosClothAssetEditorRestSpaceViewportToolBar"
+
+void SChaosClothAssetEditorRestSpaceViewportToolBar::Construct(const FArguments& InArgs, TSharedPtr<SChaosClothAssetEditorRestSpaceViewport> InChaosClothAssetEditorViewport)
+{
+	ChaosClothAssetEditorRestSpaceViewportPtr = InChaosClothAssetEditorViewport;
+	CommandList = InArgs._CommandList;
+
+	const FMargin ToolbarSlotPadding(4.0f, 1.0f);
+	TSharedPtr<SHorizontalBox> MainBoxPtr;
+
+	ChildSlot
+	[
+		SNew( SBorder )
+		.BorderImage(FAppStyle::Get().GetBrush("EditorViewportToolBar.Background"))
+		.Cursor(EMouseCursor::Default)
+		[
+			SAssignNew( MainBoxPtr, SHorizontalBox )
+		]
+	];
+
+	MainBoxPtr->AddSlot()
+		.Padding(ToolbarSlotPadding)
+		.HAlign(HAlign_Left)
+		[
+			MakeDisplayToolBar(InArgs._Extenders)
+		];
+
+	MainBoxPtr->AddSlot()
+		.Padding(ToolbarSlotPadding)
+		.HAlign(HAlign_Right)
+		[
+			MakeToolBar(InArgs._Extenders)
+		];
+
+	SViewportToolBar::Construct(SViewportToolBar::FArguments());
+}
+
+TSharedRef<SWidget> SChaosClothAssetEditorRestSpaceViewportToolBar::MakeDisplayToolBar(const TSharedPtr<FExtender> InExtenders)
+{
+	const TSharedRef<SEditorViewport> ViewportRef = StaticCastSharedPtr<SEditorViewport>(ChaosClothAssetEditorRestSpaceViewportPtr.Pin()).ToSharedRef();
+
+	return SNew(SEditorViewportViewMenu, ViewportRef, SharedThis(this))
+		.Cursor(EMouseCursor::Default)
+		.MenuExtenders(InExtenders);
+}
+
+TSharedRef<SWidget> SChaosClothAssetEditorRestSpaceViewportToolBar::MakeToolBar(const TSharedPtr<FExtender> InExtenders)
+{
+	// The following is modeled after portions of STransformViewportToolBar, which gets 
+	// used in SCommonEditorViewportToolbarBase.
+
+	// The buttons are hooked up to actual functions via command bindings in SChaosClothAssetEditorRestSpaceViewport::BindCommands(),
+	// and the toolbar gets built in SChaosClothAssetEditorRestSpaceViewport::MakeViewportToolbar().
+
+	FSlimHorizontalToolBarBuilder ToolbarBuilder(CommandList, FMultiBoxCustomization::None, InExtenders);
+
+	// Use a custom style
+	const FName ToolBarStyle = "EditorViewportToolBar";
+	ToolbarBuilder.SetStyle(&FAppStyle::Get(), ToolBarStyle);
+	ToolbarBuilder.SetLabelVisibility(EVisibility::Collapsed);
+
+	ToolbarBuilder.BeginSection("SimControls");
+	{
+		ToolbarBuilder.BeginBlockGroup();
+	    
+		static FName TogglePatternModeName = FName(TEXT("PatternMode"));
+		ToolbarBuilder.AddToolBarButton(FChaosClothAssetEditorCommands::Get().TogglePatternMode);
+
+		ToolbarBuilder.EndBlockGroup();
+	}
+
+	ToolbarBuilder.EndSection();
+
+	return ToolbarBuilder.MakeWidget();
+}
+
+#undef LOCTEXT_NAMESPACE
diff --git a/Engine/Plugins/Experimental/ChaosClothAssetEditor/Source/ChaosClothAssetEditor/Public/ChaosClothAsset/SClothEditor3DViewport.h b/Engine/Plugins/Experimental/ChaosClothAssetEditor/Source/ChaosClothAssetEditor/Public/ChaosClothAsset/SClothEditor3DViewport.h
new file mode 100644
--- /dev/null
+++ b/Engine/Plugins/Experimental/ChaosClothAssetEditor/Source/ChaosClothAssetEditor/Public/ChaosClothAsset/SClothEditor3DViewport.h
@@ -0,0 +1,19 @@
+// Copyright Epic Games, Inc. All Rights Reserved.
+
+#pragma once
+
+#include "CoreMinimal.h"
+#include "SAssetEditorViewport.h"
+
+/**
+ * Viewport used for 3D preview in cloth editor. Has a custom toolbar overlay at the top.
+ */
+class CHAOSCLOTHASSETEDITOR_API SChaosClothAssetEditor3DViewport : public SAssetEditorViewport 
+{
+
+public:
+	// SAssetEditorViewport
+	virtual void BindCommands() override;
+	virtual TSharedPtr<SWidget> MakeViewportToolbar() override;
+
+};
diff --git a/Engine/Plugins/Experimental/ChaosClothAssetEditor/Source/ChaosClothAssetEditor/Public/ChaosClothAsset/SClothEditorRestSpaceViewport.h b/Engine/Plugins/Experimental/ChaosClothAssetEditor/Source/ChaosClothAssetEditor/Public/ChaosClothAsset/SClothEditorRestSpaceViewport.h
new file mode 100644
--- /dev/null
+++ b/Engine/Plugins/Experimental/ChaosClothAssetEditor/Source/ChaosClothAssetEditor/Public/ChaosClothAsset/SClothEditorRestSpaceViewport.h
@@ -0,0 +1,16 @@
+// Copyright Epic Games, Inc. All Rights Reserved.
+
+#pragma once
+
+#include "CoreMinimal.h"
+#include "SBaseCharacterFXEditorViewport.h"
+
+class CHAOSCLOTHASSETEDITOR_API SChaosClothAssetEditorRestSpaceViewport : public SBaseCharacterFXEditorViewport
+{
+public:
+
+	// SEditorViewport
+	virtual void BindCommands() override;
+	virtual TSharedPtr<SWidget> MakeViewportToolbar() override;
+	virtual void OnFocusViewportToSelection() override;
+};
diff --git a/Engine/Plugins/Experimental/ChaosClothAssetEditor/Source/ChaosClothAssetEditorTools/Private/ChaosClothAsset/ChaosClothAssetEditorToolsModule.cpp b/Engine/Plugins/Experimental/ChaosClothAssetEditor/Source/ChaosClothAssetEditorTools/Private/ChaosClothAsset/ChaosClothAssetEditorToolsModule.cpp
new file mode 100644
--- /dev/null
+++ b/Engine/Plugins/Experimental/ChaosClothAssetEditor/Source/ChaosClothAssetEditorTools/Private/ChaosClothAsset/ChaosClothAssetEditorToolsModule.cpp
@@ -0,0 +1,20 @@
+// Copyright Epic Games, Inc. All Rights Reserved.
+
+#include "ChaosClothAsset/ChaosClothAssetEditorToolsModule.h"
+
+#define LOCTEXT_NAMESPACE "FChaosClothAssetEditorToolsModule"
+
+void FChaosClothAssetEditorToolsModule::StartupModule()
+{
+	// This code will execute after your module is loaded into memory; the exact timing is specified in the .uplugin file per-module
+}
+
+void FChaosClothAssetEditorToolsModule::ShutdownModule()
+{
+	// This function may be called during shutdown to clean up your module.  For modules that support dynamic reloading,
+	// we call this function before unloading the module.
+}
+
+#undef LOCTEXT_NAMESPACE
+
+IMPLEMENT_MODULE(FChaosClothAssetEditorToolsModule, ChaosClothAssetEditorTools)
diff --git a/Engine/Plugins/Experimental/ChaosClothAssetEditor/Source/ChaosClothAssetTools/Private/ChaosClothAsset/ChaosClothAssetToolsModule.cpp b/Engine/Plugins/Experimental/ChaosClothAssetEditor/Source/ChaosClothAssetTools/Private/ChaosClothAsset/ChaosClothAssetToolsModule.cpp
new file mode 100644
--- /dev/null
+++ b/Engine/Plugins/Experimental/ChaosClothAssetEditor/Source/ChaosClothAssetTools/Private/ChaosClothAsset/ChaosClothAssetToolsModule.cpp
@@ -0,0 +1,55 @@
+// Copyright Epic Games, Inc. All Rights Reserved.
+
+#include "CoreMinimal.h"
+#include "Modules/ModuleInterface.h"
+#include "Modules/ModuleManager.h"
+#include "AssetToolsModule.h"
+#include "ClothAssetEditorStyle.h"
+#include "ClothComponentEditorStyle.h"
+#include "ClothPresetEditorStyle.h"
+#include "ClothingAssetToClothAssetExporter.h"
+#include "ChaosClothAsset/ClothAssetBuilderEditor.h"
+
+namespace UE::Chaos::ClothAsset
+{
+	class FClothAssetEditorModule : public IModuleInterface, public IClothingAssetExporterClassProvider, public IClothAssetBuilderClassProvider
+	{
+	public:
+		// IModuleInterface implementation
+		virtual void StartupModule() override
+		{
+			// Register asset icons
+			FClothAssetEditorStyle::Get();
+			FClothComponentEditorStyle::Get();
+			FClothPresetEditorStyle::Get();
+
+			// Register modular features
+			IModularFeatures::Get().RegisterModularFeature(IClothingAssetExporterClassProvider::FeatureName, static_cast<IClothingAssetExporterClassProvider*>(this));
+			IModularFeatures::Get().RegisterModularFeature(IClothAssetBuilderClassProvider::FeatureName, static_cast<IClothAssetBuilderClassProvider*>(this));
+		}
+
+		virtual void ShutdownModule() override
+		{
+			if (UObjectInitialized())
+			{
+				// Unregister modular features
+				IModularFeatures::Get().UnregisterModularFeature(IClothingAssetExporterClassProvider::FeatureName, static_cast<IClothingAssetExporterClassProvider*>(this));
+				IModularFeatures::Get().UnregisterModularFeature(IClothAssetBuilderClassProvider::FeatureName, static_cast<IClothAssetBuilderClassProvider*>(this));
+			}
+		}
+
+		// IClothingSimulationFactoryClassProvider implementation
+		virtual TSubclassOf<UClothingAssetExporter> GetClothingAssetExporterClass() const override
+		{
+			return UClothingAssetToChaosClothAssetExporter::StaticClass();
+		}
+
+		// IClothAssetBuilderClassProvider implementation
+		virtual TSubclassOf<UClothAssetBuilder> GetClothAssetBuilderClass() const override
+		{
+			return UClothAssetBuilderEditor::StaticClass();
+		}
+	};
+}
+
+IMPLEMENT_MODULE(UE::Chaos::ClothAsset::FClothAssetEditorModule, ChaosClothAssetEditor);
diff --git a/Engine/Plugins/Experimental/ChaosClothAssetEditor/Source/ChaosClothAssetTools/Private/ChaosClothAsset/ClothAssetBuilderEditor.cpp b/Engine/Plugins/Experimental/ChaosClothAssetEditor/Source/ChaosClothAssetTools/Private/ChaosClothAsset/ClothAssetBuilderEditor.cpp
new file mode 100644
--- /dev/null
+++ b/Engine/Plugins/Experimental/ChaosClothAssetEditor/Source/ChaosClothAssetTools/Private/ChaosClothAsset/ClothAssetBuilderEditor.cpp
@@ -0,0 +1,286 @@
+// Copyright Epic Games, Inc. All Rights Reserved.
+
+#include "ChaosClothAsset/ClothAssetBuilderEditor.h"
+
+#include "BoneWeights.h"
+#include "ChaosClothAsset/ClothCollection.h"
+#include "Modules/ModuleManager.h"
+#include "Rendering/SkeletalMeshLODModel.h"
+#include "Utils/ClothingMeshUtils.h"
+#include "MeshUtilities.h"
+#include "PointWeightMap.h"
+
+#include UE_INLINE_GENERATED_CPP_BY_NAME(ClothAssetBuilderEditor)
+
+void UClothAssetBuilderEditor::BuildLod(FSkeletalMeshLODModel& LODModel, const UChaosClothAsset& ClothAsset, int32 LodIndex) const
+{
+	// Start from an empty LODModel
+	LODModel.Empty();
+
+	// Clear the mesh infos, none are stored on this asset
+	LODModel.ImportedMeshInfos.Empty();
+	LODModel.MaxImportVertex = 0;
+
+	// Set 1 texture coordinate
+	LODModel.NumTexCoords = 1;
+
+	// Init the size of the vertex buffer
+	LODModel.NumVertices = 0;
+
+	// Build the section/faces map from the LOD patterns
+	TMap<int32, TArray<int32>> SectionFacesMap;
+	SectionFacesMap.Reserve(ClothAsset.GetMaterials().Num());
+
+	const TSharedPtr<UE::Chaos::ClothAsset::FClothCollection> ClothCollection = ClothAsset.GetClothCollection();
+	const int32 PatternStart = ClothCollection->PatternStart[LodIndex];
+	const int32 PatternEnd = ClothCollection->PatternEnd[LodIndex];
+
+	for (int32 PatternIndex = PatternStart; PatternIndex <= PatternEnd; ++PatternIndex)
+	{
+		const int32 RenderFacesStart = ClothCollection->RenderFacesStart[PatternIndex];
+		const int32 RenderFacesEnd = ClothCollection->RenderFacesEnd[PatternIndex];
+
+		for (int32 RenderFaceIndex = RenderFacesStart; RenderFaceIndex <= RenderFacesEnd; ++RenderFaceIndex)
+		{
+			const int32 MaterialIndex = ClothCollection->RenderMaterialIndex[RenderFaceIndex];
+
+			TArray<int32>& SectionFaces = SectionFacesMap.FindOrAdd(MaterialIndex);
+			SectionFaces.Add(RenderFaceIndex);
+		}
+	}
+
+	// Initialize the remapping array to start at LodRenderVerticesStart with the maximum number of vertices used in this LOD
+	const int32 LodRenderVerticesStart = ClothCollection->RenderVerticesStart[PatternStart];
+	const int32 LodRenderVerticesEnd = ClothCollection->RenderVerticesEnd[PatternEnd];
+	const int32 NumLodRenderVertices = LodRenderVerticesEnd - LodRenderVerticesStart + 1;
+
+	TArray<uint32> LodRenderIndexRemap;
+	LodRenderIndexRemap.SetNumUninitialized(NumLodRenderVertices);
+
+	// Keep track of the active bone indices for this LOD model
+	TSet<FBoneIndexType> ActiveBoneIndices;
+	ActiveBoneIndices.Reserve(ClothAsset.RefSkeleton.GetNum());
+
+	// Load the mesh utilities module used to optimized the index buffer
+	IMeshUtilities& MeshUtilities = FModuleManager::Get().LoadModuleChecked<IMeshUtilities>("MeshUtilities");
+
+	// Build the sim mesh descriptor for creation of the sections' mesh to mesh mapping data
+	// TODO: Generate Mesh to Mesh data from pattern space
+	const int32 LodSimVerticesStart = ClothCollection->SimVerticesStart[PatternStart];
+	const int32 LodSimVerticesEnd = ClothCollection->SimVerticesEnd[PatternEnd];
+	const int32 NumLodSimVertices = LodSimVerticesEnd - LodSimVerticesStart + 1;
+
+	const int32 LodSimFacesStart = ClothCollection->SimFacesStart[PatternStart];
+	const int32 LodSimFacesEnd = ClothCollection->SimFacesEnd[PatternEnd];
+
+	const int32 NumLodSimFaces = LodSimFacesEnd - LodSimFacesStart + 1;
+	const int32 NumLodIndices = NumLodSimFaces * 3;
+
+	TArray<uint32> LodSimIndices;
+	LodSimIndices.Reserve(NumLodIndices);
+
+	for (int32 SimFaceIndex = LodSimFacesStart; SimFaceIndex <= LodSimFacesEnd; ++SimFaceIndex)
+	{
+		const FIntVector3& SimIndices = ClothCollection->SimIndices[SimFaceIndex];
+		LodSimIndices.Add((uint32)SimIndices.X);
+		LodSimIndices.Add((uint32)SimIndices.Y);
+		LodSimIndices.Add((uint32)SimIndices.Z);
+
+		const FVector3f& A = ClothCollection->SimRestPosition[LodSimVerticesStart + SimIndices.X];
+		const FVector3f& B = ClothCollection->SimRestPosition[LodSimVerticesStart + SimIndices.Y];
+		const FVector3f& C = ClothCollection->SimRestPosition[LodSimVerticesStart + SimIndices.Z];
+		ensure(FVector3f::CrossProduct(B - A, C - A).SizeSquared() >= SMALL_NUMBER);
+	}
+
+	const ClothingMeshUtils::ClothMeshDesc SourceMesh(
+		TConstArrayView<FVector3f>(ClothCollection->SimRestPosition.GetData() + LodSimVerticesStart, NumLodSimVertices),
+		TConstArrayView<uint32>(LodSimIndices));  // Calculate averaged normals
+
+	// Retrieve the MaxDistance map
+	FPointWeightMap MaxDistances;
+	MaxDistances.Initialize(NumLodSimVertices);
+	for (int32 Index = 0; Index < NumLodSimVertices; ++Index)
+	{
+		MaxDistances[Index] = 200.f;
+	}
+
+	// Populate this LOD's sections and the LOD index buffer
+	const int32 NumSections = SectionFacesMap.Num();
+	LODModel.Sections.SetNum(NumSections);
+
+	int32 SectionIndex = 0;
+	int32 BaseIndex = 0;
+	for (const TPair<int32, TArray<int32>>& SectionFaces : SectionFacesMap)
+	{
+		FSkelMeshSection& Section = LODModel.Sections[SectionIndex++];
+
+		const int32 MaterialIndex = SectionFaces.Key;
+		const TArray<int32>& Faces = SectionFaces.Value;
+		const int32 NumFaces = Faces.Num();
+		const int32 NumIndices = NumFaces * 3;
+
+		// Build the section face data (indices)
+		Section.MaterialIndex = (uint16)MaterialIndex;
+
+		Section.BaseIndex = (uint32)BaseIndex;
+		BaseIndex += NumIndices;
+
+		Section.NumTriangles = (uint32)NumFaces;
+
+		TArray<uint32> Indices;
+		TSet<int32> UniqueIndicesSet;
+		Indices.SetNumUninitialized(NumIndices);
+		for (int32 FaceIndex = 0; FaceIndex < NumFaces; ++FaceIndex)
+		{
+			const FIntVector3& RenderIndices = ClothCollection->RenderIndices[Faces[FaceIndex]];
+			for (int32 VertexIndex = 0; VertexIndex < 3; ++VertexIndex)
+			{
+				const int32 RenderIndex = RenderIndices[VertexIndex];
+				Indices[FaceIndex * 3 + VertexIndex] = (uint32)RenderIndex;
+				UniqueIndicesSet.Add(RenderIndex);
+			}
+		}
+
+		MeshUtilities.CacheOptimizeIndexBuffer(Indices);
+
+		LODModel.IndexBuffer.Append(MoveTemp(Indices));
+
+		// Build the section vertex data from the unique indices
+		TArray<int32> UniqueIndices = UniqueIndicesSet.Array();
+		const int32 NumVertices = UniqueIndices.Num();
+
+		Section.SoftVertices.SetNumUninitialized(NumVertices);
+		Section.NumVertices = NumVertices;
+		Section.BaseVertexIndex = LODModel.NumVertices;
+		LODModel.NumVertices += (uint32)NumVertices;
+
+		TSet<FBoneIndexType> SectionBoneMap;
+		SectionBoneMap.Reserve(ClothAsset.RefSkeleton.GetNum());
+
+		for (int32 VertexIndex = 0; VertexIndex < NumVertices; ++VertexIndex)
+		{
+			const int32 RenderIndex = UniqueIndices[VertexIndex];
+			LODModel.MaxImportVertex = FMath::Max(LODModel.MaxImportVertex, RenderIndex);
+
+			LodRenderIndexRemap[RenderIndex] = Section.BaseVertexIndex + (uint32)VertexIndex;
+
+			FSoftSkinVertex& SoftVertex = Section.SoftVertices[VertexIndex];
+
+			SoftVertex.Position = ClothCollection->RenderPosition[RenderIndex];
+			SoftVertex.TangentX = ClothCollection->RenderTangentU[RenderIndex];
+			SoftVertex.TangentY = ClothCollection->RenderTangentV[RenderIndex];
+			SoftVertex.TangentZ = ClothCollection->RenderNormal[RenderIndex];
+
+			constexpr bool bSRGB = false; // Avoid linear to srgb conversion
+			SoftVertex.Color = ClothCollection->RenderColor[RenderIndex].ToFColor(bSRGB);
+
+			const TArray<FVector2f>& RenderUVs = ClothCollection->RenderUVs[RenderIndex];
+			for (int32 TexCoord = 0; TexCoord < FMath::Min(RenderUVs.Num(), (int32)MAX_TEXCOORDS); ++TexCoord)
+			{
+				SoftVertex.UVs[TexCoord] = RenderUVs[TexCoord];
+			}
+
+			for (int32 Influence = 0; Influence < MAX_TOTAL_INFLUENCES; ++Influence)
+			{
+				const FBoneIndexType InfluenceBone = (FBoneIndexType)0;  // TODO: Set the correct bone influence from the cloth collection data
+				const int16 InfluenceWeight = (Influence == 0) ? UE::AnimationCore::MaxRawBoneWeight : 0;
+
+				SoftVertex.InfluenceBones[Influence] = InfluenceBone;
+				SoftVertex.InfluenceWeights[Influence] = InfluenceWeight;
+
+				if (InfluenceWeight)
+				{
+					SectionBoneMap.Add(InfluenceBone);
+				}
+			}
+		}
+
+		// Update the section bone map
+		Section.BoneMap = SectionBoneMap.Array();
+		ActiveBoneIndices.Append(MoveTemp(SectionBoneMap));
+
+		// Update max bone influences
+		Section.CalcMaxBoneInfluences();
+		Section.CalcUse16BitBoneIndex();
+
+		// Setup clothing data
+		Section.ClothMappingDataLODs.SetNum(1);  // TODO: LODBias maps for raytracing
+
+		Section.ClothingData.AssetLodIndex = LodIndex;
+		Section.ClothingData.AssetGuid = ClothAsset.AssetGuid;  // There is only one cloth asset,
+		Section.CorrespondClothAssetIndex = 0;       // this one
+
+		const ClothingMeshUtils::ClothMeshDesc TargetMesh(
+			TConstArrayView<FVector3f>(ClothCollection->RenderPosition.GetData(), NumVertices),
+			TConstArrayView<FVector3f>(ClothCollection->RenderNormal.GetData(), NumVertices),
+			TConstArrayView<FVector3f>(ClothCollection->RenderTangentU.GetData(), NumVertices),
+			TConstArrayView<uint32>(Indices));
+
+		ClothingMeshUtils::GenerateMeshToMeshVertData(
+			Section.ClothMappingDataLODs[0],
+			TargetMesh,
+			SourceMesh,
+			&MaxDistances,
+			ClothAsset.bSmoothTransition,
+			ClothAsset.bUseMultipleInfluences,
+			ClothAsset.SkinningKernelRadius);
+
+		// Save the original indices for the newly added vertices
+		LODModel.MeshToImportVertexMap.Append(MoveTemp(UniqueIndices));
+
+		// Compute the overlapping vertices map (inspired from MeshUtilities::BuildSkeletalMesh)
+		const TArray<FSoftSkinVertex>& SoftVertices = Section.SoftVertices;
+
+		typedef TPair<float, int32> FIndexAndZ;  // Acceleration structure, list of vertex Z / index pairs
+		TArray<FIndexAndZ> IndexAndZs;
+		IndexAndZs.Reserve(NumVertices);
+		for (int32 VertexIndex = 0; VertexIndex < NumVertices; ++VertexIndex)
+		{
+			const FVector3f& Position = SoftVertices[VertexIndex].Position;
+
+			const float Z = 0.30f * Position.X + 0.33f * Position.Y + 0.37f * Position.Z;
+			IndexAndZs.Emplace(Z, VertexIndex);
+		}
+		IndexAndZs.Sort([](const FIndexAndZ& A, const FIndexAndZ& B) { return A.Key < B.Key; });
+
+		for (int32 Index0 = 0; Index0 < IndexAndZs.Num(); ++Index0)
+		{
+			const float Z0 = IndexAndZs[Index0].Key;
+			const uint32 VertexIndex0 = IndexAndZs[Index0].Value;
+			const FVector3f& Position0 = SoftVertices[VertexIndex0].Position;
+
+			// Only need to search forward, since we add pairs both ways
+			for (int32 Index1 = Index0 + 1; Index1 < IndexAndZs.Num() && FMath::Abs(IndexAndZs[Index1].Key - Z0) <= THRESH_POINTS_ARE_SAME; ++Index1)
+			{
+				const uint32 VertexIndex1 = IndexAndZs[Index1].Value;
+				const FVector3f& Position1 = SoftVertices[VertexIndex1].Position;
+
+				if (PointsEqual(Position0, Position1))
+				{
+					// Add to the overlapping map
+					TArray<int32>& SrcValueArray = Section.OverlappingVertices.FindOrAdd(VertexIndex0);
+					SrcValueArray.Add(VertexIndex1);
+
+					TArray<int32>& IterValueArray = Section.OverlappingVertices.FindOrAdd(VertexIndex1);
+					IterValueArray.Add(VertexIndex0);
+				}
+			}
+		}
+	}
+
+	// Remap the LOD indices with the new vertex indices
+	for (uint32& RenderIndex : LODModel.IndexBuffer)
+	{
+		RenderIndex = LodRenderIndexRemap[RenderIndex];
+	}
+
+	// Update the active bone indices on the LOD model
+	LODModel.ActiveBoneIndices = ActiveBoneIndices.Array();
+
+	// Ensure parent exists with incoming active bone indices, and the result should be sorted
+	ClothAsset.RefSkeleton.EnsureParentsExistAndSort(LODModel.ActiveBoneIndices);
+
+	// Compute the required bones for this model.
+	USkeletalMesh::CalculateRequiredBones(LODModel, ClothAsset.RefSkeleton, nullptr);
+}
+
diff --git a/Engine/Plugins/Experimental/ChaosClothAssetEditor/Source/ChaosClothAssetTools/Private/ChaosClothAsset/ClothAssetBuilderEditor.h b/Engine/Plugins/Experimental/ChaosClothAssetEditor/Source/ChaosClothAssetTools/Private/ChaosClothAsset/ClothAssetBuilderEditor.h
new file mode 100644
--- /dev/null
+++ b/Engine/Plugins/Experimental/ChaosClothAssetEditor/Source/ChaosClothAssetTools/Private/ChaosClothAsset/ClothAssetBuilderEditor.h
@@ -0,0 +1,22 @@
+// Copyright Epic Games, Inc. All Rights Reserved.
+
+#pragma once
+
+#include "ChaosClothAsset/ClothAssetBuilder.h"
+
+#include "ClothAssetBuilderEditor.generated.h"
+
+class FSkeletalMeshLODModel;
+
+/**
+ * Implement a clothing asset builder.
+ */
+UCLASS()
+class UClothAssetBuilderEditor : public UClothAssetBuilder
+{
+	GENERATED_BODY()
+
+public:
+	/** Build a FSkeletalMeshLODModel out of the cloth asset for the specified LOD index. */
+	virtual void BuildLod(FSkeletalMeshLODModel& LODModel, const UChaosClothAsset& ClothAsset, int32 LodIndex) const override;
+};
diff --git a/Engine/Plugins/Experimental/ChaosClothAssetEditor/Source/ChaosClothAssetTools/Private/ChaosClothAsset/ClothAssetFactory.cpp b/Engine/Plugins/Experimental/ChaosClothAssetEditor/Source/ChaosClothAssetTools/Private/ChaosClothAsset/ClothAssetFactory.cpp
new file mode 100644
--- /dev/null
+++ b/Engine/Plugins/Experimental/ChaosClothAssetEditor/Source/ChaosClothAssetTools/Private/ChaosClothAsset/ClothAssetFactory.cpp
@@ -0,0 +1,22 @@
+// Copyright Epic Games, Inc. All Rights Reserved.
+
+#include "ClothAssetFactory.h"
+#include "ChaosClothAsset/ClothAsset.h"
+
+#include UE_INLINE_GENERATED_CPP_BY_NAME(ClothAssetFactory)
+
+UChaosClothAssetFactory::UChaosClothAssetFactory(const FObjectInitializer& ObjectInitializer)
+	: Super(ObjectInitializer)
+{
+	bEditorImport = true;
+	bEditAfterNew = true;
+	SupportedClass = UChaosClothAsset::StaticClass();
+}
+
+UObject* UChaosClothAssetFactory::FactoryCreateNew(UClass* Class, UObject* Parent, FName Name, EObjectFlags Flags, UObject* /*Context*/, FFeedbackContext* /*Warn*/)
+{
+	UChaosClothAsset* const NewClothAsset = NewObject<UChaosClothAsset>(Parent, Class, Name, Flags | RF_Transactional | RF_Public | RF_Standalone);
+	NewClothAsset->MarkPackageDirty();
+	return NewClothAsset;
+}
+
diff --git a/Engine/Plugins/Experimental/ChaosClothAssetEditor/Source/ChaosClothAssetTools/Private/ChaosClothAsset/ClothPatternToDynamicMesh.cpp b/Engine/Plugins/Experimental/ChaosClothAssetEditor/Source/ChaosClothAssetTools/Private/ChaosClothAsset/ClothPatternToDynamicMesh.cpp
new file mode 100644
--- /dev/null
+++ b/Engine/Plugins/Experimental/ChaosClothAssetEditor/Source/ChaosClothAssetTools/Private/ChaosClothAsset/ClothPatternToDynamicMesh.cpp
@@ -0,0 +1,285 @@
+// Copyright Epic Games, Inc. All Rights Reserved. 
+
+#if WITH_EDITOR
+
+#include "ChaosClothAsset/ClothPatternToDynamicMesh.h"
+#include "ChaosClothAsset/ClothAsset.h"
+#include "ChaosClothAsset/ClothComponent.h"
+#include "ChaosClothAsset/ClothCollection.h"
+#include "ChaosClothAsset/ClothAdapter.h"
+#include "ChaosClothAsset/ClothLodAdapter.h"
+#include "ChaosClothAsset/ClothPatternAdapter.h"
+#include "ToDynamicMesh.h"
+#include "DynamicMesh/MeshNormals.h"
+
+//
+// Wrapper for accessing a Cloth Pattern. Implements the interface expected by TToDynamicMesh<>.
+//
+class FClothSimPatternWrapper
+{
+public:
+
+	typedef int32 TriIDType;
+	typedef int32 VertIDType;
+	typedef int32 WedgeIDType;
+
+	typedef int32 UVIDType;
+	typedef int32 NormalIDType;
+	typedef int32 ColorIDType;
+
+	FClothSimPatternWrapper(const UE::Chaos::ClothAsset::FClothPatternConstAdapter& PatternIn, bool bGet2DPatternVertices) :
+		bGet2DPatternVertices(bGet2DPatternVertices),
+		Pattern(PatternIn),
+		Lod(PatternIn.GetClothCollection(), PatternIn.GetLodIndex())
+	{
+		using UE::Chaos::ClothAsset::FClothCollection;
+
+		const int32 PatternIndex = Pattern.GetPatternIndex();
+
+		VertIDs.SetNum(Pattern.GetNumSimVertices());
+		for (int32 VtxIndex = 0; VtxIndex < Pattern.GetNumSimVertices(); ++VtxIndex)
+		{
+			const int32 PatternOffset = Lod.GetSimVerticesStart()[PatternIndex];
+			VertIDs[VtxIndex] = PatternOffset + VtxIndex;
+		}
+
+		TriIDs.SetNum(Pattern.GetNumSimFaces());
+		for (int32 TriIndex = 0; TriIndex < Pattern.GetNumSimFaces(); ++TriIndex)
+		{
+			TriIDs[TriIndex] = TriIndex;
+		}
+
+
+		//
+		// Weight map layers precomputation
+		//
+		const FClothCollection* ClothCollection = Pattern.GetClothCollection().Get();
+		const TArray<FName> SimVerticesAttributeNames = ClothCollection->AttributeNames(FClothCollection::SimVerticesGroup);
+
+		// Of the SimVerticesGroup attributes, which ones refer to scalar weight maps?
+		for (const FName& AttributeName : SimVerticesAttributeNames)
+		{
+			if (ClothCollection->FindAttributeTyped<float>(AttributeName, FClothCollection::SimVerticesGroup))
+			{
+				WeightMapNames.Add(AttributeName);
+			}
+		}
+	}
+
+	int32 NumTris() const
+	{
+		return TriIDs.Num();
+	}
+
+	int32 NumVerts() const
+	{
+		return VertIDs.Num();
+	}
+
+	int32 NumUVLayers() const
+	{
+		return 0;	// No UVs for Sim mesh
+	}
+
+	int32 NumWeightMapLayers() const
+	{
+		return WeightMapNames.Num();
+	}
+
+	FName GetWeightMapName(int32 LayerIndex) const
+	{
+		return WeightMapNames[LayerIndex];
+	}
+
+	float GetVertexWeight(int32 LayerIndex, int32 VertexIndex) const
+	{
+		using UE::Chaos::ClothAsset::FClothCollection;
+
+		const FName& AttributeName = WeightMapNames[LayerIndex];
+		const FClothCollection* ClothCollection = Pattern.GetClothCollection().Get();
+		ensure(ClothCollection->FindAttributeTyped<float>(AttributeName, FClothCollection::SimVerticesGroup));
+		
+		const TManagedArray<float>& ManagedArray = ClothCollection->GetAttribute<float>(AttributeName, FClothCollection::SimVerticesGroup);
+		return ManagedArray[VertexIndex];
+	}
+
+	// --"Vertex Buffer" info
+	const TArray<VertIDType>& GetVertIDs() const
+	{
+		return VertIDs;
+	}
+
+	FVector3d GetPosition(VertIDType VtxID) const
+	{
+		if (bGet2DPatternVertices)
+		{
+			TConstArrayView<FVector2f> SimPositions = Lod.GetPatternsSimPosition();
+			const FVector2f& Pos = SimPositions[VtxID];
+			return FVector3d(Pos[0], Pos[1], 0.0);
+		} 
+		else
+		{
+			TConstArrayView<FVector3f> RestPositions = Lod.GetPatternsSimRestPosition();
+			const FVector3f& Pos = RestPositions[VtxID];
+			return FVector3d(Pos[0], Pos[1], Pos[2]);
+		}
+	}
+
+	// --"Index Buffer" info
+	const TArray<TriIDType>& GetTriIDs() const
+	{
+		return TriIDs;
+	}
+
+	bool GetTri(TriIDType TriID, VertIDType& VID0, VertIDType& VID1, VertIDType& VID2) const
+	{
+		const FIntVector Face = Pattern.GetSimIndices()[TriID];
+		VID0 = Face[0];
+		VID1 = Face[1];
+		VID2 = Face[2];
+
+		return true;
+	}
+	
+	bool HasNormals() const
+	{
+		return true;
+	}
+
+	bool HasTangents() const
+	{
+		return false;
+	}
+
+	bool HasBiTangents() const
+	{
+		return false;
+	}
+
+	bool HasColors() const
+	{
+		return false;
+	}
+
+	// -- Access to per-wedge attributes -- //
+	void GetWedgeIDs(const TriIDType& TriID, WedgeIDType& WID0, WedgeIDType& WID1, WedgeIDType& WID2) const
+	{
+		// TODO: Review this
+		int32 Offset = 3 * TriID;
+		WID0 = Offset;
+		WID1 = Offset + 1;
+		WID2 = Offset + 2;
+	}
+
+	FVector2f GetWedgeUV(int32 UVLayerIndex, WedgeIDType WID) const
+	{
+		ensure(false);
+		return FVector2f(0.0f, 0.0f); 
+	}
+
+	FVector3f GetWedgeNormal(WedgeIDType WID) const
+	{
+		ensure(false);
+		return FVector3f(0.0f, 0.0f, 0.0f);
+	}
+
+	FVector3f GetWedgeTangent(WedgeIDType WID) const
+	{
+		ensure(false);
+		return FVector3f(0.0f, 0.0f, 0.0f);
+	}
+
+	FVector3f GetWedgeBiTangent(WedgeIDType WID) const
+	{
+		ensure(false);
+		return FVector3f(0.0f, 0.0f, 0.0f);
+	}
+	
+	FVector4f GetWedgeColor(WedgeIDType WID) const
+	{
+		ensure(false);
+		return FVector4f(0.0f, 0.0f, 0.0f, 0.0f);
+	}
+
+	int32 GetMaterialIndex(TriIDType TriID) const
+	{
+		return 0;
+	}
+
+	const TArray<int32>& GetNormalIDs() const { return VertIDs; }
+	FVector3f GetNormal(NormalIDType ID) const { return Lod.GetPatternsSimRestNormal()[ID]; }
+	bool GetNormalTri(const TriIDType& TriID, NormalIDType& ID0, NormalIDType& ID1, NormalIDType& ID2) const { return GetTri(TriID, ID0, ID1, ID2); }
+	
+	const TArray<int32>& GetUVIDs(int32 LayerID) const { return EmptyArray; }
+	FVector2f GetUV(int32 LayerID, UVIDType UVID) const { check(0); return FVector2f(); }
+	bool GetUVTri(int32 LayerID, const TriIDType&, UVIDType& ID0, UVIDType& ID1, UVIDType& ID2) const { ID0 = ID1 = ID2 = UVIDType(-1); return false;}
+	
+	const TArray<int32>& GetTangentIDs() const { return EmptyArray; }
+	FVector3f GetTangent(NormalIDType ID) const { check(0); return FVector3f(); }
+	bool GetTangentTri(const TriIDType&, NormalIDType& ID0, NormalIDType& ID1, NormalIDType& ID2) const { ID0 = ID1 = ID2 = NormalIDType(-1); return false; }
+
+	const TArray<int32>& GetBiTangentIDs() const { return EmptyArray; }
+	FVector3f GetBiTangent(NormalIDType ID) const {check(0); return FVector3f(); }
+	bool GetBiTangentTri(const TriIDType&, NormalIDType& ID0, NormalIDType& ID1, NormalIDType& ID2) const { ID0 = ID1 = ID2 = NormalIDType(-1); return false; }
+
+	const TArray<int32>& GetColorIDs() const { return EmptyArray; }
+	FVector4f GetColor(ColorIDType ID) const { check(0); return FVector4f(); }
+	bool GetColorTri(const TriIDType&, ColorIDType& ID0, ColorIDType& ID1, ColorIDType& ID2) const { ID0 = ID1 = ID2 = ColorIDType(-1); return false; }
+
+private:
+
+	bool bGet2DPatternVertices = false;
+
+	TArray<TriIDType> TriIDs;		// indices into Pattern.GetSimIndices() 
+	TArray<VertIDType> VertIDs;		// indices into Lod.GetPatternsSimPosition(), etc.
+	
+	TArray<FName> WeightMapNames;
+
+	const UE::Chaos::ClothAsset::FClothPatternConstAdapter Pattern;
+	const UE::Chaos::ClothAsset::FClothLodConstAdapter Lod;
+
+	TArray<int32> EmptyArray;
+};
+
+
+void FClothPatternToDynamicMesh::Convert(const UChaosClothAsset* ClothAssetMeshIn, int32 LODIndex, int32 PatternIndex, bool bGet2DPattern, UE::Geometry::FDynamicMesh3& MeshOut)
+{
+	const TSharedPtr<UE::Chaos::ClothAsset::FClothCollection> ClothCollection = ClothAssetMeshIn->GetClothCollection();
+	check(ClothCollection.IsValid());
+	UE::Chaos::ClothAsset::FClothAdapter ClothAdapter(ClothCollection);
+
+	const UE::Chaos::ClothAsset::FClothLodConstAdapter ClothLodAdapter = ClothAdapter.GetLod(LODIndex);
+	const UE::Chaos::ClothAsset::FClothPatternConstAdapter ClothPatternAdapter = ClothLodAdapter.GetPattern(PatternIndex);
+
+	// Actual conversion
+	UE::Geometry::TToDynamicMesh<FClothSimPatternWrapper> PatternToDynamicMesh; 
+	FClothSimPatternWrapper PatternWrapper(ClothPatternAdapter, bGet2DPattern);
+
+	const bool bDisableAttributes = false;
+	auto TriangleToGroupFunction = [](FClothSimPatternWrapper::TriIDType) { return 0; };
+
+	if (bDisableAttributes)
+	{
+		MeshOut.DiscardAttributes();
+		PatternToDynamicMesh.ConvertWOAttributes(MeshOut, PatternWrapper, TriangleToGroupFunction);
+	}
+	else
+	{
+		MeshOut.EnableAttributes();
+
+		auto TriangleToMaterialFunction = [](FClothSimPatternWrapper::TriIDType) { return 0; };
+		constexpr bool bCopyTangents = false;
+		PatternToDynamicMesh.Convert(MeshOut, PatternWrapper, TriangleToGroupFunction, TriangleToMaterialFunction, bCopyTangents);
+	}
+}
+
+
+#else
+
+void  FClothAssetToDynamicMesh::Convert(const UChaosClothAsset* ClothAssetMeshIn, FDynamicMesh3& MeshOut, int32 LODIndex, int32 PatternIndex)
+{
+	// Conversion only supported with editor.
+	check(0);
+}
+
+#endif  // end with editor
diff --git a/Engine/Plugins/Experimental/ChaosClothAssetEditor/Source/ChaosClothAssetTools/Private/ChaosClothAsset/ClothPresetEditorStyle.cpp b/Engine/Plugins/Experimental/ChaosClothAssetEditor/Source/ChaosClothAssetTools/Private/ChaosClothAsset/ClothPresetEditorStyle.cpp
new file mode 100644
--- /dev/null
+++ b/Engine/Plugins/Experimental/ChaosClothAssetEditor/Source/ChaosClothAssetTools/Private/ChaosClothAsset/ClothPresetEditorStyle.cpp
@@ -0,0 +1,26 @@
+// Copyright Epic Games, Inc. All Rights Reserved.
+
+#include "ClothPresetEditorStyle.h"
+#include "Styling/SlateStyleRegistry.h"
+#include "Interfaces/IPluginManager.h"
+
+namespace UE::Chaos::ClothAsset
+{
+	TOptional<FClothPresetEditorStyle> FClothPresetEditorStyle::Singleton;
+
+	FClothPresetEditorStyle::FClothPresetEditorStyle()
+		: FSlateStyleSet("ClothPresetEditorStyle")
+	{
+		SetContentRoot(IPluginManager::Get().FindPlugin("ChaosClothAssetEditor")->GetBaseDir() / TEXT("Resources"));
+
+		Set("ClassIcon.ChaosClothPreset", new FSlateVectorImageBrush(RootToContentDir(TEXT("ClothPreset_16.svg")), FVector2D(16)));
+		Set("ClassThumbnail.ChaosClothPreset", new FSlateVectorImageBrush(RootToContentDir(TEXT("ClothPreset_64.svg")), FVector2D(64)));
+
+		FSlateStyleRegistry::RegisterSlateStyle(*this);
+	}
+
+	FClothPresetEditorStyle::~FClothPresetEditorStyle()
+	{
+		FSlateStyleRegistry::UnRegisterSlateStyle(*this);
+	}
+}  // End namespace UE::Chaos::ClothAsset
diff --git a/Engine/Plugins/Experimental/ChaosClothAssetEditor/Source/ChaosClothAssetTools/Private/ChaosClothAsset/ClothPresetEditorStyle.h b/Engine/Plugins/Experimental/ChaosClothAssetEditor/Source/ChaosClothAssetTools/Private/ChaosClothAsset/ClothPresetEditorStyle.h
new file mode 100644
--- /dev/null
+++ b/Engine/Plugins/Experimental/ChaosClothAssetEditor/Source/ChaosClothAssetTools/Private/ChaosClothAsset/ClothPresetEditorStyle.h
@@ -0,0 +1,36 @@
+// Copyright Epic Games, Inc. All Rights Reserved.
+
+#pragma once
+
+#include "Styling/SlateStyle.h"
+
+namespace UE::Chaos::ClothAsset
+{
+	/**
+	 * Editor style setting up the cloth preset icons in editor.
+	 */
+	class FClothPresetEditorStyle final : public FSlateStyleSet
+	{
+	public:
+		FClothPresetEditorStyle();
+		virtual ~FClothPresetEditorStyle() override;
+
+	public:
+		static FClothPresetEditorStyle& Get()
+		{
+			if (!Singleton.IsSet())
+			{
+				Singleton.Emplace();
+			}
+			return Singleton.GetValue();
+		}
+
+		static void Destroy()
+		{
+			Singleton.Reset();
+		}
+
+	private:
+		static TOptional<FClothPresetEditorStyle> Singleton;
+	};
+}  // End namespace UE::Chaos::ClothAsset
diff --git a/Engine/Plugins/Experimental/ChaosClothAssetEditor/Source/ChaosClothAssetTools/Private/ChaosClothAsset/ClothPresetFactory.cpp b/Engine/Plugins/Experimental/ChaosClothAssetEditor/Source/ChaosClothAssetTools/Private/ChaosClothAsset/ClothPresetFactory.cpp
new file mode 100644
--- /dev/null
+++ b/Engine/Plugins/Experimental/ChaosClothAssetEditor/Source/ChaosClothAssetTools/Private/ChaosClothAsset/ClothPresetFactory.cpp
@@ -0,0 +1,22 @@
+// Copyright Epic Games, Inc. All Rights Reserved.
+
+#include "ClothPresetFactory.h"
+#include "ChaosClothAsset/ClothPreset.h"
+
+#include UE_INLINE_GENERATED_CPP_BY_NAME(ClothPresetFactory)
+
+UChaosClothPresetFactory::UChaosClothPresetFactory(const FObjectInitializer& ObjectInitializer)
+	: Super(ObjectInitializer)
+{
+	bEditorImport = true;
+	bEditAfterNew = true;
+	SupportedClass = UChaosClothPreset::StaticClass();
+}
+
+UObject* UChaosClothPresetFactory::FactoryCreateNew(UClass* Class, UObject* Parent, FName Name, EObjectFlags Flags, UObject* /*Context*/, FFeedbackContext* /*Warn*/)
+{
+	UChaosClothPreset* const NewClothPreset = NewObject<UChaosClothPreset>(Parent, Class, Name, Flags | RF_Transactional | RF_Public | RF_Standalone);
+	NewClothPreset->MarkPackageDirty();
+	return NewClothPreset;
+}
+
diff --git a/Engine/Plugins/Experimental/ChaosClothAssetEditor/Source/ChaosClothAssetTools/Private/ChaosClothAsset/ClothPresetFactory.h b/Engine/Plugins/Experimental/ChaosClothAssetEditor/Source/ChaosClothAssetTools/Private/ChaosClothAsset/ClothPresetFactory.h
new file mode 100644
--- /dev/null
+++ b/Engine/Plugins/Experimental/ChaosClothAssetEditor/Source/ChaosClothAssetTools/Private/ChaosClothAsset/ClothPresetFactory.h
@@ -0,0 +1,27 @@
+// Copyright Epic Games, Inc. All Rights Reserved.
+
+#pragma once
+
+#include "CoreMinimal.h"
+#include "UObject/ObjectMacros.h"
+#include "Factories/Factory.h"
+
+#include "ClothPresetFactory.generated.h"
+
+/**
+ * Having a cloth factory allows the cloth preset to be created from the Editor's menus.
+ */
+UCLASS(Experimental)
+class UChaosClothPresetFactory : public UFactory
+{
+	GENERATED_BODY()
+public:
+	UChaosClothPresetFactory(const FObjectInitializer& ObjectInitializer);
+
+	/** UFactory Interface */
+	virtual bool CanCreateNew() const override { return true; }
+	virtual bool FactoryCanImport(const FString& Filename) override { return false; }
+	virtual bool ShouldShowInNewMenu() const override { return true; }
+	virtual UObject* FactoryCreateNew(UClass* InClass, UObject* InParent, FName InName, EObjectFlags Flags, UObject* Context, FFeedbackContext* Warn) override;
+	/** End UFactory Interface */
+};
diff --git a/Engine/Plugins/Experimental/ChaosClothAssetEditor/Source/ChaosClothAssetTools/Private/ChaosClothAsset/ClothingAssetToClothAssetExporter.cpp b/Engine/Plugins/Experimental/ChaosClothAssetEditor/Source/ChaosClothAssetTools/Private/ChaosClothAsset/ClothingAssetToClothAssetExporter.cpp
new file mode 100644
--- /dev/null
+++ b/Engine/Plugins/Experimental/ChaosClothAssetEditor/Source/ChaosClothAssetTools/Private/ChaosClothAsset/ClothingAssetToClothAssetExporter.cpp
@@ -0,0 +1,63 @@
+// Copyright Epic Games, Inc. All Rights Reserved.
+
+#include "ClothingAssetToClothAssetExporter.h"
+#include "ChaosClothAsset/ClothAsset.h"
+#include "ChaosClothAsset/ClothAdapter.h"
+#include "Animation/Skeleton.h"
+#include "Materials/Material.h"
+#include "Misc/MessageDialog.h"
+#include "ClothingAsset.h"
+
+#include UE_INLINE_GENERATED_CPP_BY_NAME(ClothingAssetToClothAssetExporter)
+
+#define LOCTEXT_NAMESPACE "ClothingAssetToClothAssetExporter"
+
+UClass* UClothingAssetToChaosClothAssetExporter::GetExportedType() const
+{
+	return UChaosClothAsset::StaticClass();
+}
+
+void UClothingAssetToChaosClothAssetExporter::Export(const UClothingAssetBase* ClothingAsset, UObject* ExportedAsset)
+{
+	using namespace UE::Chaos::ClothAsset;
+
+	const UClothingAssetCommon* const ClothingAssetCommon = ExactCast<UClothingAssetCommon>(ClothingAsset);
+	if (!ClothingAssetCommon)
+	{
+		const FText TitleMessage = LOCTEXT("ClothingAssetExporterTitle", "Error Exporting Clothing Asset");
+		const FText ErrorMessage = LOCTEXT("ClothingAssetExporterError", "Can only export from known ClothingAssetCommon types.");
+		FMessageDialog::Open(EAppMsgType::Ok, EAppReturnType::Ok, ErrorMessage, &TitleMessage);
+		return;
+	}
+
+	UChaosClothAsset* const ClothAsset = CastChecked<UChaosClothAsset>(ExportedAsset);
+	check(ClothAsset);
+
+	for (const FClothLODDataCommon& ClothLODData : ClothingAssetCommon->LodData)
+	{
+		const FClothPhysicalMeshData& PhysicalMeshData = ClothLODData.PhysicalMeshData;
+
+		// Unwrap the physical mesh data into the pattern and rest meshes
+		FClothAdapter Cloth(ClothAsset->GetClothCollection());
+		FClothLodAdapter ClothLod = Cloth.AddGetLod();
+		ClothLod.Initialize(PhysicalMeshData.Vertices, PhysicalMeshData.Indices);
+	}
+
+	// Add a default material
+	ClothAsset->GetMaterials().Reset(1);
+	UMaterialInterface* const DefaultMaterial = LoadObject<UMaterialInterface>(nullptr, TEXT("/Engine/EditorMaterials/Cloth/CameraLitDoubleSided.CameraLitDoubleSided"), nullptr, LOAD_None, nullptr);
+	const int32 MaterialId = ClothAsset->GetMaterials().Emplace(DefaultMaterial, true, false, DefaultMaterial->GetFName());
+
+	// Set a default skeleton
+	USkeleton* const DefaultSkeleton = LoadObject<USkeleton>(nullptr, TEXT("/Engine/EditorMeshes/SkeletalMesh/DefaultSkeletalMesh_Skeleton.DefaultSkeletalMesh_Skeleton"), nullptr, LOAD_None, nullptr);
+	ClothAsset->SetSkeleton(DefaultSkeleton);
+
+	// Set the render mesh to duplicate the sim mesh
+	ClothAsset->CopySimMeshToRenderMesh(MaterialId);
+
+	// Build the asset, since it is already loaded, it won't rebuild on load
+	ClothAsset->Build();
+}
+
+#undef LOCTEXT_NAMESPACE
+
diff --git a/Engine/Plugins/Experimental/ChaosClothAssetEditor/Source/ChaosClothAssetTools/Public/ChaosClothAsset/ClothPatternToDynamicMesh.h b/Engine/Plugins/Experimental/ChaosClothAssetEditor/Source/ChaosClothAssetTools/Public/ChaosClothAsset/ClothPatternToDynamicMesh.h
new file mode 100644
--- /dev/null
+++ b/Engine/Plugins/Experimental/ChaosClothAssetEditor/Source/ChaosClothAssetTools/Public/ChaosClothAsset/ClothPatternToDynamicMesh.h
@@ -0,0 +1,18 @@
+// Copyright Epic Games, Inc. All Rights Reserved.
+
+#pragma once
+
+#include "CoreMinimal.h"
+
+namespace UE::Geometry { class FDynamicMesh3; }
+class UChaosClothAsset;
+
+/**
+ * Convert a single pattern from a ClothAsset to a FDynamicMesh3
+ */
+class CHAOSCLOTHASSETTOOLS_API FClothPatternToDynamicMesh
+{
+public:
+
+	void Convert(const UChaosClothAsset* ClothAssetMeshIn, int32 LODIndex, int32 PatternIndex, bool bGet2DPattern, UE::Geometry::FDynamicMesh3& MeshOut);
+};
diff --git a/Engine/Plugins/Experimental/CharacterFXEditor/BaseCharacterFXEditor/Source/BaseCharacterFXEditor/Private/SBaseCharacterFXEditorViewport.cpp b/Engine/Plugins/Experimental/CharacterFXEditor/BaseCharacterFXEditor/Source/BaseCharacterFXEditor/Private/SBaseCharacterFXEditorViewport.cpp
new file mode 100644
--- /dev/null
+++ b/Engine/Plugins/Experimental/CharacterFXEditor/BaseCharacterFXEditor/Source/BaseCharacterFXEditor/Private/SBaseCharacterFXEditorViewport.cpp
@@ -0,0 +1,26 @@
+// Copyright Epic Games, Inc. All Rights Reserved.
+
+#include "SBaseCharacterFXEditorViewport.h"
+#include "SViewportToolBar.h"
+
+#define LOCTEXT_NAMESPACE "SBaseCharacterFXEditorViewport"
+
+void SBaseCharacterFXEditorViewport::AddOverlayWidget(TSharedRef<SWidget> OverlaidWidget)
+{
+	ViewportOverlay->AddSlot()
+	[
+		OverlaidWidget
+	];
+}
+
+void SBaseCharacterFXEditorViewport::RemoveOverlayWidget(TSharedRef<SWidget> OverlaidWidget)
+{
+	ViewportOverlay->RemoveSlot(OverlaidWidget);
+}
+
+TSharedPtr<SWidget> SBaseCharacterFXEditorViewport::MakeViewportToolbar()
+{
+	return SNew(SViewportToolBar);
+}
+
+#undef LOCTEXT_NAMESPACE
diff --git a/Engine/Plugins/Experimental/CharacterFXEditor/BaseCharacterFXEditor/Source/BaseCharacterFXEditor/Public/BaseCharacterFXEditorModeToolkit.h b/Engine/Plugins/Experimental/CharacterFXEditor/BaseCharacterFXEditor/Source/BaseCharacterFXEditor/Public/BaseCharacterFXEditorModeToolkit.h
new file mode 100644
--- /dev/null
+++ b/Engine/Plugins/Experimental/CharacterFXEditor/BaseCharacterFXEditor/Source/BaseCharacterFXEditor/Public/BaseCharacterFXEditorModeToolkit.h
@@ -0,0 +1,85 @@
+// Copyright Epic Games, Inc. All Rights Reserved.
+
+#pragma once
+
+#include "CoreMinimal.h"
+
+#include "Toolkits/BaseToolkit.h"
+#include "StatusBarSubsystem.h"
+
+class SBorder;
+class STextBlock;
+
+
+/**
+ * The editor mode toolkit is responsible for the panel on the side in the editor
+ * that shows mode and tool properties. Tool buttons would go in Init().
+ */
+
+class BASECHARACTERFXEDITOR_API FBaseCharacterFXEditorModeToolkit : public FModeToolkit
+{
+public:
+
+	FBaseCharacterFXEditorModeToolkit();
+	~FBaseCharacterFXEditorModeToolkit();
+
+	// FModeToolkit
+	virtual void Init(const TSharedPtr<IToolkitHost>& InitToolkitHost, TWeakObjectPtr<UEdMode> InOwningMode) override;
+	
+	/** Active tool name to display in the ViewportOverlay */
+	virtual FText GetActiveToolDisplayName() const override { return ActiveToolName; }
+
+	/** Returns the Mode specific tabs in the mode toolbar **/
+	virtual void GetToolPaletteNames(TArray<FName>& InPaletteName) const override;
+
+    /** Returns human readable name for the specific palette tool category **/
+	virtual FText GetToolPaletteDisplayName(FName PaletteName) const override;
+ 
+	/** Informs the built-in FModeToolkit palette building routines about how we want the toolbars to be set up **/
+	virtual bool HasIntegratedToolPalettes() const override { return true; }
+	virtual bool HasExclusiveToolPalettes() const override { return false; }
+
+	/** Returns the mode's toolbox */
+	virtual TSharedPtr<class SWidget> GetInlineContent() const override { return ToolkitWidget; }
+
+	/** Set/clear notification message area */
+	virtual void PostNotification(const FText& Message);
+	virtual void ClearNotification();
+
+	/** Set/clear warning message area */
+	virtual void PostWarning(const FText& Message);
+	virtual void ClearWarning();
+
+protected:
+
+	// Should be implemented in the concrete subclass since it will likely have to know where to look for UI resources
+	virtual const FSlateBrush* GetActiveToolIcon(const FString& Identifier) const PURE_VIRTUAL(FBaseCharacterFXEditorModeToolkit::GetActiveToolIcon, return nullptr;);
+
+	// The editor can have multiple tabs in the mode toolbar. ToolsTabName is the "Tools" tab maintained by this
+	// base class. PaletteNames_Standard is the set of all tab names in the concrete editor. Set these in the concrete
+	// subclass when adding additional toolbar tabs
+	const static FName ToolsTabName;
+	const static TArray<FName> PaletteNames_Standard;
+
+	// FModeToolkit
+	virtual void OnToolStarted(UInteractiveToolManager* Manager, UInteractiveTool* Tool) override;
+	virtual void OnToolEnded(UInteractiveToolManager* Manager, UInteractiveTool* Tool) override;
+
+	// A place for tools to write out any warnings
+	TSharedPtr<STextBlock> ToolWarningArea;
+
+	// A container for the tool settings that is populated by the DetailsView managed in FModeToolkit
+	TSharedPtr<SBorder> ToolDetailsContainer;
+
+	// Contains the widget container for the Accept/Cancel buttons, etc. for tools
+	TSharedPtr<SWidget> ViewportOverlayWidget;
+
+	// Tool name, icon and message to display in viewport
+	FText ActiveToolName;
+	const FSlateBrush* ActiveToolIcon;
+	FStatusBarMessageHandle ActiveToolMessageHandle;
+
+	void UpdateActiveToolProperties();
+	void InvalidateCachedDetailPanelState(UObject* ChangedObject);
+};
+
diff --git a/Engine/Plugins/Experimental/CharacterFXEditor/BaseCharacterFXEditor/Source/BaseCharacterFXEditor/Public/BaseCharacterFXEditorModeUILayer.h b/Engine/Plugins/Experimental/CharacterFXEditor/BaseCharacterFXEditor/Source/BaseCharacterFXEditor/Public/BaseCharacterFXEditorModeUILayer.h
new file mode 100644
--- /dev/null
+++ b/Engine/Plugins/Experimental/CharacterFXEditor/BaseCharacterFXEditor/Source/BaseCharacterFXEditor/Public/BaseCharacterFXEditorModeUILayer.h
@@ -0,0 +1,56 @@
+// Copyright Epic Games, Inc. All Rights Reserved.
+
+#pragma once
+
+#include "Toolkits/AssetEditorModeUILayer.h"
+#include "BaseCharacterFXEditorModeUILayer.generated.h"
+
+/** Handles the hosting of additional toolkits, such as the mode toolkit, within the CharacterFXEditor's toolkit. **/
+
+class BASECHARACTERFXEDITOR_API FBaseCharacterFXEditorModeUILayer : public FAssetEditorModeUILayer
+{
+public:
+	FBaseCharacterFXEditorModeUILayer(const IToolkitHost* InToolkitHost);
+	void OnToolkitHostingStarted(const TSharedRef<IToolkit>& Toolkit) override;
+	void OnToolkitHostingFinished(const TSharedRef<IToolkit>& Toolkit) override;
+
+	void SetModeMenuCategory(TSharedPtr<FWorkspaceItem> MenuCategoryIn);
+	TSharedPtr<FWorkspaceItem> GetModeMenuCategory() const override;
+
+protected:
+
+	TSharedPtr<FWorkspaceItem> CharacterFXEditorMenuCategory;
+
+};
+
+/** 
+ * Interchange layer to manage built in tab locations within the editor's layout. 
+ */
+
+UCLASS()
+class BASECHARACTERFXEDITOR_API UBaseCharacterFXEditorUISubsystem : public UAssetEditorUISubsystem
+{
+	GENERATED_BODY()
+public:
+
+	// Adds RegisterLayoutExtensions as a callback to the module corresponding to GetModuleName()
+	virtual void Initialize(FSubsystemCollectionBase& Collection) override;
+
+	// Removes RegisterLayoutExtensions callback from the module corresponding to GetModuleName()
+	virtual void Deinitialize() override;
+
+	// Docks the editor Mode tab in the Editor side panel area of the layout
+	virtual void RegisterLayoutExtensions(FLayoutExtender& Extender) override;
+
+	// Identifier for the Layout extension in BaseCharacterFXEditorToolkit's StandaloneDefaultLayout
+	static const FName EditorSidePanelAreaName;
+
+protected:
+
+	// OVERRIDE THIS
+	virtual FName GetModuleName() const
+	{
+		return FName("");
+	}
+
+};
diff --git a/Engine/Plugins/Experimental/CharacterFXEditor/BaseCharacterFXEditor/Source/BaseCharacterFXEditor/Public/BaseCharacterFXEditorToolkit.h b/Engine/Plugins/Experimental/CharacterFXEditor/BaseCharacterFXEditor/Source/BaseCharacterFXEditor/Public/BaseCharacterFXEditorToolkit.h
new file mode 100644
--- /dev/null
+++ b/Engine/Plugins/Experimental/CharacterFXEditor/BaseCharacterFXEditor/Source/BaseCharacterFXEditor/Public/BaseCharacterFXEditorToolkit.h
@@ -0,0 +1,67 @@
+// Copyright Epic Games, Inc. All Rights Reserved.
+
+#pragma once
+
+#include "CoreMinimal.h"
+#include "Tools/BaseAssetToolkit.h"
+#include "BaseCharacterFXEditorModeUILayer.h"
+#include "PreviewScene.h"
+
+class UBaseCharacterFXEditorMode;
+
+/**
+ * The toolkit is supposed to act as the UI manager for the asset editor. It's responsible 
+ * for setting up viewports and most toolbars, except for the internals of the mode panel.
+ * However, because the toolkit also sets up the mode manager, and much of the important
+ * state is held in the mode managed by the mode manager, the toolkit also ends up
+ * initializing the mode.
+ * Thus, the FBaseCharacterFXEditorToolkit ends up being the central place for the CharacterFX Asset Editor setup.
+ */
+
+class BASECHARACTERFXEDITOR_API FBaseCharacterFXEditorToolkit : public FBaseAssetToolkit
+{
+public:
+	    
+	FBaseCharacterFXEditorToolkit(UAssetEditor* InOwningAssetEditor, const FName& ModuleName);
+
+	virtual void OnToolkitHostingStarted(const TSharedRef<IToolkit>& Toolkit) override;
+	virtual void OnToolkitHostingFinished(const TSharedRef<IToolkit>& Toolkit) override;
+	virtual bool OnRequestClose() override;
+	virtual void RegisterTabSpawners(const TSharedRef<FTabManager>& InTabManager) override;
+
+	virtual AssetEditorViewportFactoryFunction GetViewportDelegate() override;
+	virtual void AddViewportOverlayWidget(TSharedRef<SWidget> InViewportOverlayWidget) override;
+	virtual void RemoveViewportOverlayWidget(TSharedRef<SWidget> InViewportOverlayWidget) override;
+
+protected:
+
+	// Override this to get the Mode ID of the concrete CharacterFXEditorMode
+	virtual FEditorModeID GetEditorModeId() const PURE_VIRTUAL(FBaseCharacterFXEditorToolkit::GetEditorModeId, return FEditorModeID("EM_ERROR"););
+
+	// Helpers for PostInitAssetEditor
+	
+	// Override this to create the concrete CharacterFXEditorModeUILayer
+	virtual void CreateEditorModeUILayer();
+
+	// Override this to provide the EdMode with any custom data it needs
+	virtual void InitializeEdMode(UBaseCharacterFXEditorMode* EdMode);
+
+	virtual void PostInitAssetEditor() override;
+
+	// Called from FBaseAssetToolkit::CreateWidgets.
+	virtual void CreateEditorModeManager() override;
+	
+	// Called from FBaseAssetToolkit::CreateWidgets to populate ViewportClient, but otherwise only used 
+	// in our own viewport delegate
+	virtual TSharedPtr<FEditorViewportClient> CreateEditorViewportClient() const override;
+
+	// Handles the hosting of additional toolkits, such as the mode toolkit
+	TSharedPtr<FBaseCharacterFXEditorModeUILayer> ModeUILayer;
+
+	// PreviewScene showing the objects being edited
+	TUniquePtr<FPreviewScene> ObjectScene;
+
+	// A default "CharacterFXEditor" category is created in RegisterTabSpawners. Override that function to set this to a concrete class type
+	TSharedPtr<FWorkspaceItem> EditorMenuCategory;
+};
+
diff --git a/Engine/Plugins/Experimental/CharacterFXEditor/BaseCharacterFXEditor/Source/BaseCharacterFXEditor/Public/SBaseCharacterFXEditorViewport.h b/Engine/Plugins/Experimental/CharacterFXEditor/BaseCharacterFXEditor/Source/BaseCharacterFXEditor/Public/SBaseCharacterFXEditorViewport.h
new file mode 100644
--- /dev/null
+++ b/Engine/Plugins/Experimental/CharacterFXEditor/BaseCharacterFXEditor/Source/BaseCharacterFXEditor/Public/SBaseCharacterFXEditorViewport.h
@@ -0,0 +1,19 @@
+// Copyright Epic Games, Inc. All Rights Reserved.
+
+#pragma once
+
+#include "CoreMinimal.h"
+#include "SAssetEditorViewport.h"
+
+class BASECHARACTERFXEDITOR_API SBaseCharacterFXEditorViewport : public SAssetEditorViewport
+{
+public:
+
+	// These allow the toolkit to add an accept/cancel overlay when needed. PopulateViewportOverlays
+	// is not helpful here because that gets called just once.
+	virtual void AddOverlayWidget(TSharedRef<SWidget> OverlaidWidget);
+	virtual void RemoveOverlayWidget(TSharedRef<SWidget> OverlaidWidget);
+
+	// SEditorViewport
+	virtual TSharedPtr<SWidget> MakeViewportToolbar() override;
+};
