Cloth Editor: fix crash when the sim  preview viewport tab is closed and the editor is restarted

diff --git a/Engine/Plugins/Experimental/ChaosClothAssetEditor/Source/ChaosClothAssetEditor/Private/ChaosClothAsset/ClothEditorToolkit.cpp b/Engine/Plugins/Experimental/ChaosClothAssetEditor/Source/ChaosClothAssetEditor/Private/ChaosClothAsset/ClothEditorToolkit.cpp
--- a/Engine/Plugins/Experimental/ChaosClothAssetEditor/Source/ChaosClothAssetEditor/Private/ChaosClothAsset/ClothEditorToolkit.cpp
+++ b/Engine/Plugins/Experimental/ChaosClothAssetEditor/Source/ChaosClothAssetEditor/Private/ChaosClothAsset/ClothEditorToolkit.cpp
@@ -361,6 +361,13 @@ void FChaosClothAssetEditorToolkit::PostInitAssetEditor()
 	ClothMode->SetRestSpaceViewportClient(StaticCastWeakPtr<FChaosClothEditorRestSpaceViewportClient>(WeakViewportClient));
 	ClothMode->RefocusRestSpaceViewportClient();
 
+	// Note: We force the cloth preview viewport to open, since some ViewportClient functions are not robust to having no viewport.
+	// See UE-114649
+	if (!TabManager->FindExistingLiveTab(ClothPreviewTabID))
+	{
+		TabManager->TryInvokeTab(ClothPreviewTabID);
+	}
+
 	// We need the viewport client to start out focused, or else it won't get ticked until
 	// we click inside it.
 	ViewportClient->ReceivedFocus(ViewportClient->Viewport);
@@ -369,9 +376,11 @@ void FChaosClothAssetEditorToolkit::PostInitAssetEditor()
 	SetCommonViewportClientOptions(ClothPreviewViewportClient.Get());
 	ClothPreviewViewportClient->SetInitialViewTransform(ELevelViewportType::LVT_Perspective, FVector(0, -100, 100), FRotator(0, 90, 0), DEFAULT_ORTHOZOOM);
 
-	FBox PreviewBounds = ClothMode->PreviewBoundingBox();
-	ClothPreviewViewportClient->FocusViewportOnBox(PreviewBounds);
-
+	if (ClothPreviewViewportClient->Viewport != nullptr)
+	{
+		FBox PreviewBounds = ClothMode->PreviewBoundingBox();
+		ClothPreviewViewportClient->FocusViewportOnBox(PreviewBounds);
+	}
 }
 
 
