Cloth Editor: enable adding/deleting weight maps directly in the weight map paint tool. (Currently the user has to go through the Attribute Editor tool to do this.)

diff --git a/Engine/Plugins/Experimental/ChaosClothAssetEditor/Source/ChaosClothAssetEditorTools/Private/ChaosClothAsset/ClothWeightMapPaintTool.cpp b/Engine/Plugins/Experimental/ChaosClothAssetEditor/Source/ChaosClothAssetEditorTools/Private/ChaosClothAsset/ClothWeightMapPaintTool.cpp
--- a/Engine/Plugins/Experimental/ChaosClothAssetEditor/Source/ChaosClothAssetEditorTools/Private/ChaosClothAsset/ClothWeightMapPaintTool.cpp
+++ b/Engine/Plugins/Experimental/ChaosClothAssetEditor/Source/ChaosClothAssetEditorTools/Private/ChaosClothAsset/ClothWeightMapPaintTool.cpp
@@ -142,16 +142,17 @@ void UClothEditorWeightMapPaintTool::Setup()
 
 	WeightMapSetProperties = NewObject<UWeightMapSetProperties>(this);
 	WeightMapSetProperties->RestoreProperties(this, TEXT("ClothEditorWeightMapPaintTool"));
-	TArray<FName> WeightMapNames;
-	for (int32 WeightMapIndex = 0; WeightMapIndex < Mesh->Attributes()->NumWeightLayers(); ++WeightMapIndex)
-	{
-		WeightMapNames.Add(Mesh->Attributes()->GetWeightLayer(WeightMapIndex)->GetName());
-	}
-	WeightMapSetProperties->InitializeWeightMaps(WeightMapNames);
+
+	InitializeWeightMapNames();
+
 	WeightMapSetProperties->WatchProperty(WeightMapSetProperties->WeightMap, [&](FName) { OnSelectedWeightMapChanged(); });
 	UpdateActiveWeightMap();
 	AddToolPropertySource(WeightMapSetProperties);
 
+	ClothEditorWeightMapActions = NewObject<UClothEditorWeightMapActions>(this);
+	ClothEditorWeightMapActions->Initialize(this);
+	AddToolPropertySource(ClothEditorWeightMapActions);
+
 	// initialize other properties
 	FilterProperties = NewObject<UClothEditorWeightMapPaintBrushFilterProperties>(this);
 	FilterProperties->WatchProperty(FilterProperties->SubToolType,
@@ -1218,7 +1219,105 @@ void UClothEditorWeightMapPaintTool::ClearAllWeightsAction()
 }
 
 
+void UClothEditorWeightMapPaintTool::InitializeWeightMapNames()
+{
+	const FDynamicMesh3* const Mesh = GetSculptMesh();
+
+	if (!WeightMapSetProperties || !Mesh || !Mesh->HasAttributes())
+	{
+		return;
+	}
+
+	TArray<FString>& PropWeightMapNames = WeightMapSetProperties->WeightMapsList;
+	PropWeightMapNames.Reset();
+	PropWeightMapNames.Add(TEXT("None"));
+
+	// get names from the mesh attributes
+	for (int32 WeightMapIndex = 0; WeightMapIndex < Mesh->Attributes()->NumWeightLayers(); ++WeightMapIndex)
+	{
+		const FName& MeshWeightMapName = Mesh->Attributes()->GetWeightLayer(WeightMapIndex)->GetName();
+		PropWeightMapNames.Add(MeshWeightMapName.ToString());
+	}
+
+	// if selected weight map name is no longer in the list of valid names, reset it
+	if (!PropWeightMapNames.Contains(WeightMapSetProperties->WeightMap.ToString()))
+	{
+		WeightMapSetProperties->WeightMap = FName(PropWeightMapNames[0]);
+	}
+}
+
+
+void UClothEditorWeightMapPaintTool::AddWeightMapAction(const FName& NewWeightMapName)
+{
+	if (DynamicMeshComponent && DynamicMeshComponent->GetMesh() && DynamicMeshComponent->GetMesh()->HasAttributes())
+	{
+		FDynamicMesh3* const Mesh = DynamicMeshComponent->GetMesh();
 
+		FDynamicMeshAttributeSet* const MeshAttributes = Mesh->Attributes();
+		const int32 NumExistingWeights = MeshAttributes->NumWeightLayers();
+		
+		if (NewWeightMapName.IsNone())
+		{
+			GetToolManager()->DisplayMessage(LOCTEXT("InvalidAttributeName", "Invalid weight map name"), EToolMessageLevel::UserWarning);
+			return;
+		}
+
+		// Check for matching name	
+		for (int32 WeightMapIndex = 0; WeightMapIndex < NumExistingWeights; ++WeightMapIndex)
+		{
+			const FDynamicMeshWeightAttribute* const ExistingWeightMap = MeshAttributes->GetWeightLayer(WeightMapIndex);
+			if (ExistingWeightMap->GetName() == NewWeightMapName)
+			{
+				GetToolManager()->DisplayMessage(LOCTEXT("ErrorAddingDuplicateNameMessage", "Weight map with this name already exists"), EToolMessageLevel::UserWarning);
+				return;
+			}
+		}
+
+		MeshAttributes->SetNumWeightLayers(NumExistingWeights + 1);
+
+		const int32 NewWeightMapIndex = NumExistingWeights;
+		FDynamicMeshWeightAttribute* const NewWeightMap = MeshAttributes->GetWeightLayer(NewWeightMapIndex);
+		NewWeightMap->SetName(NewWeightMapName);
+
+		// Refresh the drop down of available weight maps and select the new one
+		InitializeWeightMapNames();
+		WeightMapSetProperties->SetSelectedFromWeightMapIndex(NewWeightMapIndex);
+
+		// Clear the New Name text field
+		ClothEditorWeightMapActions->NewWeightMapName.Empty();
+	}
+}
+
+
+void UClothEditorWeightMapPaintTool::DeleteWeightMapAction(const FName& SelectedWeightMapName)
+{
+	FDynamicMeshAttributeSet* const MeshAttributes = DynamicMeshComponent->GetMesh()->Attributes();
+
+	bool bFoundSelectedWeightMap = false;
+	
+	const int32 NumWeightMaps = MeshAttributes->NumWeightLayers();
+	for (int32 WeightMapIndex = 0; WeightMapIndex < NumWeightMaps; ++WeightMapIndex)
+	{
+		const FDynamicMeshWeightAttribute* const WeightMap = MeshAttributes->GetWeightLayer(WeightMapIndex);
+		if (WeightMap->GetName() == SelectedWeightMapName)
+		{
+			bFoundSelectedWeightMap = true;
+			MeshAttributes->RemoveWeightLayer(WeightMapIndex);		// Hopefully nothing else is holding indices into this...
+			break;
+		}
+	}
+
+	if (!bFoundSelectedWeightMap)
+	{
+		GetToolManager()->DisplayMessage(LOCTEXT("InvalidExistingWeightMapName", "No weight map with the specified name exists"), EToolMessageLevel::UserWarning);
+	}
+	else
+	{
+		InitializeWeightMapNames();
+		UpdateActiveWeightMap();
+		GetToolManager()->DisplayMessage(FText(), EToolMessageLevel::UserWarning);
+	}
+}
 
 //
 // Change Tracking
@@ -1525,6 +1624,14 @@ void UClothEditorWeightMapPaintTool::ApplyAction(EClothEditorWeightMapPaintToolA
 		ClearAllWeightsAction();
 		break;
 
+	case EClothEditorWeightMapPaintToolActions::AddWeightMap:
+		AddWeightMapAction(FName(ClothEditorWeightMapActions->NewWeightMapName));
+		break;
+
+	case EClothEditorWeightMapPaintToolActions::DeleteWeightMap:
+		DeleteWeightMapAction(WeightMapSetProperties->WeightMap);
+		break;
+
 	}
 }
 
@@ -1532,4 +1639,3 @@ void UClothEditorWeightMapPaintTool::ApplyAction(EClothEditorWeightMapPaintToolA
 
 
 #undef LOCTEXT_NAMESPACE
-
diff --git a/Engine/Plugins/Experimental/ChaosClothAssetEditor/Source/ChaosClothAssetEditorTools/Public/ChaosClothAsset/ClothWeightMapPaintTool.h b/Engine/Plugins/Experimental/ChaosClothAssetEditor/Source/ChaosClothAssetEditorTools/Public/ChaosClothAsset/ClothWeightMapPaintTool.h
--- a/Engine/Plugins/Experimental/ChaosClothAssetEditor/Source/ChaosClothAssetEditorTools/Public/ChaosClothAsset/ClothWeightMapPaintTool.h
+++ b/Engine/Plugins/Experimental/ChaosClothAssetEditor/Source/ChaosClothAssetEditorTools/Public/ChaosClothAsset/ClothWeightMapPaintTool.h
@@ -162,7 +162,10 @@ enum class EClothEditorWeightMapPaintToolActions
 	NoAction,
 
 	FloodFillCurrent,
-	ClearAll
+	ClearAll,
+
+	AddWeightMap,
+	DeleteWeightMap
 };
 
 
@@ -187,19 +190,43 @@ class CHAOSCLOTHASSETEDITORTOOLS_API UMeshWeightMapPaintToolActions : public UCl
 public:
 
 	UFUNCTION(CallInEditor, Category = Operations, meta = (DisplayPriority = 10))
-		void ClearAll()
+	void ClearAll()
 	{
 		PostAction(EClothEditorWeightMapPaintToolActions::ClearAll);
 	}
 
 	UFUNCTION(CallInEditor, Category = Operations, meta = (DisplayPriority = 12))
-		void FloodFillCurrent()
+	void FloodFillCurrent()
 	{
 		PostAction(EClothEditorWeightMapPaintToolActions::FloodFillCurrent);
 	}
 
 };
 
+
+UCLASS()
+class CHAOSCLOTHASSETEDITORTOOLS_API UClothEditorWeightMapActions : public UClothEditorWeightMapPaintToolActionPropertySet
+{
+	GENERATED_BODY()
+
+public:
+
+	UPROPERTY(EditAnywhere, Category = AddWeightMap, meta = (DisplayName = "New Weight Map Name"))
+	FString NewWeightMapName;
+
+	UFUNCTION(CallInEditor, Category = AddWeightMap, meta = (DisplayPriority = 2))
+	void AddNewWeightMap()
+	{
+		PostAction(EClothEditorWeightMapPaintToolActions::AddWeightMap);
+	}
+
+	UFUNCTION(CallInEditor, Category = DeleteWeightMap, meta = (DisplayPriority = 3))
+	void DeleteSelectedWeightMap()
+	{
+		PostAction(EClothEditorWeightMapPaintToolActions::DeleteWeightMap);
+	}
+};
+
 /**
  * Mesh Element Paint Tool Class
  */
@@ -251,6 +278,9 @@ public:
 	void FloodFillCurrentWeightAction();
 	void ClearAllWeightsAction();
 
+	void AddWeightMapAction(const FName& NewWeightMapName);
+	void DeleteWeightMapAction(const FName& SelectedWeightMapName);
+
 	void SetVerticesToWeightMap(const TSet<int32>& Vertices, double WeightValue, bool bIsErase);
 
 	bool HaveVisibilityFilter() const;
@@ -290,6 +320,9 @@ public:
 	UPROPERTY()
 	TObjectPtr<UMeshWeightMapPaintToolActions> ActionsProps;
 
+	UPROPERTY()
+	TObjectPtr<UClothEditorWeightMapActions> ClothEditorWeightMapActions;
+
 protected:
 	bool bHavePendingAction = false;
 	EClothEditorWeightMapPaintToolActions PendingAction;
@@ -385,6 +418,8 @@ protected:
 	TArray<int32> NormalSeamEdges;
 	void PrecomputeFilterData();
 
+	// Populate the drop-down WeightMapsList in WeightMapSetProperties
+	void InitializeWeightMapNames();
 
 protected:
 	virtual bool ShowWorkPlane() const override { return false; }
diff --git a/Engine/Source/Runtime/GeometryCore/Private/DynamicMesh/DynamicMeshAttributeSet.cpp b/Engine/Source/Runtime/GeometryCore/Private/DynamicMesh/DynamicMeshAttributeSet.cpp
--- a/Engine/Source/Runtime/GeometryCore/Private/DynamicMesh/DynamicMeshAttributeSet.cpp
+++ b/Engine/Source/Runtime/GeometryCore/Private/DynamicMesh/DynamicMeshAttributeSet.cpp
@@ -535,6 +535,11 @@ void FDynamicMeshAttributeSet::SetNumWeightLayers(int32 Num)
 	ensure(WeightLayers.Num() == Num);
 }
 
+void FDynamicMeshAttributeSet::RemoveWeightLayer(int32 Index)
+{
+	WeightLayers.RemoveAt(Index);
+}
+
 FDynamicMeshWeightAttribute* FDynamicMeshAttributeSet::GetWeightLayer(int Index)
 {
 	return &WeightLayers[Index];
diff --git a/Engine/Source/Runtime/GeometryCore/Public/DynamicMesh/DynamicMeshAttributeSet.h b/Engine/Source/Runtime/GeometryCore/Public/DynamicMesh/DynamicMeshAttributeSet.h
--- a/Engine/Source/Runtime/GeometryCore/Public/DynamicMesh/DynamicMeshAttributeSet.h
+++ b/Engine/Source/Runtime/GeometryCore/Public/DynamicMesh/DynamicMeshAttributeSet.h
@@ -307,6 +307,9 @@ public:
 	/** Set the number of weight layers */
 	virtual void SetNumWeightLayers(int32 Num);
 
+	/** Remove a weight layer at the specified index */
+	virtual void RemoveWeightLayer(int32 Index);
+
 	/** @return the weight layer at the given Index */
 	FDynamicMeshWeightAttribute* GetWeightLayer(int Index);
 
