Cloth Editor:
- Allow deletion of weight map layers. If the AttributeEditor tool removes a weight map for any DynamicMeshComponent, then remove that weight map from the ClothAsset. This will remove the corresponding weight map for all other DynamicMeshComponents
- Enable saving ClothAsset
- Enable undo

diff --git a/Engine/Plugins/Experimental/ChaosClothAsset/Source/ChaosClothAsset/Public/ChaosClothAsset/ClothAdapter.h b/Engine/Plugins/Experimental/ChaosClothAsset/Source/ChaosClothAsset/Public/ChaosClothAsset/ClothAdapter.h
--- a/Engine/Plugins/Experimental/ChaosClothAsset/Source/ChaosClothAsset/Public/ChaosClothAsset/ClothAdapter.h
+++ b/Engine/Plugins/Experimental/ChaosClothAsset/Source/ChaosClothAsset/Public/ChaosClothAsset/ClothAdapter.h
@@ -65,6 +65,9 @@ namespace UE::Chaos::ClothAsset
 		/** Add a new weight map to this cloth. Access is then done per pattern. */
 		void AddWeightMap(const FName& Name) { GetClothCollection()->AddAttribute<float>(Name, FClothCollection::SimVerticesGroup); }
 
+		/** Remove a weight map from this cloth. */
+		void RemoveWeightMap(const FName& Name) { GetClothCollection()->RemoveAttribute(Name, FClothCollection::SimVerticesGroup); }
+
 		/** Return the underlaying cloth collection this adapter has been created with. */
 		TSharedPtr<FClothCollection> GetClothCollection() { return ConstCastSharedPtr<FClothCollection>(FClothConstAdapter::GetClothCollection()); }
 	};
diff --git a/Engine/Plugins/Experimental/ChaosClothAssetEditor/Source/ChaosClothAssetEditor/Private/ChaosClothAsset/ClothEditorMode.cpp b/Engine/Plugins/Experimental/ChaosClothAssetEditor/Source/ChaosClothAssetEditor/Private/ChaosClothAsset/ClothEditorMode.cpp
--- a/Engine/Plugins/Experimental/ChaosClothAssetEditor/Source/ChaosClothAssetEditor/Private/ChaosClothAsset/ClothEditorMode.cpp
+++ b/Engine/Plugins/Experimental/ChaosClothAssetEditor/Source/ChaosClothAssetEditor/Private/ChaosClothAsset/ClothEditorMode.cpp
@@ -188,6 +188,11 @@ void UChaosClothAssetEditorMode::OnToolEnded(UInteractiveToolManager* Manager, U
 	ReinitializeDynamicMeshComponents();
 }
 
+void UChaosClothAssetEditorMode::PostUndo()
+{
+	ReinitializeDynamicMeshComponents();
+}
+
 void UChaosClothAssetEditorMode::BindCommands()
 {
 	const FChaosClothAssetEditorCommands& CommandInfos = FChaosClothAssetEditorCommands::Get();
@@ -258,14 +263,70 @@ void UChaosClothAssetEditorMode::GetRestSpaceMesh(UToolTarget* ToolTarget, UE::G
 	RestSpaceMesh = UE::ToolTarget::GetDynamicMeshCopy(ToolTarget);
 }
 
+namespace ChaosClothAssetEditorModeHelpers
+{
+	TArray<FName> GetClothAssetWeightMapNames(const UE::Chaos::ClothAsset::FClothAdapter& ClothAdapter)
+	{
+		const TSharedPtr<const UE::Chaos::ClothAsset::FClothCollection> ClothCollection = ClothAdapter.FClothConstAdapter::GetClothCollection();
+
+		TArray<FName> OutWeightMapNames;
+		const TArray<FName> SimVerticesAttributeNames = ClothCollection->AttributeNames(UE::Chaos::ClothAsset::FClothCollection::SimVerticesGroup);
+		for (const FName& AttributeName : SimVerticesAttributeNames)
+		{
+			if (ClothCollection->FindAttributeTyped<float>(AttributeName, UE::Chaos::ClothAsset::FClothCollection::SimVerticesGroup))
+			{
+				OutWeightMapNames.Add(AttributeName);
+			}
+		}
+
+		return OutWeightMapNames;
+	}
+
+	void RemoveClothWeightMaps(UE::Chaos::ClothAsset::FClothAdapter& ClothAdapter, const TArray<FName>& WeightMapNames)
+	{
+		const TSharedPtr<const UE::Chaos::ClothAsset::FClothCollection> ClothCollection = ClothAdapter.GetClothCollection();
+		for (const FName& RemovedMap : WeightMapNames)
+		{
+			if (ClothCollection->FindAttributeTyped<float>(RemovedMap, UE::Chaos::ClothAsset::FClothCollection::SimVerticesGroup))
+			{
+				ClothAdapter.RemoveWeightMap(RemovedMap);
+			}
+		}
+	}
+
+	TArray<FName> GetDynamicMeshWeightMapNames(const UE::Geometry::FDynamicMesh3& DynamicMesh)
+	{
+		TArray<FName> OutWeightMapNames;
+
+		for (int32 DynamicMeshWeightMapIndex = 0; DynamicMeshWeightMapIndex < DynamicMesh.Attributes()->NumWeightLayers(); ++DynamicMeshWeightMapIndex)
+		{
+			const UE::Geometry::FDynamicMeshWeightAttribute* const WeightMapAttribute = DynamicMesh.Attributes()->GetWeightLayer(DynamicMeshWeightMapIndex);
+			const FName WeightMapName = WeightMapAttribute->GetName();
+			OutWeightMapNames.Add(WeightMapName);
+		}
+
+		return OutWeightMapNames;
+	}
+}
+
 
 void UChaosClothAssetEditorMode::UpdateSimulationMeshes()
 {
-	const UChaosClothAsset* ChaosClothAsset = ClothComponent->GetClothAsset();
+	GetToolManager()->BeginUndoTransaction(LOCTEXT("ChaosClothAssetEditorApplyChangesTransaction", "Cloth Editor Apply Changes"));
+
+	UChaosClothAsset* ChaosClothAsset = ClothComponent->GetClothAsset();
+	ChaosClothAsset->Modify();
+
 	UE::Chaos::ClothAsset::FClothAdapter ClothAdapter(ChaosClothAsset->GetClothCollection());
-	
+
 	check(DynamicMeshSourceInfos.Num() == DynamicMeshComponents.Num());
-	
+
+	// Search for weight map names that are in the original cloth asset, but are not in *all* dynamic mesh components. These are weight maps
+	// that have been removed by one of the mesh tools, and so should be removed from the asset.
+	const TSet<FName> ClothAssetWeightMapNames = TSet<FName>(ChaosClothAssetEditorModeHelpers::GetClothAssetWeightMapNames(ClothAdapter));
+	TSet<FName> CommonDynamicMeshWeightMapNames;
+	bool bFirstIteration = true;
+
 	for (int DynamicMeshIndex = 0; DynamicMeshIndex < DynamicMeshComponents.Num(); ++DynamicMeshIndex)
 	{
 		const int32 LodIndex = DynamicMeshSourceInfos[DynamicMeshIndex].LodIndex;
@@ -290,6 +351,17 @@ void UChaosClothAssetEditorMode::UpdateSimulationMeshes()
 		{
 			continue;
 		}
+		 
+		const TSet<FName> DynamicMeshWeightMapNames(ChaosClothAssetEditorModeHelpers::GetDynamicMeshWeightMapNames(*DynamicMesh));
+		if (bFirstIteration)
+		{
+			CommonDynamicMeshWeightMapNames = DynamicMeshWeightMapNames;
+			bFirstIteration = false;
+		}
+		else
+		{
+			CommonDynamicMeshWeightMapNames = CommonDynamicMeshWeightMapNames.Intersect(DynamicMeshWeightMapNames);
+		}
 
 		const int NumInputWeightElements = DynamicMesh->MaxVertexID();
 
@@ -309,6 +381,11 @@ void UChaosClothAssetEditorMode::UpdateSimulationMeshes()
 			}
 		}
 	}
+
+	const TSet<FName> WeightMapsToRemove = ClothAssetWeightMapNames.Difference(CommonDynamicMeshWeightMapNames);
+	ChaosClothAssetEditorModeHelpers::RemoveClothWeightMaps(ClothAdapter, WeightMapsToRemove.Array());
+
+	GetToolManager()->EndUndoTransaction();
 }
 
 void UChaosClothAssetEditorMode::CreateToolTargets(const TArray<TObjectPtr<UObject>>& AssetsIn) 
diff --git a/Engine/Plugins/Experimental/ChaosClothAssetEditor/Source/ChaosClothAssetEditor/Public/ChaosClothAsset/ClothEditorMode.h b/Engine/Plugins/Experimental/ChaosClothAssetEditor/Source/ChaosClothAssetEditor/Public/ChaosClothAsset/ClothEditorMode.h
--- a/Engine/Plugins/Experimental/ChaosClothAssetEditor/Source/ChaosClothAssetEditor/Public/ChaosClothAsset/ClothEditorMode.h
+++ b/Engine/Plugins/Experimental/ChaosClothAssetEditor/Source/ChaosClothAssetEditor/Public/ChaosClothAsset/ClothEditorMode.h
@@ -68,6 +68,7 @@ public:
 	virtual void ModeTick(float DeltaTime) override;
 	virtual void OnToolStarted(UInteractiveToolManager* Manager, UInteractiveTool* Tool) override;
 	virtual void OnToolEnded(UInteractiveToolManager* Manager, UInteractiveTool* Tool) override;
+	virtual void PostUndo() override;
 
 	// We don't actually override MouseEnter, etc, because things get forwarded to the input
 	// router via FEditorModeTools, and we don't have any additional input handling to do at the mode level.
