Cloth Editor: if the rest space viewport is in 3D mode, don't automatically assume we should use the orbit camera -- instead use the default FEditorViewport behavior

diff --git a/Engine/Plugins/Experimental/ChaosClothAssetEditor/Source/ChaosClothAssetEditor/Private/ChaosClothAsset/ClothEditorRestSpaceViewportClient.cpp b/Engine/Plugins/Experimental/ChaosClothAssetEditor/Source/ChaosClothAssetEditor/Private/ChaosClothAsset/ClothEditorRestSpaceViewportClient.cpp
--- a/Engine/Plugins/Experimental/ChaosClothAssetEditor/Source/ChaosClothAssetEditor/Private/ChaosClothAsset/ClothEditorRestSpaceViewportClient.cpp
+++ b/Engine/Plugins/Experimental/ChaosClothAssetEditor/Source/ChaosClothAssetEditor/Private/ChaosClothAsset/ClothEditorRestSpaceViewportClient.cpp
@@ -90,7 +90,14 @@ void FChaosClothEditorRestSpaceViewportClient::AddReferencedObjects(FReferenceCo
 
 bool FChaosClothEditorRestSpaceViewportClient::ShouldOrbitCamera() const
 {
-	return !b2DMode;
+	if (b2DMode)
+	{
+		return false;
+	}
+	else
+	{
+		return FEditorViewportClient::ShouldOrbitCamera();
+	}
 }
 
 bool FChaosClothEditorRestSpaceViewportClient::InputKey(const FInputKeyEventArgs& EventArgs)
