Cloth Editor: when editing the 2D patterns, use the built-in Ortho view rather than the pseudo-ortho view copied from the UV Editor

diff --git a/Engine/Plugins/Experimental/ChaosClothAssetEditor/Source/ChaosClothAssetEditor/Private/ChaosClothAsset/ClothEditorMode.cpp b/Engine/Plugins/Experimental/ChaosClothAssetEditor/Source/ChaosClothAssetEditor/Private/ChaosClothAsset/ClothEditorMode.cpp
--- a/Engine/Plugins/Experimental/ChaosClothAssetEditor/Source/ChaosClothAssetEditor/Private/ChaosClothAsset/ClothEditorMode.cpp
+++ b/Engine/Plugins/Experimental/ChaosClothAssetEditor/Source/ChaosClothAssetEditor/Private/ChaosClothAsset/ClothEditorMode.cpp
@@ -475,6 +475,20 @@ void UChaosClothAssetEditorMode::ReinitializeDynamicMeshComponents()
 
 				TObjectPtr<UDynamicMeshComponent> PatternMeshComponent = NewObject<UDynamicMeshComponent>(ParentActor);
 				PatternMeshComponent->SetMesh(MoveTemp(PatternMesh));
+
+				UMaterialInterface* Material = ToolSetupUtil::GetDefaultSculptMaterial(GetToolManager());
+				UMaterialInstanceDynamic* MatInstance = UMaterialInstanceDynamic::Create(Material, GetToolManager());
+				MatInstance->SetVectorParameterValue(TEXT("Color"), FLinearColor::Gray);
+				MatInstance->TwoSided = 1;
+				PatternMeshComponent->SetMaterial(0, MatInstance);
+
+				if (bPattern2DMode)
+				{
+					// The LVT_OrthoXY viewport transform results in the positive y axis pointing down in screen space. The following transform rotates 
+					// the mesh around the z axis so that increasing y value in the cloth asset corresponds to "up" in screen space
+					PatternMeshComponent->SetWorldRotation(FQuat(FVector(0,0,1), FMathd::Pi));
+				}
+
 				PatternMeshComponent->RegisterComponentWithWorld(this->GetWorld());
 				DynamicMeshComponents[PatternIndex] = PatternMeshComponent;
 				DynamicMeshSourceInfos.Add(FDynamicMeshSourceInfo{ LodIndex, PatternIndex });
@@ -493,7 +507,12 @@ void UChaosClothAssetEditorMode::ReinitializeDynamicMeshComponents()
 	{
 		// Set up the wireframe display of the rest space mesh.
 		TObjectPtr<UMeshElementsVisualizer> WireframeDisplay = NewObject<UMeshElementsVisualizer>(this);
-		WireframeDisplay->CreateInWorld(GetWorld(), FTransform::Identity);
+
+		// The LVT_OrthoXY viewport transform results in the positive y axis pointing down in screen space. The following transform rotates 
+		// the mesh around the z axis so that increasing y value in the cloth asset corresponds to "up" in screen space
+		const FTransform WireframeTransform = bPattern2DMode ? FTransform(FQuat(FVector(0, 0, 1), FMathd::Pi)) : FTransform::Identity;
+
+		WireframeDisplay->CreateInWorld(GetWorld(), WireframeTransform);
 
 		WireframeDisplay->Settings->DepthBias = 2.0;
 		WireframeDisplay->Settings->bAdjustDepthBiasUsingMeshSize = false;
@@ -572,7 +591,7 @@ void UChaosClothAssetEditorMode::RefocusRestSpaceViewportClient()
 		if (IsPattern2DModeActive())
 		{
 			// 2D pattern
-			PinnedVC->SetInitialViewTransform(ELevelViewportType::LVT_Perspective, FVector(0, 150, 100), FRotator(-90, 90, 0), DEFAULT_ORTHOZOOM);
+			PinnedVC->SetInitialViewTransform(ELevelViewportType::LVT_OrthoXY, FVector(0, 0, 0), FRotator(0, 0, 0), DEFAULT_ORTHOZOOM);
 		}
 		else
 		{
diff --git a/Engine/Plugins/Experimental/ChaosClothAssetEditor/Source/ChaosClothAssetEditor/Private/ChaosClothAsset/ClothEditorRestSpaceViewportClient.cpp b/Engine/Plugins/Experimental/ChaosClothAssetEditor/Source/ChaosClothAssetEditor/Private/ChaosClothAsset/ClothEditorRestSpaceViewportClient.cpp
--- a/Engine/Plugins/Experimental/ChaosClothAssetEditor/Source/ChaosClothAssetEditor/Private/ChaosClothAsset/ClothEditorRestSpaceViewportClient.cpp
+++ b/Engine/Plugins/Experimental/ChaosClothAssetEditor/Source/ChaosClothAssetEditor/Private/ChaosClothAsset/ClothEditorRestSpaceViewportClient.cpp
@@ -2,91 +2,41 @@
 
 #include "ChaosClothAsset/ClothEditorRestSpaceViewportClient.h"
 
-#include "BaseBehaviors/ClickDragBehavior.h"
-#include "BaseBehaviors/MouseWheelBehavior.h"
 #include "EngineUtils.h"
 #include "EditorModeManager.h"
 #include "Engine/Selection.h"
-#include "Behaviors/2DViewportBehaviorTargets.h"
 #include "EdModeInteractiveToolsContext.h"
 #include "CameraController.h"
 #include "Components/DynamicMeshComponent.h"
 
-constexpr float CameraFarPlaneWorldZ = -10.0f;
-constexpr float CameraNearPlaneProportionZ = 0.8f;
-constexpr float DefaultNearClipPlaneDistance = KINDA_SMALL_NUMBER;
-constexpr float DefaultFarClipPlaneDistance = UE_FLOAT_HUGE_DISTANCE / 4.0f;
-
 FChaosClothEditorRestSpaceViewportClient::FChaosClothEditorRestSpaceViewportClient(FEditorModeTools* InModeTools,
 	FPreviewScene* InPreviewScene,
 	const TWeakPtr<SEditorViewport>& InEditorViewportWidget) :
 	FEditorViewportClient(InModeTools, InPreviewScene, InEditorViewportWidget)
 {
-	OverrideNearClipPlane(DefaultNearClipPlaneDistance);
-	OverrideFarClipPlane(DefaultFarClipPlaneDistance);
-
-	BehaviorSet = NewObject<UInputBehaviorSet>();
-
-	// We'll have the priority of our viewport manipulation behaviors be lower (i.e. higher
-	// numerically) than both the gizmo default and the tool default.
-	static constexpr int DEFAULT_VIEWPORT_BEHAVIOR_PRIORITY = 150;
-
-	ScrollBehaviorTarget = MakeUnique<FEditor2DScrollBehaviorTarget>(this);
-	UClickDragInputBehavior* ScrollBehavior = NewObject<UClickDragInputBehavior>();
-	ScrollBehavior->Initialize(ScrollBehaviorTarget.Get());
-	ScrollBehavior->SetDefaultPriority(DEFAULT_VIEWPORT_BEHAVIOR_PRIORITY);
-	ScrollBehavior->SetUseRightMouseButton();
-	BehaviorSet->Add(ScrollBehavior);
-	SavedBehaviors.Add(ScrollBehavior);
-
-	ZoomBehaviorTarget = MakeUnique<FEditor2DMouseWheelZoomBehaviorTarget>(this);
-	ZoomBehaviorTarget->SetCameraFarPlaneWorldZ(CameraFarPlaneWorldZ);
-	ZoomBehaviorTarget->SetCameraNearPlaneProportionZ(CameraNearPlaneProportionZ);
-	ZoomBehaviorTarget->SetZoomLimits(0.001, 100000);
-	UMouseWheelInputBehavior* ZoomBehavior = NewObject<UMouseWheelInputBehavior>();
-	ZoomBehavior->Initialize(ZoomBehaviorTarget.Get());
-	ZoomBehavior->SetDefaultPriority(DEFAULT_VIEWPORT_BEHAVIOR_PRIORITY);
-	BehaviorSet->Add(ZoomBehavior);
-	SavedBehaviors.Add(ZoomBehavior);
-
-	ModeTools->GetInteractiveToolsContext()->InputRouter->RegisterSource(this);
+	EngineShowFlags.Grid = true;
+	DrawHelper.bDrawGrid = true;
+
+	bDrawAxes = false;
+
+	// Don't automatically switch to wireframe rendering in ortho mode
+	SetViewModes(EViewModeIndex::VMI_Lit, EViewModeIndex::VMI_Lit);
 }
 
 void FChaosClothEditorRestSpaceViewportClient::Set2DMode(bool In2DMode)
 {
 	b2DMode = In2DMode;
 
-	BehaviorSet->RemoveAll();
-	ModeTools->GetInteractiveToolsContext()->InputRouter->DeregisterSource(this);
-
 	if (b2DMode)
 	{
-		for (UInputBehavior* Behavior : SavedBehaviors)
-		{
-			BehaviorSet->Add(Behavior);
-		}
-		ModeTools->GetInteractiveToolsContext()->InputRouter->RegisterSource(this);
+		SetViewportType(ELevelViewportType::LVT_OrthoXY);
 	}
 	else
 	{
-		OverrideNearClipPlane(DefaultNearClipPlaneDistance);
-		OverrideFarClipPlane(DefaultFarClipPlaneDistance);
+		SetViewportType(ELevelViewportType::LVT_Perspective);
 	}
-
-}
-
-const UInputBehaviorSet* FChaosClothEditorRestSpaceViewportClient::GetInputBehaviors() const
-{
-	return BehaviorSet;
 }
 
-// Collects UObjects that we don't want the garbage collecter to throw away under us
-void FChaosClothEditorRestSpaceViewportClient::AddReferencedObjects(FReferenceCollector& Collector)
-{
-	FEditorViewportClient::AddReferencedObjects(Collector);
-	Collector.AddReferencedObject(BehaviorSet);
-	Collector.AddReferencedObjects(SavedBehaviors);
-}
 
 bool FChaosClothEditorRestSpaceViewportClient::ShouldOrbitCamera() const
 {
@@ -113,83 +63,7 @@ bool FChaosClothEditorRestSpaceViewportClient::InputKey(const FInputKeyEventArgs
 		}
 	}
 
-	if (!b2DMode)
-	{
-		return FEditorViewportClient::InputKey(EventArgs);
-	}
-
-
-	// We'll support disabling input like our base class, even if it does not end up being used.
-	if (bDisableInput)
-	{
-		return true;
-	}
-
-	// Our viewport manipulation is placed in the input router that ModeTools manages
-	if (ModeTools->InputKey(this, EventArgs.Viewport, EventArgs.Key, EventArgs.Event))
-	{
-		return true;
-	}
-	
-	//
-	// Recreate the click handling from FEditorViewportClient::Internal_InputKey so that we can still have 
-	// selection work, but skip the code dealing with camera controls
-	//
-
-	bool bHandled = false;
-
-	const FInputEventState InputState(EventArgs.Viewport, EventArgs.Key, EventArgs.Event);
-	const FViewport* InViewport = EventArgs.Viewport;
-	const EInputEvent Event = EventArgs.Event;
-
-	const bool bWasCursorVisible = InViewport->IsCursorVisible();
-	const bool bWasSoftwareCursorVisible = InViewport->IsSoftwareCursorVisible();
-
-	const bool AltDown = InputState.IsAltButtonPressed();
-	const bool ShiftDown = InputState.IsShiftButtonPressed();
-	const bool ControlDown = InputState.IsCtrlButtonPressed();
-
-	RequiredCursorVisibiltyAndAppearance.bDontResetCursor = false;
-	UpdateRequiredCursorVisibility();
-
-	if (bWasCursorVisible != RequiredCursorVisibiltyAndAppearance.bHardwareCursorVisible || bWasSoftwareCursorVisible != RequiredCursorVisibiltyAndAppearance.bSoftwareCursorVisible)
-	{
-		bHandled = true;
-	}
-
-	// Compute a view.
-	FSceneViewFamilyContext ViewFamily(FSceneViewFamily::ConstructionValues(
-		InViewport,
-		GetScene(),
-		EngineShowFlags)
-		.SetRealtimeUpdate(IsRealtime()));
-	FSceneView* View = CalcSceneView(&ViewFamily);
-
-	// Start tracking if any mouse button is down and it was a tracking event (MouseButton/Ctrl/Shift/Alt):
-	if (InputState.IsAnyMouseButtonDown()
-		&& (Event == IE_Pressed || Event == IE_Released)
-		&& (InputState.IsMouseButtonEvent() || InputState.IsCtrlButtonEvent() || InputState.IsAltButtonEvent() || InputState.IsShiftButtonEvent()))
-	{
-		StartTrackingDueToInput(InputState, *View);
-		return true;
-	}
-
-	// If we are tracking and no mouse button is down and this input event released the mouse button stop tracking and process any clicks if necessary
-	if (bIsTracking && !InputState.IsAnyMouseButtonDown() && InputState.IsMouseButtonEvent())
-	{
-		// Handle possible mouse click viewport
-		ProcessClickInViewport(InputState, *View);
-
-		// Stop tracking if no mouse button is down
-		StopTracking();
-
-		bHandled |= true;
-	}
-
-	// apply the visibility and set the cursor positions
-	ApplyRequiredCursorVisibility(true);
-
-	return bHandled;
+	return FEditorViewportClient::InputKey(EventArgs);
 }
 
 
diff --git a/Engine/Plugins/Experimental/ChaosClothAssetEditor/Source/ChaosClothAssetEditor/Private/ChaosClothAsset/ClothEditorToolkit.cpp b/Engine/Plugins/Experimental/ChaosClothAssetEditor/Source/ChaosClothAssetEditor/Private/ChaosClothAsset/ClothEditorToolkit.cpp
--- a/Engine/Plugins/Experimental/ChaosClothAssetEditor/Source/ChaosClothAssetEditor/Private/ChaosClothAsset/ClothEditorToolkit.cpp
+++ b/Engine/Plugins/Experimental/ChaosClothAssetEditor/Source/ChaosClothAssetEditor/Private/ChaosClothAsset/ClothEditorToolkit.cpp
@@ -346,15 +346,9 @@ void FChaosClothAssetEditorToolkit::PostInitAssetEditor()
 	
 	SetCommonViewportClientOptions(ViewportClient.Get());
 
-	// Ortho has too many problems with rendering things, unfortunately, so we should use perspective.
-	ViewportClient->SetViewportType(ELevelViewportType::LVT_Perspective);
-
 	// Lit gives us the most options in terms of the materials we can use.
 	ViewportClient->SetViewMode(EViewModeIndex::VMI_Lit);
 
-	// If exposure isn't set to fixed, it will flash as we stare into the void
-	ViewportClient->ExposureSettings.bFixed = true;
-
 	UChaosClothAssetEditorMode* ClothMode = Cast<UChaosClothAssetEditorMode>(EditorModeManager->GetActiveScriptableMode(UChaosClothAssetEditorMode::EM_ChaosClothAssetEditorModeId));
 	check(ClothMode);
 	const TWeakPtr<FViewportClient> WeakViewportClient(ViewportClient);
diff --git a/Engine/Plugins/Experimental/ChaosClothAssetEditor/Source/ChaosClothAssetEditor/Public/ChaosClothAsset/ClothEditorRestSpaceViewportClient.h b/Engine/Plugins/Experimental/ChaosClothAssetEditor/Source/ChaosClothAssetEditor/Public/ChaosClothAsset/ClothEditorRestSpaceViewportClient.h
--- a/Engine/Plugins/Experimental/ChaosClothAssetEditor/Source/ChaosClothAssetEditor/Public/ChaosClothAsset/ClothEditorRestSpaceViewportClient.h
+++ b/Engine/Plugins/Experimental/ChaosClothAssetEditor/Source/ChaosClothAssetEditor/Public/ChaosClothAsset/ClothEditorRestSpaceViewportClient.h
@@ -9,7 +9,7 @@
 
 class UInputBehaviorSet;
 
-class CHAOSCLOTHASSETEDITOR_API FChaosClothEditorRestSpaceViewportClient : public FEditorViewportClient, public IInputBehaviorSource
+class CHAOSCLOTHASSETEDITOR_API FChaosClothEditorRestSpaceViewportClient : public FEditorViewportClient
 {
 public:
 
@@ -17,12 +17,7 @@ public:
 
 	virtual ~FChaosClothEditorRestSpaceViewportClient() = default;
 
-	// IInputBehaviorSource
-	virtual const UInputBehaviorSet* GetInputBehaviors() const override;
-
-	// FGCObject
-	virtual void AddReferencedObjects(FReferenceCollector& Collector) override;
-
+	// FEditorViewportClient
 	virtual void ProcessClick(FSceneView& View, HHitProxy* HitProxy, FKey Key, EInputEvent Event, uint32 HitX, uint32 HitY) override;
 	virtual bool InputKey(const FInputKeyEventArgs& EventArgs) override;
 
@@ -37,13 +32,6 @@ private:
 
 	bool b2DMode = false;
 
-	UInputBehaviorSet* BehaviorSet;
-
-	TArray<UInputBehavior*> SavedBehaviors;
-
-	TUniquePtr<FEditor2DScrollBehaviorTarget> ScrollBehaviorTarget;
-	TUniquePtr<FEditor2DMouseWheelZoomBehaviorTarget> ZoomBehaviorTarget;
-
 	TWeakPtr<FUICommandList> ToolCommandList;
 };
 
