Fixed missing includes when building non unity/pch

diff --git a/Engine/Plugins/Editor/ObjectMixer/ObjectMixer/Source/ObjectMixer/Public/Views/Widgets/ObjectMixerEditorListMenuContext.h b/Engine/Plugins/Editor/ObjectMixer/ObjectMixer/Source/ObjectMixer/Public/Views/Widgets/ObjectMixerEditorListMenuContext.h
--- a/Engine/Plugins/Editor/ObjectMixer/ObjectMixer/Source/ObjectMixer/Public/Views/Widgets/ObjectMixerEditorListMenuContext.h
+++ b/Engine/Plugins/Editor/ObjectMixer/ObjectMixer/Source/ObjectMixer/Public/Views/Widgets/ObjectMixerEditorListMenuContext.h
@@ -8,6 +8,8 @@
 
 #include "ObjectMixerEditorListMenuContext.generated.h"	
 
+class UToolMenu;
+
 UCLASS()
 class OBJECTMIXEREDITOR_API UObjectMixerEditorListMenuContext : public UObject
 {
diff --git a/Engine/Plugins/Experimental/ChaosClothAsset/Source/ChaosClothAssetEngine/Private/ChaosClothAsset/ClothSimulationContext.cpp b/Engine/Plugins/Experimental/ChaosClothAsset/Source/ChaosClothAssetEngine/Private/ChaosClothAsset/ClothSimulationContext.cpp
--- a/Engine/Plugins/Experimental/ChaosClothAsset/Source/ChaosClothAssetEngine/Private/ChaosClothAsset/ClothSimulationContext.cpp
+++ b/Engine/Plugins/Experimental/ChaosClothAsset/Source/ChaosClothAssetEngine/Private/ChaosClothAsset/ClothSimulationContext.cpp
@@ -3,7 +3,9 @@
 #include "ChaosClothAsset/ClothSimulationContext.h"
 #include "ChaosClothAsset/ClothAsset.h"
 #include "ChaosClothAsset/ClothComponent.h"
+#include "Engine/World.h"
 #include "HAL/IConsoleManager.h"
+#include "SceneInterface.h"
 
 namespace UE::Chaos::ClothAsset
 {
diff --git a/Engine/Plugins/Experimental/ChaosClothAssetEditor/Source/ChaosClothAssetEditor/Private/ChaosClothAsset/ClothEditorRestSpaceViewportClient.cpp b/Engine/Plugins/Experimental/ChaosClothAssetEditor/Source/ChaosClothAssetEditor/Private/ChaosClothAsset/ClothEditorRestSpaceViewportClient.cpp
--- a/Engine/Plugins/Experimental/ChaosClothAssetEditor/Source/ChaosClothAssetEditor/Private/ChaosClothAsset/ClothEditorRestSpaceViewportClient.cpp
+++ b/Engine/Plugins/Experimental/ChaosClothAssetEditor/Source/ChaosClothAssetEditor/Private/ChaosClothAsset/ClothEditorRestSpaceViewportClient.cpp
@@ -8,6 +8,7 @@
 #include "EdModeInteractiveToolsContext.h"
 #include "CameraController.h"
 #include "Components/DynamicMeshComponent.h"
+#include "Framework/Application/SlateApplication.h"
 
 FChaosClothEditorRestSpaceViewportClient::FChaosClothEditorRestSpaceViewportClient(FEditorModeTools* InModeTools,
 	FPreviewScene* InPreviewScene,
diff --git a/Engine/Plugins/Experimental/ChaosClothAssetEditor/Source/ChaosClothAssetTools/Private/ChaosClothAsset/ClothAssetBuilderEditor.cpp b/Engine/Plugins/Experimental/ChaosClothAssetEditor/Source/ChaosClothAssetTools/Private/ChaosClothAsset/ClothAssetBuilderEditor.cpp
--- a/Engine/Plugins/Experimental/ChaosClothAssetEditor/Source/ChaosClothAssetTools/Private/ChaosClothAsset/ClothAssetBuilderEditor.cpp
+++ b/Engine/Plugins/Experimental/ChaosClothAssetEditor/Source/ChaosClothAssetTools/Private/ChaosClothAsset/ClothAssetBuilderEditor.cpp
@@ -3,6 +3,7 @@
 #include "ChaosClothAsset/ClothAssetBuilderEditor.h"
 #include "ChaosClothAsset/ClothAsset.h"
 #include "ChaosClothAsset/ClothCollection.h"
+#include "Engine/SkeletalMesh.h"
 #include "Modules/ModuleManager.h"
 #include "Rendering/SkeletalMeshLODModel.h"
 #include "Utils/ClothingMeshUtils.h"
diff --git a/Engine/Plugins/FX/Niagara/Source/Niagara/Private/NiagaraDataInterfaceTexture.cpp b/Engine/Plugins/FX/Niagara/Source/Niagara/Private/NiagaraDataInterfaceTexture.cpp
--- a/Engine/Plugins/FX/Niagara/Source/Niagara/Private/NiagaraDataInterfaceTexture.cpp
+++ b/Engine/Plugins/FX/Niagara/Source/Niagara/Private/NiagaraDataInterfaceTexture.cpp
@@ -2,6 +2,7 @@
 
 #include "NiagaraDataInterfaceTexture.h"
 #include "Engine/Texture.h"
+#include "Engine/Texture2D.h"
 #include "Engine/TextureRenderTarget2D.h"
 #include "NiagaraComputeExecutionContext.h"
 #include "NiagaraCustomVersion.h"
diff --git a/Engine/Plugins/Runtime/HairStrands/Source/HairStrandsCore/Private/Deformer/DeformerDataInterfaceGroom.cpp b/Engine/Plugins/Runtime/HairStrands/Source/HairStrandsCore/Private/Deformer/DeformerDataInterfaceGroom.cpp
--- a/Engine/Plugins/Runtime/HairStrands/Source/HairStrandsCore/Private/Deformer/DeformerDataInterfaceGroom.cpp
+++ b/Engine/Plugins/Runtime/HairStrands/Source/HairStrandsCore/Private/Deformer/DeformerDataInterfaceGroom.cpp
@@ -5,6 +5,7 @@
 #include "ComputeFramework/ShaderParamTypeDefinition.h"
 #include "ComputeFramework/ComputeKernelPermutationVector.h"
 #include "OptimusDataDomain.h"
+#include "ShaderCompilerCore.h"
 #include "ShaderParameterMetadataBuilder.h"
 #include "GroomComponent.h"
 #include "GroomInstance.h"
diff --git a/Engine/Plugins/Runtime/HairStrands/Source/HairStrandsCore/Private/Deformer/DeformerDataInterfaceGroomExec.cpp b/Engine/Plugins/Runtime/HairStrands/Source/HairStrandsCore/Private/Deformer/DeformerDataInterfaceGroomExec.cpp
--- a/Engine/Plugins/Runtime/HairStrands/Source/HairStrandsCore/Private/Deformer/DeformerDataInterfaceGroomExec.cpp
+++ b/Engine/Plugins/Runtime/HairStrands/Source/HairStrandsCore/Private/Deformer/DeformerDataInterfaceGroomExec.cpp
@@ -4,6 +4,7 @@
 
 #include "GroomComponent.h"
 #include "ComputeFramework/ShaderParamTypeDefinition.h"
+#include "ShaderCompilerCore.h"
 #include "ShaderParameterMetadataBuilder.h"
 #include "SkeletalRenderPublic.h"
 
diff --git a/Engine/Plugins/Runtime/MeshModelingToolset/Source/ModelingComponents/Private/Selection/MeshTopologySelectionMechanic.cpp b/Engine/Plugins/Runtime/MeshModelingToolset/Source/ModelingComponents/Private/Selection/MeshTopologySelectionMechanic.cpp
--- a/Engine/Plugins/Runtime/MeshModelingToolset/Source/ModelingComponents/Private/Selection/MeshTopologySelectionMechanic.cpp
+++ b/Engine/Plugins/Runtime/MeshModelingToolset/Source/ModelingComponents/Private/Selection/MeshTopologySelectionMechanic.cpp
@@ -4,6 +4,7 @@
 
 #include "BaseBehaviors/SingleClickOrDragBehavior.h"
 #include "BaseBehaviors/MouseHoverBehavior.h"
+#include "Engine/World.h"
 #include "InteractiveToolManager.h"
 #include "Util/ColorConstants.h"
 #include "Selection/PersistentMeshSelection.h"
diff --git a/Engine/Plugins/TypedElementsDataStorage/Source/TypedElementsDataStorage/Private/TypedElementsDataStorage.cpp b/Engine/Plugins/TypedElementsDataStorage/Source/TypedElementsDataStorage/Private/TypedElementsDataStorage.cpp
--- a/Engine/Plugins/TypedElementsDataStorage/Source/TypedElementsDataStorage/Private/TypedElementsDataStorage.cpp
+++ b/Engine/Plugins/TypedElementsDataStorage/Source/TypedElementsDataStorage/Private/TypedElementsDataStorage.cpp
@@ -3,6 +3,7 @@
 #include "TypedElementsDataStorage.h"
 
 #include "Elements/Framework/TypedElementRegistry.h"
+#include "Misc/CoreDelegates.h"
 #include "TypedElementDatabase.h"
 #include "TypedElementDatabaseCompatibility.h"
 #include "UObject/UObjectGlobals.h"
diff --git a/Engine/Source/Editor/UnrealEd/Classes/Editor/PropertyEditorTestObject.h b/Engine/Source/Editor/UnrealEd/Classes/Editor/PropertyEditorTestObject.h
--- a/Engine/Source/Editor/UnrealEd/Classes/Editor/PropertyEditorTestObject.h
+++ b/Engine/Source/Editor/UnrealEd/Classes/Editor/PropertyEditorTestObject.h
@@ -3,6 +3,7 @@
 #pragma once 
 
 #include "CoreMinimal.h"
+#include "Misc/Timecode.h"
 #include "UObject/ObjectMacros.h"
 #include "UObject/Object.h"
 #include "UObject/PrimaryAssetId.h"
diff --git a/Engine/Source/Runtime/Engine/Classes/SparseVolumeTexture/SparseVolumeTexture.h b/Engine/Source/Runtime/Engine/Classes/SparseVolumeTexture/SparseVolumeTexture.h
--- a/Engine/Source/Runtime/Engine/Classes/SparseVolumeTexture/SparseVolumeTexture.h
+++ b/Engine/Source/Runtime/Engine/Classes/SparseVolumeTexture/SparseVolumeTexture.h
@@ -8,10 +8,12 @@
 
 #include "Serialization/EditorBulkData.h"
 #include "Containers/Array.h"
+#include "UnrealClient.h"
 #include "UObject/ObjectSaveContext.h"
 
 #include "SparseVolumeTexture.generated.h"
 
+namespace UE { namespace Shader	{ enum class EValueType : uint8; } }
 
 #define SPARSE_VOLUME_TILE_RES	16
 
diff --git a/Engine/Source/Runtime/Engine/Private/Animation/AnimData/IAnimationDataModel.cpp b/Engine/Source/Runtime/Engine/Private/Animation/AnimData/IAnimationDataModel.cpp
--- a/Engine/Source/Runtime/Engine/Private/Animation/AnimData/IAnimationDataModel.cpp
+++ b/Engine/Source/Runtime/Engine/Private/Animation/AnimData/IAnimationDataModel.cpp
@@ -2,6 +2,7 @@
 
 #include "Animation/AnimData/IAnimationDataModel.h"
 #include "Animation/AnimData/AnimDataModel.h"
+#include "Features/IModularFeatures.h"
 
 namespace UE::Anim::DataModel
 {
diff --git a/Engine/Source/Runtime/Engine/Private/SparseVolumeTexture/SparseVolumeTexture.cpp b/Engine/Source/Runtime/Engine/Private/SparseVolumeTexture/SparseVolumeTexture.cpp
--- a/Engine/Source/Runtime/Engine/Private/SparseVolumeTexture/SparseVolumeTexture.cpp
+++ b/Engine/Source/Runtime/Engine/Private/SparseVolumeTexture/SparseVolumeTexture.cpp
@@ -8,6 +8,7 @@
 
 #include "Misc/SecureHash.h"
 #include "EngineUtils.h"
+#include "Shader/ShaderTypes.h"
 
 #if WITH_EDITORONLY_DATA
 #include "DerivedDataCache/Public/DerivedDataCacheInterface.h"
diff --git a/Engine/Source/Runtime/Engine/Private/Streaming/StreamingTexture.cpp b/Engine/Source/Runtime/Engine/Private/Streaming/StreamingTexture.cpp
--- a/Engine/Source/Runtime/Engine/Private/Streaming/StreamingTexture.cpp
+++ b/Engine/Source/Runtime/Engine/Private/Streaming/StreamingTexture.cpp
@@ -5,6 +5,7 @@ StreamingTexture.cpp: Definitions of classes used for texture.
 =============================================================================*/
 
 #include "Streaming/StreamingTexture.h"
+#include "Engine/Texture.h"
 #include "Misc/App.h"
 #include "Streaming/StreamingManagerTexture.h"
 #include "HAL/FileManager.h"
