Cloth Editor:
- Adding a stubbed out ClothTrainingTool for ML cloth training. Actual training part to be filled in later
- Adding a ClothComponentToolTarget for tools operating on ChaosClothComponents

diff --git a/Engine/Plugins/Experimental/ChaosClothAssetEditor/Source/ChaosClothAssetEditor/Private/ChaosClothAsset/ClothEditorCommands.cpp b/Engine/Plugins/Experimental/ChaosClothAssetEditor/Source/ChaosClothAssetEditor/Private/ChaosClothAsset/ClothEditorCommands.cpp
--- a/Engine/Plugins/Experimental/ChaosClothAssetEditor/Source/ChaosClothAssetEditor/Private/ChaosClothAsset/ClothEditorCommands.cpp
+++ b/Engine/Plugins/Experimental/ChaosClothAssetEditor/Source/ChaosClothAssetEditor/Private/ChaosClothAsset/ClothEditorCommands.cpp
@@ -9,6 +9,7 @@
 const FString FChaosClothAssetEditorCommands::BeginRemeshToolIdentifier = TEXT("BeginRemeshTool");
 const FString FChaosClothAssetEditorCommands::BeginAttributeEditorToolIdentifier = TEXT("BeginAttributeEditorTool");
 const FString FChaosClothAssetEditorCommands::BeginWeightMapPaintToolIdentifier = TEXT("BeginWeightMapPaintTool");
+const FString FChaosClothAssetEditorCommands::BeginClothTrainingToolIdentifier = TEXT("BeginClothTrainingTool");
 const FString FChaosClothAssetEditorCommands::ToggleSimMeshWireframeIdentifier = TEXT("ToggleSimMeshWireframe");
 const FString FChaosClothAssetEditorCommands::ToggleRenderMeshWireframeIdentifier = TEXT("ToggleRenderMeshWireframe");
 const FString FChaosClothAssetEditorCommands::TogglePatternModeIdentifier = TEXT("TogglePatternMode");
@@ -33,6 +34,8 @@ void FChaosClothAssetEditorCommands::RegisterCommands()
 	UI_COMMAND(BeginAttributeEditorTool, "AttrEd", "Edit/configure mesh attributes", EUserInterfaceActionType::Button, FInputChord());
 	UI_COMMAND(BeginWeightMapPaintTool, "MapPnt", "Paint Weight Maps on the mesh", EUserInterfaceActionType::Button, FInputChord());
 
+	UI_COMMAND(BeginClothTrainingTool, "Train", "Launch Cloth Training tool", EUserInterfaceActionType::Button, FInputChord());
+
 	UI_COMMAND(TogglePatternMode, "TogglePatternMode", "Toggle pattern mode", EUserInterfaceActionType::Button, FInputChord());
 
 	UI_COMMAND(ToggleSimMeshWireframe, "ToggleSimMeshWireframe", "Toggle simulation mesh wireframe", EUserInterfaceActionType::Button, FInputChord());
diff --git a/Engine/Plugins/Experimental/ChaosClothAssetEditor/Source/ChaosClothAssetEditor/Private/ChaosClothAsset/ClothEditorMode.cpp b/Engine/Plugins/Experimental/ChaosClothAssetEditor/Source/ChaosClothAssetEditor/Private/ChaosClothAsset/ClothEditorMode.cpp
--- a/Engine/Plugins/Experimental/ChaosClothAssetEditor/Source/ChaosClothAssetEditor/Private/ChaosClothAsset/ClothEditorMode.cpp
+++ b/Engine/Plugins/Experimental/ChaosClothAssetEditor/Source/ChaosClothAssetEditor/Private/ChaosClothAsset/ClothEditorMode.cpp
@@ -24,10 +24,8 @@
 #include "ToolTargetManager.h"
 #include "Toolkits/BaseToolkit.h"
 #include "ToolTargets/ToolTarget.h"
-#include "ToolTargets/StaticMeshComponentToolTarget.h"
-#include "ToolTargets/VolumeComponentToolTarget.h"
 #include "ToolTargets/DynamicMeshComponentToolTarget.h"
-#include "ToolTargets/SkeletalMeshComponentToolTarget.h"
+#include "ChaosClothAsset/ClothComponentToolTarget.h"
 #include "Engine/Selection.h"
 #include "Misc/ConfigCacheIni.h"
 #include "Modules/ModuleManager.h"
@@ -35,6 +33,7 @@
 #include "AttributeEditorTool.h"
 #include "MeshAttributePaintTool.h"
 #include "ChaosClothAsset/ClothWeightMapPaintTool.h"
+#include "ChaosClothAsset/ClothTrainingTool.h"
 #include "DynamicMesh/DynamicMeshAttributeSet.h"
 #include "DynamicMesh/DynamicVertexSkinWeightsAttribute.h"
 #include "DynamicMeshEditor.h"
@@ -141,9 +140,8 @@ void UChaosClothAssetEditorMode::Enter()
 
 void UChaosClothAssetEditorMode::AddToolTargetFactories()
 {
-	GetInteractiveToolsContext()->TargetManager->AddTargetFactory(NewObject<UStaticMeshComponentToolTargetFactory>(GetToolManager()));
-	GetInteractiveToolsContext()->TargetManager->AddTargetFactory(NewObject<UVolumeComponentToolTargetFactory>(GetToolManager()));
 	GetInteractiveToolsContext()->TargetManager->AddTargetFactory(NewObject<UDynamicMeshComponentToolTargetFactory>(GetToolManager()));
+	GetInteractiveToolsContext()->TargetManager->AddTargetFactory(NewObject<UClothComponentToolTargetFactory>(GetToolManager()));
 }
 
 void UChaosClothAssetEditorMode::RegisterTools()
@@ -157,6 +155,7 @@ void UChaosClothAssetEditorMode::RegisterTools()
 	RegisterTool(CommandInfos.BeginRemeshTool, FChaosClothAssetEditorCommands::BeginRemeshToolIdentifier, NewObject<URemeshMeshToolBuilder>());
 	RegisterTool(CommandInfos.BeginAttributeEditorTool, FChaosClothAssetEditorCommands::BeginAttributeEditorToolIdentifier, NewObject<UAttributeEditorToolBuilder>());
 	RegisterTool(CommandInfos.BeginWeightMapPaintTool, FChaosClothAssetEditorCommands::BeginWeightMapPaintToolIdentifier, NewObject<UClothEditorWeightMapPaintToolBuilder>());
+	RegisterTool(CommandInfos.BeginClothTrainingTool, FChaosClothAssetEditorCommands::BeginClothTrainingToolIdentifier, NewObject<UClothTrainingToolBuilder>());
 }
 
 bool UChaosClothAssetEditorMode::ShouldToolStartBeAllowed(const FString& ToolIdentifier) const
@@ -630,6 +629,10 @@ void UChaosClothAssetEditorMode::InitializeTargets(const TArray<TObjectPtr<UObje
 			ClothComponent->SetClothAsset(ChaosClothAsset);
 			ClothComponent->RegisterComponentWithWorld(PreviewWorld);
 
+			USelection* SelectedComponents = GetModeManager()->GetSelectedComponents();
+			SelectedComponents->Modify();
+			SelectedComponents->Select(ClothComponent);
+
 			break;
 		}
 	}
@@ -692,10 +695,11 @@ void UChaosClothAssetEditorMode::ModeTick(float DeltaTime)
 		}
 	}
 
-	if (PreviewWorld)
-	{
-		PreviewWorld->Tick(ELevelTick::LEVELTICK_All, DeltaTime);
-	}
+	// TODO: Uncomment this to enable simulation stepping
+	//if (PreviewWorld)
+	//{
+	//	PreviewWorld->Tick(ELevelTick::LEVELTICK_All, DeltaTime);
+	//}
 }
 
 
diff --git a/Engine/Plugins/Experimental/ChaosClothAssetEditor/Source/ChaosClothAssetEditor/Private/ChaosClothAsset/ClothEditorModeToolkit.cpp b/Engine/Plugins/Experimental/ChaosClothAssetEditor/Source/ChaosClothAssetEditor/Private/ChaosClothAsset/ClothEditorModeToolkit.cpp
--- a/Engine/Plugins/Experimental/ChaosClothAssetEditor/Source/ChaosClothAssetEditor/Private/ChaosClothAsset/ClothEditorModeToolkit.cpp
+++ b/Engine/Plugins/Experimental/ChaosClothAssetEditor/Source/ChaosClothAssetEditor/Private/ChaosClothAsset/ClothEditorModeToolkit.cpp
@@ -21,11 +21,14 @@ void FChaosClothAssetEditorModeToolkit::BuildToolPalette(FName PaletteIndex, cla
 {
 	const FChaosClothAssetEditorCommands& Commands = FChaosClothAssetEditorCommands::Get();
 
+	// TODO: make separate palettes to split up the tools tab
 	if (PaletteIndex == FBaseCharacterFXEditorModeToolkit::ToolsTabName)
 	{
 		ToolbarBuilder.AddToolBarButton(Commands.BeginRemeshTool);
 		ToolbarBuilder.AddToolBarButton(Commands.BeginAttributeEditorTool);
 		ToolbarBuilder.AddToolBarButton(Commands.BeginWeightMapPaintTool);
+
+		ToolbarBuilder.AddToolBarButton(Commands.BeginClothTrainingTool);
 	}
 }
 
diff --git a/Engine/Plugins/Experimental/ChaosClothAssetEditor/Source/ChaosClothAssetEditor/Private/ChaosClothAsset/ClothEditorRestSpaceViewportClient.cpp b/Engine/Plugins/Experimental/ChaosClothAssetEditor/Source/ChaosClothAssetEditor/Private/ChaosClothAsset/ClothEditorRestSpaceViewportClient.cpp
--- a/Engine/Plugins/Experimental/ChaosClothAssetEditor/Source/ChaosClothAssetEditor/Private/ChaosClothAsset/ClothEditorRestSpaceViewportClient.cpp
+++ b/Engine/Plugins/Experimental/ChaosClothAssetEditor/Source/ChaosClothAssetEditor/Private/ChaosClothAsset/ClothEditorRestSpaceViewportClient.cpp
@@ -76,7 +76,7 @@ void FChaosClothEditorRestSpaceViewportClient::ProcessClick(FSceneView& View, HH
 	USelection* SelectedComponents = ModeTools->GetSelectedComponents();
 	SelectedComponents->Modify();
 	SelectedComponents->BeginBatchSelectOperation();
-	SelectedComponents->DeselectAll();
+	SelectedComponents->DeselectAll(UDynamicMeshComponent::StaticClass());
 
 	if (HitProxy && HitProxy->IsA(HActor::StaticGetType()))
 	{
diff --git a/Engine/Plugins/Experimental/ChaosClothAssetEditor/Source/ChaosClothAssetEditor/Private/ChaosClothAsset/ClothEditorStyle.cpp b/Engine/Plugins/Experimental/ChaosClothAssetEditor/Source/ChaosClothAssetEditor/Private/ChaosClothAsset/ClothEditorStyle.cpp
--- a/Engine/Plugins/Experimental/ChaosClothAssetEditor/Source/ChaosClothAssetEditor/Private/ChaosClothAsset/ClothEditorStyle.cpp
+++ b/Engine/Plugins/Experimental/ChaosClothAssetEditor/Source/ChaosClothAssetEditor/Private/ChaosClothAsset/ClothEditorStyle.cpp
@@ -48,6 +48,9 @@ FChaosClothAssetEditorStyle::FChaosClothAssetEditorStyle()
 	PropertyNameString = "ChaosClothAssetEditor." + FChaosClothAssetEditorCommands::BeginWeightMapPaintToolIdentifier;
 	Set(*PropertyNameString, new FSlateImageBrush(FChaosClothAssetEditorStyle::InContent("Icons/ModelingAttributePaint_x40", ".png"), ToolbarIconSize));
 
+	PropertyNameString = "ChaosClothAssetEditor." + FChaosClothAssetEditorCommands::BeginClothTrainingToolIdentifier;
+	Set(*PropertyNameString, new FSlateImageBrush(FChaosClothAssetEditorStyle::InContent("Icons/ModelingAttributePaint_x40", ".png"), ToolbarIconSize));
+
 	PropertyNameString = "ChaosClothAssetEditor." + FChaosClothAssetEditorCommands::ToggleSimMeshWireframeIdentifier;
 	Set(*PropertyNameString, new FSlateImageBrush(FChaosClothAssetEditorStyle::InContent("Icons/icon_ViewMode_BrushWireframe_16px", ".png"), ViewportToolbarIconSize));		// Placeholder
 
diff --git a/Engine/Plugins/Experimental/ChaosClothAssetEditor/Source/ChaosClothAssetEditor/Public/ChaosClothAsset/ClothEditorCommands.h b/Engine/Plugins/Experimental/ChaosClothAssetEditor/Source/ChaosClothAssetEditor/Public/ChaosClothAsset/ClothEditorCommands.h
--- a/Engine/Plugins/Experimental/ChaosClothAssetEditor/Source/ChaosClothAssetEditor/Public/ChaosClothAsset/ClothEditorCommands.h
+++ b/Engine/Plugins/Experimental/ChaosClothAssetEditor/Source/ChaosClothAssetEditor/Public/ChaosClothAsset/ClothEditorCommands.h
@@ -36,6 +36,9 @@ public:
 	const static FString BeginAttributeEditorToolIdentifier;
 	TSharedPtr<FUICommandInfo> BeginAttributeEditorTool;
 
+	const static FString BeginClothTrainingToolIdentifier;
+	TSharedPtr<FUICommandInfo> BeginClothTrainingTool;
+
 	// Rest space viewport commands
 	const static FString TogglePatternModeIdentifier;
 	TSharedPtr<FUICommandInfo> TogglePatternMode;
diff --git a/Engine/Plugins/Experimental/ChaosClothAssetEditor/Source/ChaosClothAssetEditorTools/ChaosClothAssetEditorTools.Build.cs b/Engine/Plugins/Experimental/ChaosClothAssetEditor/Source/ChaosClothAssetEditorTools/ChaosClothAssetEditorTools.Build.cs
--- a/Engine/Plugins/Experimental/ChaosClothAssetEditor/Source/ChaosClothAssetEditorTools/ChaosClothAssetEditorTools.Build.cs
+++ b/Engine/Plugins/Experimental/ChaosClothAssetEditor/Source/ChaosClothAssetEditorTools/ChaosClothAssetEditorTools.Build.cs
@@ -48,6 +48,7 @@ public class ChaosClothAssetEditorTools : ModuleRules
 				"ClothingSystemRuntimeInterface",
 				"ChaosCloth",
 				"ChaosClothAssetTools",
+				"ChaosClothAssetEngine",
 				"ClothingSystemEditor",
 				"MeshConversion"
 				// ... add private dependencies that you statically link with here ...	
diff --git a/Engine/Plugins/Experimental/ChaosClothAssetEditor/Source/ChaosClothAssetEditorTools/Private/ChaosClothAsset/ClothComponentToolTarget.cpp b/Engine/Plugins/Experimental/ChaosClothAssetEditor/Source/ChaosClothAssetEditorTools/Private/ChaosClothAsset/ClothComponentToolTarget.cpp
new file mode 100644
--- /dev/null
+++ b/Engine/Plugins/Experimental/ChaosClothAssetEditor/Source/ChaosClothAssetEditorTools/Private/ChaosClothAsset/ClothComponentToolTarget.cpp
@@ -0,0 +1,62 @@
+// Copyright Epic Games, Inc. All Rights Reserved.
+
+#include "ChaosClothAsset/ClothComponentToolTarget.h"
+
+#include "ChaosClothAsset/ClothAsset.h"
+#include "ChaosClothAsset/ClothComponent.h"
+#include "ComponentReregisterContext.h"
+#include "Materials/Material.h"
+
+#include UE_INLINE_GENERATED_CPP_BY_NAME(ClothComponentToolTarget)
+
+
+bool UClothComponentToolTarget::IsValid() const
+{
+	if (!UPrimitiveComponentToolTarget::IsValid())
+	{
+		return false;
+	}
+
+	const UChaosClothComponent* const ClothComponent = Cast<UChaosClothComponent>(Component);
+	if (ClothComponent == nullptr)
+	{
+		return false;
+	}
+
+	const UChaosClothAsset* const ClothAsset = ClothComponent->GetClothAsset();
+	if (!ClothAsset || !IsValidChecked(ClothAsset) || ClothAsset->IsUnreachable() || !ClothAsset->IsValidLowLevel())
+	{
+		return false;
+	}
+
+	return true;
+}
+
+UChaosClothAsset* UClothComponentToolTarget::GetClothAsset() const
+{
+	return IsValid() ? Cast<UChaosClothComponent>(Component)->GetClothAsset() : nullptr;
+}
+
+UChaosClothComponent* UClothComponentToolTarget::GetClothComponent() const
+{
+	return IsValid() ? Cast<UChaosClothComponent>(Component) : nullptr;
+}
+
+
+// Factory
+
+bool UClothComponentToolTargetFactory::CanBuildTarget(UObject* SourceObject, const FToolTargetTypeRequirements& Requirements) const
+{
+	const UChaosClothComponent* const Component = GetValid(Cast<UChaosClothComponent>(SourceObject));
+	return Component && !Component->IsUnreachable() && Component->IsValidLowLevel() && Requirements.AreSatisfiedBy(UClothComponentToolTarget::StaticClass());
+}
+
+UToolTarget* UClothComponentToolTargetFactory::BuildTarget(UObject* SourceObject, const FToolTargetTypeRequirements& Requirements)
+{
+	UClothComponentToolTarget* const Target = NewObject<UClothComponentToolTarget>(); // TODO: Should we set an outer here?
+	Target->Component = Cast<UChaosClothComponent>(SourceObject);
+
+	ensure(Target->Component && Requirements.AreSatisfiedBy(Target));
+
+	return Target;
+}
diff --git a/Engine/Plugins/Experimental/ChaosClothAssetEditor/Source/ChaosClothAssetEditorTools/Private/ChaosClothAsset/ClothTrainingTool.cpp b/Engine/Plugins/Experimental/ChaosClothAssetEditor/Source/ChaosClothAssetEditorTools/Private/ChaosClothAsset/ClothTrainingTool.cpp
new file mode 100644
--- /dev/null
+++ b/Engine/Plugins/Experimental/ChaosClothAssetEditor/Source/ChaosClothAssetEditorTools/Private/ChaosClothAsset/ClothTrainingTool.cpp
@@ -0,0 +1,202 @@
+// Copyright Epic Games, Inc. All Rights Reserved.
+
+#include "ChaosClothAsset/ClothTrainingTool.h"
+#include "CoreMinimal.h"
+#include "ChaosClothAsset/ClothComponent.h"
+#include "ModelingOperators.h"
+#include "Misc/ScopedSlowTask.h"
+#include "ContextObjectStore.h"
+#include "InteractiveToolManager.h"
+#include "ToolTargetManager.h"
+#include "ChaosClothAsset/ClothComponentToolTarget.h"
+
+#define LOCTEXT_NAMESPACE "ClothTrainingTool"
+
+// ------------------- Operator -------------------
+
+class FTrainClothOp : public UE::Geometry::TGenericDataOperator<FSkinnedMeshVertices>
+{
+public: 
+
+	TObjectPtr<UAnimSequence> AnimSequence;
+
+	virtual void CalculateResult(FProgressCancel* Progress) override;
+};
+
+
+void FTrainClothOp::CalculateResult(FProgressCancel* Progress)
+{
+	// @@@@@@@@@ TODO: Do the posing, simulation, getting outputs, etc. Check for Progress->Cancelled() regularly if possible.
+
+	FSkinnedMeshVertices OutputVertices;
+	
+	bool bCancelled = false;
+
+	for (int i = 0; i < 10; ++i)
+	{
+		FPlatformProcess::Sleep(1.0f);
+
+		if (Progress)
+		{
+			if (Progress->Cancelled())
+			{
+				bCancelled = true;
+				break;
+			}
+
+			OutputVertices.Vertices.Add(FVector3f(i,i,i));
+
+			Progress->AdvanceCurrentScopeProgressBy(0.1f);
+		}
+	}
+
+	if (!bCancelled)
+	{
+		SetResult(MakeUnique<FSkinnedMeshVertices>(OutputVertices));
+	}
+}
+
+
+// ------------------- Properties -------------------
+
+void UClothTrainingToolActionProperties::PostAction(EClothTrainingToolActions Action)
+{
+	if (ParentTool.IsValid())
+	{
+		ParentTool->RequestAction(Action);
+	}
+}
+
+
+// ------------------- Builder -------------------
+
+const FToolTargetTypeRequirements& UClothTrainingToolBuilder::GetTargetRequirements() const
+{
+	static FToolTargetTypeRequirements TypeRequirements({ 
+		UPrimitiveComponentBackedTarget::StaticClass(),
+		UClothAssetBackedTarget::StaticClass()
+		});
+	return TypeRequirements;
+}
+
+bool UClothTrainingToolBuilder::CanBuildTool(const FToolBuilderState& SceneState) const
+{
+	return (SceneState.TargetManager->CountSelectedAndTargetable(SceneState, GetTargetRequirements()) == 1);
+}
+
+UInteractiveTool* UClothTrainingToolBuilder::BuildTool(const FToolBuilderState& SceneState) const
+{
+	UClothTrainingTool* NewTool = NewObject<UClothTrainingTool>();
+	
+	UToolTarget* Target = SceneState.TargetManager->BuildFirstSelectedTargetable(SceneState, GetTargetRequirements());
+	NewTool->SetTarget(Target);
+
+	return NewTool;
+}
+
+// ------------------- Tool -------------------
+
+
+void UClothTrainingTool::Setup()
+{
+	UInteractiveTool::Setup();
+
+	if (const UClothComponentToolTarget* const ClothComponentTarget = Cast<UClothComponentToolTarget>(Target))
+	{
+		const UChaosClothComponent* const TargetClothComponent = ClothComponentTarget->GetClothComponent();
+		ClothComponent = TargetClothComponent;
+	}
+
+	ToolProperties = NewObject<UClothTrainingToolProperties>(this);
+	AddToolPropertySource(ToolProperties);
+
+	ActionProperties = NewObject<UClothTrainingToolActionProperties>(this);
+	ActionProperties->ParentTool = this;
+	AddToolPropertySource(ActionProperties);
+}
+
+void UClothTrainingTool::RunTraining()
+{
+	const FText DefaultMessage(LOCTEXT("ClothTrainingMessage", "Training..."));
+	
+	using FTaskType = UE::Geometry::TModelingOpTask<FTrainClothOp>;
+	using FExecuterType = UE::Geometry::FAsyncTaskExecuterWithProgressCancel<FTaskType>;
+
+	TUniquePtr<FTrainClothOp> NewOp = MakeUnique<FTrainClothOp>();
+	NewOp->AnimSequence = ToolProperties->AnimationSequence;
+	// @@@@@@@@@@@ TODO: Add anything else needed by the Operator
+
+
+	TUniquePtr<FExecuterType> BackgroundTaskExecuter = MakeUnique<FExecuterType>(MoveTemp(NewOp));
+	BackgroundTaskExecuter->StartBackgroundTask();
+
+	FScopedSlowTask SlowTask(1, DefaultMessage);
+	SlowTask.MakeDialog(true);
+
+	bool bSuccess = false;
+	while (true)
+	{
+		if (SlowTask.ShouldCancel())
+		{
+			// Release ownership to the TDeleterTask that is spawned by CancelAndDelete()
+			BackgroundTaskExecuter.Release()->CancelAndDelete();
+			break;
+		}
+		if (BackgroundTaskExecuter->IsDone())
+		{
+			bSuccess = !BackgroundTaskExecuter->GetTask().IsAborted();
+			break;
+		}
+		FPlatformProcess::Sleep(.2); // SlowTask::ShouldCancel will throttle any updates faster than .2 seconds
+		float ProgressFrac;
+		FText ProgressMessage;
+		bool bMadeProgress = BackgroundTaskExecuter->PollProgress(ProgressFrac, ProgressMessage);
+		if (bMadeProgress)
+		{
+			// SlowTask expects progress to be reported before it happens; we work around this by directly updating the progress amount
+			SlowTask.CompletedWork = ProgressFrac;
+			SlowTask.EnterProgressFrame(0, ProgressMessage);
+		}
+		else
+		{
+			SlowTask.TickProgress(); // Still tick the UI when we don't get new progress frames
+		}
+	}
+
+	if (bSuccess)
+	{
+		check(BackgroundTaskExecuter != nullptr && BackgroundTaskExecuter->IsDone());
+		TUniquePtr<FTrainClothOp> Op = BackgroundTaskExecuter->GetTask().ExtractOperator();
+
+		TUniquePtr<FSkinnedMeshVertices> Result = Op->ExtractResult();
+
+		// @@@@@@@@@@@ TODO: Do something with the result
+
+	}
+
+}
+
+
+void UClothTrainingTool::OnTick(float DeltaTime)
+{
+	if (PendingAction != EClothTrainingToolActions::NoAction)
+	{
+		if (PendingAction == EClothTrainingToolActions::Train)
+		{
+			RunTraining();
+		}
+		PendingAction = EClothTrainingToolActions::NoAction;
+	}
+}
+
+
+void UClothTrainingTool::RequestAction(EClothTrainingToolActions ActionType)
+{
+	if (PendingAction != EClothTrainingToolActions::NoAction)
+	{
+		return;
+	}
+	PendingAction = ActionType;
+}
+
+#undef LOCTEXT_NAMESPACE
diff --git a/Engine/Plugins/Experimental/ChaosClothAssetEditor/Source/ChaosClothAssetEditorTools/Public/ChaosClothAsset/ClothAssetBackedTarget.h b/Engine/Plugins/Experimental/ChaosClothAssetEditor/Source/ChaosClothAssetEditorTools/Public/ChaosClothAsset/ClothAssetBackedTarget.h
new file mode 100644
--- /dev/null
+++ b/Engine/Plugins/Experimental/ChaosClothAssetEditor/Source/ChaosClothAssetEditorTools/Public/ChaosClothAsset/ClothAssetBackedTarget.h
@@ -0,0 +1,38 @@
+// Copyright Epic Games, Inc. All Rights Reserved.
+
+#pragma once
+
+#include "CoreMinimal.h"
+#include "Engine/EngineTypes.h"
+#include "ChaosClothAsset/ClothAsset.h"
+#include "UObject/Interface.h"
+#include "TargetInterfaces/AssetBackedTarget.h"
+
+#include "ClothAssetBackedTarget.generated.h"
+
+UINTERFACE()
+class CHAOSCLOTHASSETEDITORTOOLS_API UClothAssetBackedTarget : public UAssetBackedTarget
+{
+	GENERATED_BODY()
+};
+
+class CHAOSCLOTHASSETEDITORTOOLS_API IClothAssetBackedTarget : public IAssetBackedTarget
+{
+	GENERATED_BODY()
+
+public:
+
+	// IAssetBackedTarget
+	/**
+	 * @return the underlying source asset for this Target.
+	 */
+	virtual UObject* GetSourceData() const override
+	{
+		return GetClothAsset();
+	}
+
+	/**
+	 * @return the underlying source ClothAsset asset for this Target
+	 */
+	virtual UChaosClothAsset* GetClothAsset() const = 0;
+};
diff --git a/Engine/Plugins/Experimental/ChaosClothAssetEditor/Source/ChaosClothAssetEditorTools/Public/ChaosClothAsset/ClothComponentToolTarget.h b/Engine/Plugins/Experimental/ChaosClothAssetEditor/Source/ChaosClothAssetEditorTools/Public/ChaosClothAsset/ClothComponentToolTarget.h
new file mode 100644
--- /dev/null
+++ b/Engine/Plugins/Experimental/ChaosClothAssetEditor/Source/ChaosClothAssetEditorTools/Public/ChaosClothAsset/ClothComponentToolTarget.h
@@ -0,0 +1,45 @@
+// Copyright Epic Games, Inc. All Rights Reserved.
+
+#pragma once
+
+#include "CoreMinimal.h"
+#include "ClothAssetBackedTarget.h"
+#include "ToolTargets/PrimitiveComponentToolTarget.h"
+#include "ClothComponentToolTarget.generated.h"
+
+class UChaosClothComponent;
+
+/**
+ * A tool target backed by a cloth component
+ */
+UCLASS(Transient)
+class CHAOSCLOTHASSETEDITORTOOLS_API UClothComponentToolTarget : public UPrimitiveComponentToolTarget, public IClothAssetBackedTarget
+{
+	GENERATED_BODY()
+
+public:
+
+	// UToolTarget
+	virtual bool IsValid() const override;
+
+	// IClothAssetBackedTarget
+	virtual UChaosClothAsset* GetClothAsset() const override;
+
+	UChaosClothComponent* GetClothComponent() const;
+
+protected:
+	friend class UClothComponentToolTargetFactory;
+};
+
+/** Factory for UClothComponentToolTarget to be used by the target manager. */
+UCLASS(Transient)
+class CHAOSCLOTHASSETEDITORTOOLS_API UClothComponentToolTargetFactory : public UToolTargetFactory
+{
+	GENERATED_BODY()
+
+public:
+
+	virtual bool CanBuildTarget(UObject* SourceObject, const FToolTargetTypeRequirements& TargetTypeInfo) const override;
+
+	virtual UToolTarget* BuildTarget(UObject* SourceObject, const FToolTargetTypeRequirements& TargetTypeInfo) override;
+};
diff --git a/Engine/Plugins/Experimental/ChaosClothAssetEditor/Source/ChaosClothAssetEditorTools/Public/ChaosClothAsset/ClothTrainingTool.h b/Engine/Plugins/Experimental/ChaosClothAssetEditor/Source/ChaosClothAssetEditorTools/Public/ChaosClothAsset/ClothTrainingTool.h
new file mode 100644
--- /dev/null
+++ b/Engine/Plugins/Experimental/ChaosClothAssetEditor/Source/ChaosClothAssetEditorTools/Public/ChaosClothAsset/ClothTrainingTool.h
@@ -0,0 +1,115 @@
+// Copyright Epic Games, Inc. All Rights Reserved.
+
+#pragma once
+
+#include "UObject/Object.h"
+#include "SingleSelectionTool.h"
+#include "InteractiveToolBuilder.h"
+#include "Animation/AnimSequence.h"
+#include "ModelingOperators.h"
+#include "MeshOpPreviewHelpers.h"
+#include "Misc/SlowTask.h"
+
+#include "ClothTrainingTool.generated.h"
+
+class UChaosClothComponent;
+class UClothTrainingTool;
+
+
+// @@@@@@@@@ TODO: Change this to whatever makes sense for output
+struct FSkinnedMeshVertices
+{
+	TArray<FVector3f> Vertices;
+};
+
+
+UCLASS()
+class CHAOSCLOTHASSETEDITORTOOLS_API UClothTrainingToolProperties : public UObject
+{
+	GENERATED_BODY()
+public:
+
+	UPROPERTY(EditAnywhere, Category = Files)
+	TObjectPtr<UAnimSequence> AnimationSequence;
+
+	UPROPERTY(EditAnywhere, Category = Files)
+	FText OutputBufferLocation;
+};
+
+UENUM()
+enum class EClothTrainingToolActions
+{
+	NoAction,
+	Train
+};
+
+UCLASS()
+class CHAOSCLOTHASSETEDITORTOOLS_API UClothTrainingToolActionProperties : public UObject
+{
+	GENERATED_BODY()
+public:
+
+	TWeakObjectPtr<UClothTrainingTool> ParentTool;
+
+	void Initialize(UClothTrainingTool* ParentToolIn) { ParentTool = ParentToolIn; }
+
+	void PostAction(EClothTrainingToolActions Action);
+
+	UFUNCTION(CallInEditor, Category = Actions, meta = (DisplayName = "Begin Training", DisplayPriority = 1))
+		void StartTraining()
+	{
+		PostAction(EClothTrainingToolActions::Train);
+	}
+
+};
+
+
+UCLASS()
+class CHAOSCLOTHASSETEDITORTOOLS_API UClothTrainingToolBuilder : public UInteractiveToolWithToolTargetsBuilder
+{
+	GENERATED_BODY()
+
+protected:
+
+	virtual const FToolTargetTypeRequirements& GetTargetRequirements() const override;
+
+public:
+
+	virtual bool CanBuildTool(const FToolBuilderState& SceneState) const override;
+
+	virtual UInteractiveTool* BuildTool(const FToolBuilderState& SceneState) const override;
+};
+
+
+UCLASS()
+class CHAOSCLOTHASSETEDITORTOOLS_API UClothTrainingTool : public USingleSelectionTool
+{
+	GENERATED_BODY()
+public:
+
+	// UInteractiveTool
+	virtual void Setup() override;
+	virtual void OnTick(float DeltaTime) override;
+
+private:
+
+	friend class UClothTrainingToolActionProperties;
+
+	UPROPERTY()
+	TObjectPtr<UClothTrainingToolProperties> ToolProperties;
+
+	UPROPERTY()
+	TObjectPtr<UClothTrainingToolActionProperties> ActionProperties;
+
+	EClothTrainingToolActions PendingAction = EClothTrainingToolActions::NoAction;
+
+	UPROPERTY()
+	TObjectPtr<const UChaosClothComponent> ClothComponent;
+
+	void RequestAction(EClothTrainingToolActions ActionType);
+
+	void RunTraining();
+
+};
+
+
