Cloth Editor: add a Details panel

diff --git a/Engine/Plugins/Experimental/ChaosClothAssetEditor/Source/ChaosClothAssetEditor/Private/ChaosClothAsset/ClothEditorToolkit.cpp b/Engine/Plugins/Experimental/ChaosClothAssetEditor/Source/ChaosClothAssetEditor/Private/ChaosClothAsset/ClothEditorToolkit.cpp
--- a/Engine/Plugins/Experimental/ChaosClothAssetEditor/Source/ChaosClothAssetEditor/Private/ChaosClothAsset/ClothEditorToolkit.cpp
+++ b/Engine/Plugins/Experimental/ChaosClothAssetEditor/Source/ChaosClothAssetEditor/Private/ChaosClothAsset/ClothEditorToolkit.cpp
@@ -36,7 +36,7 @@ FChaosClothAssetEditorToolkit::FChaosClothAssetEditorToolkit(UAssetEditor* InOwn
 	// Note: Changes to the layout should include a increment to the layout's ID, i.e.
 	// ChaosClothAssetEditorLayout[X] -> ChaosClothAssetEditorLayout[X+1]. Otherwise, layouts may be messed up
 	// without a full reset to layout defaults inside the editor.
-	StandaloneDefaultLayout = FTabManager::NewLayout(FName("ChaosClothAssetEditorLayout2"))
+	StandaloneDefaultLayout = FTabManager::NewLayout(FName("ChaosClothAssetEditorLayout3"))
 		->AddArea
 		(
 			FTabManager::NewPrimaryArea()->SetOrientation(Orient_Vertical)
@@ -46,14 +46,14 @@ FChaosClothAssetEditorToolkit::FChaosClothAssetEditorToolkit(UAssetEditor* InOwn
 				->Split
 				(
 					FTabManager::NewStack()
-					->SetSizeCoefficient(0.2f)
+					->SetSizeCoefficient(0.15f)
 					->SetExtensionId(UChaosClothAssetEditorUISubsystem::EditorSidePanelAreaName)
 					->SetHideTabWell(true)
 				)
 				->Split
 				(
 					FTabManager::NewStack()
-					->SetSizeCoefficient(0.4f)
+					->SetSizeCoefficient(0.35f)
 					->AddTab(ViewportTabID, ETabState::OpenedTab)
 					->SetExtensionId("RestSpaceViewportArea")
 					->SetHideTabWell(true)
@@ -61,11 +61,19 @@ FChaosClothAssetEditorToolkit::FChaosClothAssetEditorToolkit(UAssetEditor* InOwn
 				->Split
 				(
 					FTabManager::NewStack()
-					->SetSizeCoefficient(0.4f)
+					->SetSizeCoefficient(0.35f)
 					->AddTab(ClothPreviewTabID, ETabState::OpenedTab)
 					->SetExtensionId("Viewport3DArea")
 					->SetHideTabWell(true)
 				)
+				->Split
+				(
+					FTabManager::NewStack()
+					->SetSizeCoefficient(0.15f)
+					->AddTab(DetailsTabID, ETabState::OpenedTab)
+					->SetExtensionId("Details")
+					->SetHideTabWell(true)
+				)
 			)
 		);
 
@@ -166,7 +174,7 @@ void FChaosClothAssetEditorToolkit::RegisterTabSpawners(const TSharedRef<FTabMan
 	// Here we set up the tabs we referenced in StandaloneDefaultLayout (in the constructor).
 	// We don't deal with the toolbar palette here, since this is handled by existing
     // infrastructure in FModeToolkit. We only setup spawners for our custom tabs, namely
-    // the 2D and 3D viewports.
+    // the 2D and 3D viewports, and the details panel.
 	InTabManager->RegisterTabSpawner(ClothPreviewTabID, FOnSpawnTab::CreateSP(this, &FChaosClothAssetEditorToolkit::SpawnTab_ClothPreview))
 		.SetDisplayName(LOCTEXT("3DViewportTabLabel", "Cloth 3D Preview Viewport"))
 		.SetGroup(EditorMenuCategory.ToSharedRef())
@@ -176,6 +184,20 @@ void FChaosClothAssetEditorToolkit::RegisterTabSpawners(const TSharedRef<FTabMan
 		.SetDisplayName(LOCTEXT("RestSpaceViewportTabLabel", "Cloth Rest Space Viewport"))
 		.SetGroup(EditorMenuCategory.ToSharedRef())
 		.SetIcon(FSlateIcon(FAppStyle::GetAppStyleSetName(), "LevelEditor.Tabs.Viewports"));
+
+	InTabManager->RegisterTabSpawner(DetailsTabID, FOnSpawnTab::CreateSP(this, &FChaosClothAssetEditorToolkit::SpawnTab_Details))
+		.SetDisplayName(LOCTEXT("Details", "Details"))
+		.SetGroup(AssetEditorTabsCategory.ToSharedRef())
+		.SetIcon(FSlateIcon(FAppStyle::GetAppStyleSetName(), "LevelEditor.Tabs.Details"));
+}
+
+void FChaosClothAssetEditorToolkit::UnregisterTabSpawners(const TSharedRef<FTabManager>& InTabManager)
+{
+	FAssetEditorToolkit::UnregisterTabSpawners(InTabManager);
+
+	InTabManager->UnregisterTabSpawner(ClothPreviewTabID);
+	InTabManager->UnregisterTabSpawner(ViewportTabID);
+	InTabManager->UnregisterTabSpawner(DetailsTabID);
 }
 
 bool FChaosClothAssetEditorToolkit::OnRequestClose()
@@ -375,6 +397,8 @@ void FChaosClothAssetEditorToolkit::PostInitAssetEditor()
 		FBox PreviewBounds = ClothMode->PreviewBoundingBox();
 		ClothPreviewViewportClient->FocusViewportOnBox(PreviewBounds);
 	}
+
+	InitDetailsViewPanel();
 }
 
 
@@ -405,5 +429,17 @@ void FChaosClothAssetEditorToolkit::CreateEditorModeUILayer()
 	ModeUILayer = MakeShared<FChaosClothAssetEditorModeUILayer>(PinnedToolkitHost.Get());
 }
 
+void FChaosClothAssetEditorToolkit::InitDetailsViewPanel()
+{
+	TArray<TObjectPtr<UObject>> ObjectsToEdit;
+	OwningAssetEditor->GetObjectsToEdit(ObjectsToEdit);
+
+	if (ObjectsToEdit.Num() > 0)
+	{
+		UObject* const ObjectToEditInDetailsView = ObjectsToEdit[0];
+		ensure(ObjectToEditInDetailsView->HasAnyFlags(RF_Transactional));		// Ensure all objects are transactable for undo/redo in the details panel
+		SetEditingObject(ObjectToEditInDetailsView);
+	}
+}
 
 #undef LOCTEXT_NAMESPACE
diff --git a/Engine/Plugins/Experimental/ChaosClothAssetEditor/Source/ChaosClothAssetEditor/Public/ChaosClothAsset/ClothEditorToolkit.h b/Engine/Plugins/Experimental/ChaosClothAssetEditor/Source/ChaosClothAssetEditor/Public/ChaosClothAsset/ClothEditorToolkit.h
--- a/Engine/Plugins/Experimental/ChaosClothAssetEditor/Source/ChaosClothAssetEditor/Public/ChaosClothAsset/ClothEditorToolkit.h
+++ b/Engine/Plugins/Experimental/ChaosClothAssetEditor/Source/ChaosClothAssetEditor/Public/ChaosClothAsset/ClothEditorToolkit.h
@@ -34,6 +34,7 @@ public:
 	virtual FText GetBaseToolkitName() const override;
 	virtual FText GetToolkitToolTipText() const override;
 	virtual void RegisterTabSpawners(const TSharedRef<FTabManager>& InTabManager) override;
+	virtual void UnregisterTabSpawners(const TSharedRef<FTabManager>& TabManager) override;
 	virtual bool OnRequestClose() override;
 
 	// TODO: Implement this if we want to apply tool changes before saving an asset
@@ -66,6 +67,8 @@ protected:
 	//virtual const FSlateBrush* GetDefaultTabIcon() const override;
 	//virtual FLinearColor GetDefaultTabColor() const override;
 
+private:
+
 	/** Scene in which the 3D sim space preview meshes live. */
 	TUniquePtr<FAdvancedPreviewScene> ClothPreviewScene;
 
@@ -76,6 +79,8 @@ protected:
 
 	TWeakPtr<SEditorViewport> RestSpaceViewport;
 
+	void InitDetailsViewPanel();
+
 	// TODO as necessary:
 	//TObjectPtr<UInputRouter> ClothPreviewInputRouter = nullptr;
 	//UClothToolViewportButtonsAPI* ViewportButtonsAPI = nullptr;
