Cloth Editor: remove an unnecessary delegate. It was syncing the hidden/visible state between DynamicMeshComponent and Wireframe rendering, but the components should not be hidden

diff --git a/Engine/Plugins/Experimental/ChaosClothAssetEditor/Source/ChaosClothAssetEditor/Private/ChaosClothAsset/ClothEditorMode.cpp b/Engine/Plugins/Experimental/ChaosClothAssetEditor/Source/ChaosClothAssetEditor/Private/ChaosClothAsset/ClothEditorMode.cpp
--- a/Engine/Plugins/Experimental/ChaosClothAssetEditor/Source/ChaosClothAssetEditor/Private/ChaosClothAsset/ClothEditorMode.cpp
+++ b/Engine/Plugins/Experimental/ChaosClothAssetEditor/Source/ChaosClothAssetEditor/Private/ChaosClothAsset/ClothEditorMode.cpp
@@ -538,27 +538,7 @@ void UChaosClothAssetEditorMode::ReinitializeDynamicMeshComponents()
 		// The settings object and wireframe are not part of a tool, so they won't get ticked like they
 		// are supposed to (to enable property watching), unless we add this here.
 		PropertyObjectsToTick.Add(WireframeDisplay->Settings);
-		int32 WireframeIndex = WireframesToTick.Add(WireframeDisplay);
-
-		RestSpaceMeshComponent->MarkRenderStateDirtyEvent.AddLambda( [WireframeIndex, this](UActorComponent& ActorComponent)
-		{
-			if (WireframeIndex >= WireframesToTick.Num() || WireframeIndex >= DynamicMeshComponents.Num())
-			{
-				return;
-			}
-
-			TObjectPtr<UMeshElementsVisualizer> WireframeDisplay = WireframesToTick[WireframeIndex];
-			TObjectPtr<UDynamicMeshComponent> RestSpaceMesh = DynamicMeshComponents[WireframeIndex];
-
-			if (!WireframeDisplay || !RestSpaceMesh)
-			{
-				return;
-			}
-
-			bool bRestSpaceMeshVisible = RestSpaceMesh->GetVisibleFlag();
-			WireframeDisplay->SetAllVisible(bRestSpaceMeshVisible);
-			WireframeDisplay->Settings->bVisible = bRestSpaceMeshVisible;
-		});
+		WireframesToTick.Add(WireframeDisplay);
 	}
 
 	const bool bNumberDynamicMeshesUnchanged = (TotalNumberExistingDynamicMeshComponents == DynamicMeshComponents.Num());
