Cloth Editor:
Add buttons to start/stop/reset cloth simulation
Enable wireframe rendering of cloth component render mesh

diff --git a/Engine/Plugins/Experimental/ChaosClothAssetEditor/Source/ChaosClothAssetEditor/Private/ChaosClothAsset/ClothEditor3DViewportClient.cpp b/Engine/Plugins/Experimental/ChaosClothAssetEditor/Source/ChaosClothAssetEditor/Private/ChaosClothAsset/ClothEditor3DViewportClient.cpp
--- a/Engine/Plugins/Experimental/ChaosClothAssetEditor/Source/ChaosClothAssetEditor/Private/ChaosClothAsset/ClothEditor3DViewportClient.cpp
+++ b/Engine/Plugins/Experimental/ChaosClothAssetEditor/Source/ChaosClothAssetEditor/Private/ChaosClothAsset/ClothEditor3DViewportClient.cpp
@@ -3,7 +3,7 @@
 #include "ChaosClothAsset/ClothEditor3DViewportClient.h"
 #include "ChaosClothAsset/ClothEditorMode.h"
 #include "EditorModeManager.h"
-#include "Components/SkeletalMeshComponent.h"
+#include "ChaosClothAsset/ClothComponent.h"
 
 FChaosClothAssetEditor3DViewportClient::FChaosClothAssetEditor3DViewportClient(FEditorModeTools* InModeTools,
 	FPreviewScene* InPreviewScene, 
@@ -22,15 +22,64 @@ void FChaosClothAssetEditor3DViewportClient::EnableRenderMeshWireframe(bool bEna
 {
 	bRenderMeshWireframe = bEnable;
 
-	for (TObjectPtr<USkeletalMeshComponent> SkeletalMeshComponent : SkeletalMeshComponents)
+	if (ClothComponent)
 	{
-		if (SkeletalMeshComponent)
-		{
-			SkeletalMeshComponent->SetForceWireframe(bRenderMeshWireframe);
-		}
+		ClothComponent->SetForceWireframe(bRenderMeshWireframe);
 	}
 }
 
+void FChaosClothAssetEditor3DViewportClient::SetClothComponent(TObjectPtr<UChaosClothComponent> InClothComponent)
+{
+	ClothComponent = InClothComponent;
+}
+
+void FChaosClothAssetEditor3DViewportClient::SetClothEdMode(TObjectPtr<UChaosClothAssetEditorMode> InClothEdMode)
+{
+	ClothEdMode = InClothEdMode;
+}
+
+void FChaosClothAssetEditor3DViewportClient::SoftResetSimulation()
+{
+	if (ClothEdMode)
+	{
+		ClothEdMode->SoftResetSimulation();
+	}
+}
+
+void FChaosClothAssetEditor3DViewportClient::HardResetSimulation()
+{
+	if (ClothEdMode)
+	{
+		ClothEdMode->HardResetSimulation();
+	}
+}
+
+void FChaosClothAssetEditor3DViewportClient::SuspendSimulation()
+{
+	if (ClothEdMode)
+	{
+		ClothEdMode->SuspendSimulation();
+	}
+}
+
+void FChaosClothAssetEditor3DViewportClient::ResumeSimulation()
+{
+	if (ClothEdMode)
+	{
+		ClothEdMode->ResumeSimulation();
+	}
+}
+
+bool FChaosClothAssetEditor3DViewportClient::IsSimulationSuspended() const
+{
+	if (ClothEdMode)
+	{
+		return ClothEdMode->IsSimulationSuspended();
+	}
+
+	return false;
+}
+
 void FChaosClothAssetEditor3DViewportClient::Draw(const FSceneView* View, FPrimitiveDrawInterface* PDI)
 {
 	if (!bSimMeshWireframe)
diff --git a/Engine/Plugins/Experimental/ChaosClothAssetEditor/Source/ChaosClothAssetEditor/Private/ChaosClothAsset/ClothEditorCommands.cpp b/Engine/Plugins/Experimental/ChaosClothAssetEditor/Source/ChaosClothAssetEditor/Private/ChaosClothAsset/ClothEditorCommands.cpp
--- a/Engine/Plugins/Experimental/ChaosClothAssetEditor/Source/ChaosClothAssetEditor/Private/ChaosClothAsset/ClothEditorCommands.cpp
+++ b/Engine/Plugins/Experimental/ChaosClothAssetEditor/Source/ChaosClothAssetEditor/Private/ChaosClothAsset/ClothEditorCommands.cpp
@@ -12,6 +12,9 @@ const FString FChaosClothAssetEditorCommands::BeginWeightMapPaintToolIdentifier
 const FString FChaosClothAssetEditorCommands::BeginClothTrainingToolIdentifier = TEXT("BeginClothTrainingTool");
 const FString FChaosClothAssetEditorCommands::ToggleSimMeshWireframeIdentifier = TEXT("ToggleSimMeshWireframe");
 const FString FChaosClothAssetEditorCommands::ToggleRenderMeshWireframeIdentifier = TEXT("ToggleRenderMeshWireframe");
+const FString FChaosClothAssetEditorCommands::ToggleSimulationSuspendedIdentifier = TEXT("ToggleSimulationSuspended");
+const FString FChaosClothAssetEditorCommands::SoftResetSimulationIdentifier = TEXT("SoftResetSimulation");
+const FString FChaosClothAssetEditorCommands::HardResetSimulationIdentifier = TEXT("HardResetSimulation");
 const FString FChaosClothAssetEditorCommands::TogglePatternModeIdentifier = TEXT("TogglePatternMode");
 
 
@@ -38,8 +41,13 @@ void FChaosClothAssetEditorCommands::RegisterCommands()
 
 	UI_COMMAND(TogglePatternMode, "TogglePatternMode", "Toggle pattern mode", EUserInterfaceActionType::Button, FInputChord());
 
-	UI_COMMAND(ToggleSimMeshWireframe, "ToggleSimMeshWireframe", "Toggle simulation mesh wireframe", EUserInterfaceActionType::Button, FInputChord());
-	UI_COMMAND(ToggleRenderMeshWireframe, "ToggleRenderMeshWireframe", "Toggle render mesh wireframe", EUserInterfaceActionType::Button, FInputChord());
+	UI_COMMAND(ToggleSimMeshWireframe, "ToggleSimMeshWireframe", "Toggle simulation mesh wireframe", EUserInterfaceActionType::ToggleButton, FInputChord());
+	UI_COMMAND(ToggleRenderMeshWireframe, "ToggleRenderMeshWireframe", "Toggle render mesh wireframe", EUserInterfaceActionType::ToggleButton, FInputChord());
+	
+	UI_COMMAND(SoftResetSimulation, "SoftResetSimulation", "Soft reset simulation", EUserInterfaceActionType::ToggleButton, FInputChord());
+	UI_COMMAND(HardResetSimulation, "HardResetClothSimulation", "Hard reset simulation", EUserInterfaceActionType::ToggleButton, FInputChord());
+	UI_COMMAND(ToggleSimulationSuspended, "ToggleSimulationSuspended", "Toggle simulation suspended", EUserInterfaceActionType::ToggleButton, FInputChord());
+
 }
 
 void FChaosClothAssetEditorCommands::GetToolDefaultObjectList(TArray<UInteractiveTool*>& ToolCDOs)
diff --git a/Engine/Plugins/Experimental/ChaosClothAssetEditor/Source/ChaosClothAssetEditor/Private/ChaosClothAsset/ClothEditorMode.cpp b/Engine/Plugins/Experimental/ChaosClothAssetEditor/Source/ChaosClothAssetEditor/Private/ChaosClothAsset/ClothEditorMode.cpp
--- a/Engine/Plugins/Experimental/ChaosClothAssetEditor/Source/ChaosClothAssetEditor/Private/ChaosClothAsset/ClothEditorMode.cpp
+++ b/Engine/Plugins/Experimental/ChaosClothAssetEditor/Source/ChaosClothAssetEditor/Private/ChaosClothAsset/ClothEditorMode.cpp
@@ -648,6 +648,36 @@ void UChaosClothAssetEditorMode::ApplyChanges()
 	GetToolManager()->EndUndoTransaction();
 }
 
+
+void UChaosClothAssetEditorMode::SoftResetSimulation()
+{
+	bShouldResetSimulation = true;
+	bShouldClearTeleportFlag = false;
+	bHardReset = false;
+}
+
+void UChaosClothAssetEditorMode::HardResetSimulation()
+{
+	bShouldResetSimulation = true;
+	bShouldClearTeleportFlag = false;
+	bHardReset = true;
+}
+
+void UChaosClothAssetEditorMode::SuspendSimulation()
+{
+	bIsSimulationSuspended = true;
+}
+
+void UChaosClothAssetEditorMode::ResumeSimulation()
+{
+	bIsSimulationSuspended = false;
+}
+
+bool UChaosClothAssetEditorMode::IsSimulationSuspended() const
+{
+	return bIsSimulationSuspended;
+}
+
 void UChaosClothAssetEditorMode::ModeTick(float DeltaTime)
 {
 	Super::ModeTick(DeltaTime);
@@ -675,11 +705,35 @@ void UChaosClothAssetEditorMode::ModeTick(float DeltaTime)
 		}
 	}
 
-	// TODO: Uncomment this to enable simulation stepping
-	//if (PreviewWorld)
-	//{
-	//	PreviewWorld->Tick(ELevelTick::LEVELTICK_All, DeltaTime);
-	//}
+
+	bool bShouldTickPreviewWorld = !bIsSimulationSuspended;
+
+	if (bShouldClearTeleportFlag)
+	{
+		ClothComponent->ResetClothTeleportMode();
+		bShouldClearTeleportFlag = false;
+	}
+
+	if (bShouldResetSimulation)
+	{
+		if (bHardReset)
+		{
+			const FComponentReregisterContext Context(ClothComponent);
+		}
+		else
+		{
+			ClothComponent->ForceClothNextUpdateTeleportAndReset();
+		}
+
+		bShouldResetSimulation = false;
+		bShouldClearTeleportFlag = true;		// clear the flag next tick
+		bShouldTickPreviewWorld = true;
+	}
+
+	if (PreviewWorld && bShouldTickPreviewWorld)
+	{
+		PreviewWorld->Tick(ELevelTick::LEVELTICK_All, DeltaTime);
+	}
 }
 
 
diff --git a/Engine/Plugins/Experimental/ChaosClothAssetEditor/Source/ChaosClothAssetEditor/Private/ChaosClothAsset/ClothEditorStyle.cpp b/Engine/Plugins/Experimental/ChaosClothAssetEditor/Source/ChaosClothAssetEditor/Private/ChaosClothAsset/ClothEditorStyle.cpp
--- a/Engine/Plugins/Experimental/ChaosClothAssetEditor/Source/ChaosClothAssetEditor/Private/ChaosClothAsset/ClothEditorStyle.cpp
+++ b/Engine/Plugins/Experimental/ChaosClothAssetEditor/Source/ChaosClothAssetEditor/Private/ChaosClothAsset/ClothEditorStyle.cpp
@@ -22,6 +22,12 @@ FChaosClothAssetEditorStyle::FChaosClothAssetEditorStyle()
 {
 	// Used FUVEditorStyle and FModelingToolsEditorModeStyle and as models
 
+	SetCoreContentRoot(FPaths::EngineContentDir() / TEXT("Editor/Slate"));
+
+	// TODO: Set this once we get our own UI content
+	FString ContentDir = IPluginManager::Get().FindPlugin(TEXT("ModelingToolsEditorMode"))->GetContentDir();
+	SetContentRoot(ContentDir);
+	
 	// Some standard icon sizes used elsewhere in the editor
 	const FVector2D Icon8x8(8.0f, 8.0f);
 	const FVector2D Icon16x16(16.0f, 16.0f);
@@ -34,11 +40,6 @@ FChaosClothAssetEditorStyle::FChaosClothAssetEditorStyle()
 	const FVector2D ViewportToolbarIconSize = Icon16x16;
 	const FVector2D ToolbarIconSize = Icon20x20;
 
-	// TODO: Set this once we get our own UI content
-	FString ContentDir = IPluginManager::Get().FindPlugin(TEXT("ModelingToolsEditorMode"))->GetContentDir();
-	SetContentRoot(ContentDir);
-	SetCoreContentRoot(FPaths::EngineContentDir() / TEXT("Slate"));
-
 	FString PropertyNameString = "ChaosClothAssetEditor." + FChaosClothAssetEditorCommands::BeginRemeshToolIdentifier;
 	Set(*PropertyNameString, new FSlateImageBrush(FChaosClothAssetEditorStyle::InContent("Icons/Remesh_40x", ".png"), ToolbarIconSize ));
 	
@@ -48,18 +49,35 @@ FChaosClothAssetEditorStyle::FChaosClothAssetEditorStyle()
 	PropertyNameString = "ChaosClothAssetEditor." + FChaosClothAssetEditorCommands::BeginWeightMapPaintToolIdentifier;
 	Set(*PropertyNameString, new FSlateImageBrush(FChaosClothAssetEditorStyle::InContent("Icons/ModelingAttributePaint_x40", ".png"), ToolbarIconSize));
 
+	// TODO: Get a permanent icon to replace this placeholder
 	PropertyNameString = "ChaosClothAssetEditor." + FChaosClothAssetEditorCommands::BeginClothTrainingToolIdentifier;
 	Set(*PropertyNameString, new FSlateImageBrush(FChaosClothAssetEditorStyle::InContent("Icons/ModelingAttributePaint_x40", ".png"), ToolbarIconSize));
 
+	// TODO: Get a permanent icon to replace this placeholder
+	PropertyNameString = "ChaosClothAssetEditor." + FChaosClothAssetEditorCommands::TogglePatternModeIdentifier;
+	Set(*PropertyNameString, new FSlateImageBrush(FChaosClothAssetEditorStyle::InContent("Icons/TogglePatternMode_40x", ".png"), ToolbarIconSize));
+
+	// TODO: Get a permanent icon to replace this placeholder
 	PropertyNameString = "ChaosClothAssetEditor." + FChaosClothAssetEditorCommands::ToggleSimMeshWireframeIdentifier;
-	Set(*PropertyNameString, new FSlateImageBrush(FChaosClothAssetEditorStyle::InContent("Icons/icon_ViewMode_BrushWireframe_16px", ".png"), ViewportToolbarIconSize));		// Placeholder
+	Set(*PropertyNameString, new CORE_IMAGE_BRUSH_SVG("Starship/EditorViewport/grid", ToolbarIconSize));
 
+	// TODO: Get a permanent icon to replace this placeholder
 	PropertyNameString = "ChaosClothAssetEditor." + FChaosClothAssetEditorCommands::ToggleRenderMeshWireframeIdentifier;
-	Set(*PropertyNameString, new FSlateImageBrush(FChaosClothAssetEditorStyle::InContent("Icons/icon_ViewMode_BrushWireframe_16px", ".png"), ViewportToolbarIconSize));		// Placeholder
+	Set(*PropertyNameString, new CORE_IMAGE_BRUSH_SVG("Starship/EditorViewport/grid", ToolbarIconSize));
 
-	PropertyNameString = "ChaosClothAssetEditor." + FChaosClothAssetEditorCommands::TogglePatternModeIdentifier;
-	Set(*PropertyNameString, new FSlateImageBrush(FChaosClothAssetEditorStyle::InContent("Icons/TogglePatternMode_40x", ".png"), ViewportToolbarIconSize));
+	// TODO: Get a permanent icon to replace this placeholder
+	// TODO: Toggle back and forth between Play and Pause icons when the button is clicked?
+	PropertyNameString = "ChaosClothAssetEditor." + FChaosClothAssetEditorCommands::ToggleSimulationSuspendedIdentifier;
+	Set(*PropertyNameString, new CORE_IMAGE_BRUSH_SVG("Starship/MainToolbar/pause", ToolbarIconSize));
+
+	// TODO: Get a permanent icon to replace this placeholder
+	PropertyNameString = "ChaosClothAssetEditor." + FChaosClothAssetEditorCommands::SoftResetSimulationIdentifier;
+	Set(*PropertyNameString, new CORE_IMAGE_BRUSH_SVG("Starship/Common/Reset", ToolbarIconSize));
 
+	// TODO: Get a permanent icon to replace this placeholder
+	PropertyNameString = "ChaosClothAssetEditor." + FChaosClothAssetEditorCommands::HardResetSimulationIdentifier;
+	Set(*PropertyNameString, new CORE_IMAGE_BRUSH_SVG("Starship/Common/Apply", ToolbarIconSize));
+	
 	FSlateStyleRegistry::RegisterSlateStyle(*this);
 }
 
diff --git a/Engine/Plugins/Experimental/ChaosClothAssetEditor/Source/ChaosClothAssetEditor/Private/ChaosClothAsset/ClothEditorToolkit.cpp b/Engine/Plugins/Experimental/ChaosClothAssetEditor/Source/ChaosClothAssetEditor/Private/ChaosClothAsset/ClothEditorToolkit.cpp
--- a/Engine/Plugins/Experimental/ChaosClothAssetEditor/Source/ChaosClothAssetEditor/Private/ChaosClothAsset/ClothEditorToolkit.cpp
+++ b/Engine/Plugins/Experimental/ChaosClothAssetEditor/Source/ChaosClothAssetEditor/Private/ChaosClothAsset/ClothEditorToolkit.cpp
@@ -389,6 +389,9 @@ void FChaosClothAssetEditorToolkit::PostInitAssetEditor()
 	ViewportClient->ReceivedFocus(ViewportClient->Viewport);
 
 	// Set up 3D viewport
+	ClothPreviewViewportClient->SetClothComponent(ClothMode->ClothComponent);
+	ClothPreviewViewportClient->SetClothEdMode(ClothMode);
+
 	SetCommonViewportClientOptions(ClothPreviewViewportClient.Get());
 	ClothPreviewViewportClient->SetInitialViewTransform(ELevelViewportType::LVT_Perspective, FVector(0, -100, 100), FRotator(0, 90, 0), DEFAULT_ORTHOZOOM);
 
diff --git a/Engine/Plugins/Experimental/ChaosClothAssetEditor/Source/ChaosClothAssetEditor/Private/ChaosClothAsset/SClothEditor3DViewport.cpp b/Engine/Plugins/Experimental/ChaosClothAssetEditor/Source/ChaosClothAssetEditor/Private/ChaosClothAsset/SClothEditor3DViewport.cpp
--- a/Engine/Plugins/Experimental/ChaosClothAssetEditor/Source/ChaosClothAssetEditor/Private/ChaosClothAsset/SClothEditor3DViewport.cpp
+++ b/Engine/Plugins/Experimental/ChaosClothAssetEditor/Source/ChaosClothAssetEditor/Private/ChaosClothAsset/SClothEditor3DViewport.cpp
@@ -33,6 +33,47 @@ void SChaosClothAssetEditor3DViewport::BindCommands()
 		}),
 		FCanExecuteAction::CreateLambda([this]() { return true; }),
 		FIsActionChecked::CreateLambda([this]() { return StaticCastSharedPtr<FChaosClothAssetEditor3DViewportClient>(Client)->RenderMeshWireframeEnabled(); }));
+
+	CommandList->MapAction(
+		CommandInfos.SoftResetSimulation,
+		FExecuteAction::CreateLambda([this]()
+		{
+			TSharedPtr<FChaosClothAssetEditor3DViewportClient> ClothViewportClient = StaticCastSharedPtr<FChaosClothAssetEditor3DViewportClient>(Client);
+			ClothViewportClient->SoftResetSimulation();
+		}),
+		FCanExecuteAction::CreateLambda([this]() { return true; }),
+		FIsActionChecked::CreateLambda([this]() { return false; }));
+
+	CommandList->MapAction(
+		CommandInfos.HardResetSimulation,
+		FExecuteAction::CreateLambda([this]()
+		{
+			TSharedPtr<FChaosClothAssetEditor3DViewportClient> ClothViewportClient = StaticCastSharedPtr<FChaosClothAssetEditor3DViewportClient>(Client);
+			ClothViewportClient->HardResetSimulation();
+		}),
+		FCanExecuteAction::CreateLambda([this]() { return true; }),
+		FIsActionChecked::CreateLambda([this]() { return false; }));
+
+
+	CommandList->MapAction(
+		CommandInfos.ToggleSimulationSuspended,
+		FExecuteAction::CreateLambda([this]()
+		{
+			TSharedPtr<FChaosClothAssetEditor3DViewportClient> ClothViewportClient = StaticCastSharedPtr<FChaosClothAssetEditor3DViewportClient>(Client);
+
+			const bool bIsSuspended = ClothViewportClient->IsSimulationSuspended();
+			if (bIsSuspended)
+			{
+				ClothViewportClient->ResumeSimulation();
+			}
+			else
+			{
+				ClothViewportClient->SuspendSimulation();
+			}
+		}),
+		FCanExecuteAction::CreateLambda([this]() { return true; }),
+		FIsActionChecked::CreateLambda([this]() { return StaticCastSharedPtr<FChaosClothAssetEditor3DViewportClient>(Client)->IsSimulationSuspended(); }) );
+
 }
 
 TSharedPtr<SWidget> SChaosClothAssetEditor3DViewport::MakeViewportToolbar()
diff --git a/Engine/Plugins/Experimental/ChaosClothAssetEditor/Source/ChaosClothAssetEditor/Private/ChaosClothAsset/SClothEditor3DViewportToolBar.cpp b/Engine/Plugins/Experimental/ChaosClothAssetEditor/Source/ChaosClothAssetEditor/Private/ChaosClothAsset/SClothEditor3DViewportToolBar.cpp
--- a/Engine/Plugins/Experimental/ChaosClothAssetEditor/Source/ChaosClothAssetEditor/Private/ChaosClothAsset/SClothEditor3DViewportToolBar.cpp
+++ b/Engine/Plugins/Experimental/ChaosClothAssetEditor/Source/ChaosClothAssetEditor/Private/ChaosClothAsset/SClothEditor3DViewportToolBar.cpp
@@ -3,6 +3,7 @@
 #include "ChaosClothAsset/SClothEditor3DViewportToolBar.h"
 #include "ChaosClothAsset/SClothEditor3DViewport.h"
 #include "ChaosClothAsset/ClothEditorCommands.h"
+#include "ChaosClothAsset/ClothEditorStyle.h"
 #include "EditorViewportCommands.h"
 #include "Styling/AppStyle.h"
 #include "Framework/MultiBox/MultiBoxBuilder.h"
@@ -73,19 +74,52 @@ TSharedRef<SWidget> SChaosClothAssetEditor3DViewportToolBar::MakeToolBar(const T
 	ToolbarBuilder.SetStyle(&FAppStyle::Get(), ToolBarStyle);
 	ToolbarBuilder.SetLabelVisibility(EVisibility::Collapsed);
 
-	ToolbarBuilder.BeginSection("SimControls");
-	{
-		ToolbarBuilder.BeginBlockGroup();
-	    
-		static FName ToggleSimMeshWireframeName = FName(TEXT("SimMeshWireframe"));
-		ToolbarBuilder.AddToolBarButton(FChaosClothAssetEditorCommands::Get().ToggleSimMeshWireframe);
-
-		static FName ToggleRenderMeshWireframeName = FName(TEXT("RenderMeshWireframe"));
-		ToolbarBuilder.AddToolBarButton(FChaosClothAssetEditorCommands::Get().ToggleRenderMeshWireframe);
+	FString PropertyNameString;
 
-		ToolbarBuilder.EndBlockGroup();
+	ToolbarBuilder.BeginSection("Visualization");
+	ToolbarBuilder.BeginBlockGroup();
+	{
+		// TODO: Add button to toggle between sim and render meshes
+
+		PropertyNameString = "ChaosClothAssetEditor." + FChaosClothAssetEditorCommands::ToggleRenderMeshWireframeIdentifier;
+		ToolbarBuilder.AddToolBarButton(FChaosClothAssetEditorCommands::Get().ToggleRenderMeshWireframe,
+			NAME_None,
+			TAttribute<FText>(),
+			TAttribute<FText>(),
+			TAttribute<FSlateIcon>(),
+			TEXT("ToggleRenderMeshWireframe"));
 	}
+	ToolbarBuilder.EndBlockGroup();
+	ToolbarBuilder.EndSection();
 
+	ToolbarBuilder.BeginSection("Sim Controls");
+	ToolbarBuilder.BeginBlockGroup();
+	{
+		PropertyNameString = "ChaosClothAssetEditor." + FChaosClothAssetEditorCommands::HardResetSimulationIdentifier;
+		ToolbarBuilder.AddToolBarButton(FChaosClothAssetEditorCommands::Get().HardResetSimulation,
+			NAME_None,
+			TAttribute<FText>(),
+			TAttribute<FText>(),
+			TAttribute<FSlateIcon>(),
+			FName(*FChaosClothAssetEditorCommands::HardResetSimulationIdentifier));
+
+		PropertyNameString = "ChaosClothAssetEditor." + FChaosClothAssetEditorCommands::SoftResetSimulationIdentifier;
+		ToolbarBuilder.AddToolBarButton(FChaosClothAssetEditorCommands::Get().SoftResetSimulation,
+			NAME_None,
+			TAttribute<FText>(),
+			TAttribute<FText>(),
+			TAttribute<FSlateIcon>(),
+			FName(FChaosClothAssetEditorCommands::SoftResetSimulationIdentifier));
+
+		PropertyNameString = "ChaosClothAssetEditor." + FChaosClothAssetEditorCommands::ToggleSimulationSuspendedIdentifier;
+		ToolbarBuilder.AddToolBarButton(FChaosClothAssetEditorCommands::Get().ToggleSimulationSuspended,
+			NAME_None,
+			TAttribute<FText>(),
+			TAttribute<FText>(),
+			TAttribute<FSlateIcon>(),
+			FName(FChaosClothAssetEditorCommands::ToggleSimulationSuspendedIdentifier));
+	}
+	ToolbarBuilder.EndBlockGroup();
 	ToolbarBuilder.EndSection();
 
 	return ToolbarBuilder.MakeWidget();
diff --git a/Engine/Plugins/Experimental/ChaosClothAssetEditor/Source/ChaosClothAssetEditor/Public/ChaosClothAsset/ClothEditor3DViewportClient.h b/Engine/Plugins/Experimental/ChaosClothAssetEditor/Source/ChaosClothAssetEditor/Public/ChaosClothAsset/ClothEditor3DViewportClient.h
--- a/Engine/Plugins/Experimental/ChaosClothAssetEditor/Source/ChaosClothAssetEditor/Public/ChaosClothAsset/ClothEditor3DViewportClient.h
+++ b/Engine/Plugins/Experimental/ChaosClothAssetEditor/Source/ChaosClothAssetEditor/Public/ChaosClothAsset/ClothEditor3DViewportClient.h
@@ -5,7 +5,9 @@
 #include "CoreMinimal.h"
 #include "EditorViewportClient.h"
 
-class USkeletalMeshComponent;
+class UChaosClothComponent;
+class UChaosClothAssetEditorMode;
+
 /**
  * Viewport client for the 3d sim preview in the cloth editor. Currently same as editor viewport
  * client but doesn't allow editor gizmos/widgets.
@@ -29,18 +31,28 @@ public:
 	void EnableRenderMeshWireframe(bool bEnable);
 	bool RenderMeshWireframeEnabled() const { return bRenderMeshWireframe; }
 
-	void SetSkeletalMeshComponents(const TArray<TObjectPtr<USkeletalMeshComponent>>& NewSkMeshComponents)
-	{
-		SkeletalMeshComponents = NewSkMeshComponents;
-	}
+	void SetClothComponent(TObjectPtr<UChaosClothComponent> ClothComponent);
+	void SetClothEdMode(TObjectPtr<UChaosClothAssetEditorMode> ClothEdMode);
+
+	void SoftResetSimulation();
+	void HardResetSimulation();
+	void SuspendSimulation();
+	void ResumeSimulation();
+	bool IsSimulationSuspended() const;
+
 
-protected:
+private:
 
-	TArray<TObjectPtr<USkeletalMeshComponent>> SkeletalMeshComponents;
+	TObjectPtr<UChaosClothComponent> ClothComponent;
+
+	TObjectPtr<UChaosClothAssetEditorMode> ClothEdMode;
 
 	// Debug draw of simulation meshes
 	virtual void Draw(const FSceneView* View, FPrimitiveDrawInterface* PDI) override;
 
+	UChaosClothAssetEditorMode* GetClothEdMode();
+	const UChaosClothAssetEditorMode* GetClothEdMode() const;
+
 	bool bSimMeshWireframe = true;
 	bool bRenderMeshWireframe = false;
 };
diff --git a/Engine/Plugins/Experimental/ChaosClothAssetEditor/Source/ChaosClothAssetEditor/Public/ChaosClothAsset/ClothEditorCommands.h b/Engine/Plugins/Experimental/ChaosClothAssetEditor/Source/ChaosClothAssetEditor/Public/ChaosClothAsset/ClothEditorCommands.h
--- a/Engine/Plugins/Experimental/ChaosClothAssetEditor/Source/ChaosClothAssetEditor/Public/ChaosClothAsset/ClothEditorCommands.h
+++ b/Engine/Plugins/Experimental/ChaosClothAssetEditor/Source/ChaosClothAssetEditor/Public/ChaosClothAsset/ClothEditorCommands.h
@@ -50,4 +50,14 @@ public:
 	const static FString ToggleRenderMeshWireframeIdentifier;
 	TSharedPtr<FUICommandInfo> ToggleRenderMeshWireframe;
 
+
+	const static FString SoftResetSimulationIdentifier;
+	TSharedPtr<FUICommandInfo> SoftResetSimulation;
+
+	const static FString HardResetSimulationIdentifier;
+	TSharedPtr<FUICommandInfo> HardResetSimulation;
+
+	const static FString ToggleSimulationSuspendedIdentifier;
+	TSharedPtr<FUICommandInfo> ToggleSimulationSuspended;
+
 };
diff --git a/Engine/Plugins/Experimental/ChaosClothAssetEditor/Source/ChaosClothAssetEditor/Public/ChaosClothAsset/ClothEditorMode.h b/Engine/Plugins/Experimental/ChaosClothAssetEditor/Source/ChaosClothAssetEditor/Public/ChaosClothAsset/ClothEditorMode.h
--- a/Engine/Plugins/Experimental/ChaosClothAssetEditor/Source/ChaosClothAssetEditor/Public/ChaosClothAsset/ClothEditorMode.h
+++ b/Engine/Plugins/Experimental/ChaosClothAssetEditor/Source/ChaosClothAssetEditor/Public/ChaosClothAsset/ClothEditorMode.h
@@ -85,6 +85,13 @@ public:
 	void TogglePatternMode();
 	bool CanTogglePatternMode() const;
 
+	// Simulation controls
+	void SoftResetSimulation();
+	void HardResetSimulation();
+	void SuspendSimulation();
+	void ResumeSimulation();
+	bool IsSimulationSuspended() const;
+
 private:
 
 	friend class FChaosClothAssetEditorToolkit;
@@ -164,5 +171,9 @@ private:
 	// Set up the preview simulation mesh from the given rest-space mesh
 	void UpdateSimulationMeshes();
 
+	bool bIsSimulationSuspended = false;
+	bool bShouldResetSimulation = false;
+	bool bHardReset = false;
+	bool bShouldClearTeleportFlag = false;
 };
 
