Cloth Editor:
Undo CL 22956278 and implement a better fix.

Turns out we do need the delegates that sync the DynamicMeshComponent and WireframeDisplay visible states, since some tools will hide the DynamicMeshComponent and create a temporary Preview mesh. Without these delegates, the wireframe visualizers will remain visible while the tool is active.

To fix the delegates not being cleaned up properly, add them using AddWeakLambda, then remove them all when exiting the cloth editor.

diff --git a/Engine/Plugins/Experimental/ChaosClothAssetEditor/Source/ChaosClothAssetEditor/Private/ChaosClothAsset/ClothEditorMode.cpp b/Engine/Plugins/Experimental/ChaosClothAssetEditor/Source/ChaosClothAssetEditor/Private/ChaosClothAsset/ClothEditorMode.cpp
--- a/Engine/Plugins/Experimental/ChaosClothAssetEditor/Source/ChaosClothAssetEditor/Private/ChaosClothAsset/ClothEditorMode.cpp
+++ b/Engine/Plugins/Experimental/ChaosClothAssetEditor/Source/ChaosClothAssetEditor/Private/ChaosClothAsset/ClothEditorMode.cpp
@@ -227,6 +227,7 @@ void UChaosClothAssetEditorMode::BindCommands()
 void UChaosClothAssetEditorMode::Exit()
 {
 	USelection::SelectionChangedEvent.Remove(SelectionModifiedEventHandle);
+	UActorComponent::MarkRenderStateDirtyEvent.RemoveAll(this);
 
 	for (TObjectPtr<UMeshElementsVisualizer> WireframeDisplay : WireframesToTick)
 	{
@@ -538,7 +539,30 @@ void UChaosClothAssetEditorMode::ReinitializeDynamicMeshComponents()
 		// The settings object and wireframe are not part of a tool, so they won't get ticked like they
 		// are supposed to (to enable property watching), unless we add this here.
 		PropertyObjectsToTick.Add(WireframeDisplay->Settings);
-		WireframesToTick.Add(WireframeDisplay);
+		int32 WireframeIndex = WireframesToTick.Add(WireframeDisplay);
+
+		// Some interactive tools will hide the input DynamicMeshComponent and create their own temporary PreviewMesh for visualization. If this
+		// occurs, we should also hide the corresponding WireframeDisplay (and un-hide it when the tool finishes).
+		UActorComponent::MarkRenderStateDirtyEvent.AddWeakLambda(this, [WireframeIndex, this](UActorComponent& ActorComponent)
+		{
+			if (WireframeIndex >= WireframesToTick.Num() || WireframeIndex >= DynamicMeshComponents.Num())
+			{
+				return;
+			}
+
+			TObjectPtr<UMeshElementsVisualizer> WireframeDisplay = WireframesToTick[WireframeIndex];
+			const TObjectPtr<UDynamicMeshComponent> RestSpaceMesh = DynamicMeshComponents[WireframeIndex];
+
+			if (!WireframeDisplay || !RestSpaceMesh)
+			{
+				return;
+			}
+
+			bool bRestSpaceMeshVisible = RestSpaceMesh->GetVisibleFlag();
+			WireframeDisplay->SetAllVisible(bRestSpaceMeshVisible);
+			WireframeDisplay->Settings->bVisible = bRestSpaceMeshVisible;
+		});
+
 	}
 
 	const bool bNumberDynamicMeshesUnchanged = (TotalNumberExistingDynamicMeshComponents == DynamicMeshComponents.Num());
