ClothEditor: press F in the preview viewport to focus on the simulation mesh

diff --git a/Engine/Plugins/Experimental/ChaosClothAssetEditor/Source/ChaosClothAssetEditor/Private/ChaosClothAsset/ClothEditor3DViewportClient.cpp b/Engine/Plugins/Experimental/ChaosClothAssetEditor/Source/ChaosClothAssetEditor/Private/ChaosClothAsset/ClothEditor3DViewportClient.cpp
--- a/Engine/Plugins/Experimental/ChaosClothAssetEditor/Source/ChaosClothAssetEditor/Private/ChaosClothAsset/ClothEditor3DViewportClient.cpp
+++ b/Engine/Plugins/Experimental/ChaosClothAssetEditor/Source/ChaosClothAssetEditor/Private/ChaosClothAsset/ClothEditor3DViewportClient.cpp
@@ -91,3 +91,12 @@ void FChaosClothAssetEditor3DViewportClient::Draw(const FSceneView* View, FPrimi
 
 }
 
+FBox FChaosClothAssetEditor3DViewportClient::PreviewBoundingBox() const
+{
+	if (ClothEdMode)
+	{
+		return ClothEdMode->PreviewBoundingBox();
+	}
+
+	return FBox(ForceInitToZero);
+}
\ No newline at end of file
diff --git a/Engine/Plugins/Experimental/ChaosClothAssetEditor/Source/ChaosClothAssetEditor/Private/ChaosClothAsset/ClothEditorMode.cpp b/Engine/Plugins/Experimental/ChaosClothAssetEditor/Source/ChaosClothAssetEditor/Private/ChaosClothAsset/ClothEditorMode.cpp
--- a/Engine/Plugins/Experimental/ChaosClothAssetEditor/Source/ChaosClothAssetEditor/Private/ChaosClothAsset/ClothEditorMode.cpp
+++ b/Engine/Plugins/Experimental/ChaosClothAssetEditor/Source/ChaosClothAssetEditor/Private/ChaosClothAsset/ClothEditorMode.cpp
@@ -66,22 +66,6 @@ namespace ChaosClothAssetEditorModeLocals
 		// TODO: use a non-temp category
 		UE_LOG(LogTemp, Warning, TEXT("%s"), *Error.ToString());
 	}
-
-	FBoxSphereBounds RenderMeshBounds(const UE::Chaos::ClothAsset::FClothPatternConstAdapter& ClothPatternAdapter)
-	{
-		const TSharedPtr<const UE::Chaos::ClothAsset::FClothCollection> ClothCollection = ClothPatternAdapter.GetClothCollection();
-		const int32 PatternElementIndex = ClothPatternAdapter.GetElementIndex();
-
-		const int32 RenderVerticesStart = ClothCollection->RenderVerticesStart[PatternElementIndex];
-		const int32 RenderVerticesEnd = ClothCollection->RenderVerticesEnd[PatternElementIndex];
-
-		const FVector3f* RawArray = &ClothCollection->RenderPosition[0] + RenderVerticesStart;
-		const int32 NumPoints = RenderVerticesEnd - RenderVerticesStart;
-		const FBoxSphereBounds3f Bounds(RawArray, NumPoints);
-
-		return FBoxSphereBounds(Bounds);
-	}
-
 }		// namespace ChaosClothAssetEditorModeLocals
 
 const FToolTargetTypeRequirements& UChaosClothAssetEditorMode::GetToolTargetRequirements()
@@ -832,37 +816,13 @@ FBox UChaosClothAssetEditorMode::SelectionBoundingBox() const
 
 FBox UChaosClothAssetEditorMode::PreviewBoundingBox() const
 {
-	FBoxSphereBounds TotalBounds(ForceInit);
-
-	bool bAnyValid = false;
-
-	if (ClothComponent && ClothComponent->GetClothAsset())
+	if (ClothComponent)
 	{
-		const UChaosClothAsset* ChaosClothAsset = ClothComponent->GetClothAsset();
-		const UE::Chaos::ClothAsset::FClothConstAdapter ClothAdapter(ChaosClothAsset->GetClothCollection());
-
-		for (int32 LodIndex = 0; LodIndex < ClothAdapter.GetNumLods(); ++LodIndex)
-		{
-			const UE::Chaos::ClothAsset::FClothLodConstAdapter ClothLodAdapter = ClothAdapter.GetLod(LodIndex);
-
-			for (int32 PatternIndex = 0; PatternIndex < ClothLodAdapter.GetNumPatterns(); ++PatternIndex)
-			{
-				const UE::Chaos::ClothAsset::FClothPatternConstAdapter ClothPatternAdapter = ClothLodAdapter.GetPattern(PatternIndex);
-				FBoxSphereBounds BoxSphere = ChaosClothAssetEditorModeLocals::RenderMeshBounds(ClothPatternAdapter);
-				TotalBounds = TotalBounds + BoxSphere;
-				bAnyValid = true;
-			}
-		}
+		FBoxSphereBounds Bounds = ClothComponent->Bounds;
+		return Bounds.GetBox();
 	}
 
-	if (bAnyValid)
-	{
-		return TotalBounds.GetBox();
-	}
-	else
-	{
-		return FBox(ForceInitToZero);
-	}
+	return FBox(ForceInitToZero);
 }
 
 bool UChaosClothAssetEditorMode::IsPattern2DModeActive() const
diff --git a/Engine/Plugins/Experimental/ChaosClothAssetEditor/Source/ChaosClothAssetEditor/Private/ChaosClothAsset/SClothEditor3DViewport.cpp b/Engine/Plugins/Experimental/ChaosClothAssetEditor/Source/ChaosClothAssetEditor/Private/ChaosClothAsset/SClothEditor3DViewport.cpp
--- a/Engine/Plugins/Experimental/ChaosClothAssetEditor/Source/ChaosClothAssetEditor/Private/ChaosClothAsset/SClothEditor3DViewport.cpp
+++ b/Engine/Plugins/Experimental/ChaosClothAssetEditor/Source/ChaosClothAssetEditor/Private/ChaosClothAsset/SClothEditor3DViewport.cpp
@@ -82,4 +82,10 @@ TSharedPtr<SWidget> SChaosClothAssetEditor3DViewport::MakeViewportToolbar()
 		.CommandList(CommandList);
 }
 
+void SChaosClothAssetEditor3DViewport::OnFocusViewportToSelection()
+{
+	const FBox PreviewBoundingBox = StaticCastSharedPtr<FChaosClothAssetEditor3DViewportClient>(Client)->PreviewBoundingBox();
+	Client->FocusViewportOnBox(PreviewBoundingBox);
+}
+
 #undef LOCTEXT_NAMESPACE
diff --git a/Engine/Plugins/Experimental/ChaosClothAssetEditor/Source/ChaosClothAssetEditor/Public/ChaosClothAsset/ClothEditor3DViewportClient.h b/Engine/Plugins/Experimental/ChaosClothAssetEditor/Source/ChaosClothAssetEditor/Public/ChaosClothAsset/ClothEditor3DViewportClient.h
--- a/Engine/Plugins/Experimental/ChaosClothAssetEditor/Source/ChaosClothAssetEditor/Public/ChaosClothAsset/ClothEditor3DViewportClient.h
+++ b/Engine/Plugins/Experimental/ChaosClothAssetEditor/Source/ChaosClothAssetEditor/Public/ChaosClothAsset/ClothEditor3DViewportClient.h
@@ -40,6 +40,7 @@ public:
 	void ResumeSimulation();
 	bool IsSimulationSuspended() const;
 
+	FBox PreviewBoundingBox() const;
 
 private:
 
diff --git a/Engine/Plugins/Experimental/ChaosClothAssetEditor/Source/ChaosClothAssetEditor/Public/ChaosClothAsset/SClothEditor3DViewport.h b/Engine/Plugins/Experimental/ChaosClothAssetEditor/Source/ChaosClothAssetEditor/Public/ChaosClothAsset/SClothEditor3DViewport.h
--- a/Engine/Plugins/Experimental/ChaosClothAssetEditor/Source/ChaosClothAssetEditor/Public/ChaosClothAsset/SClothEditor3DViewport.h
+++ b/Engine/Plugins/Experimental/ChaosClothAssetEditor/Source/ChaosClothAssetEditor/Public/ChaosClothAsset/SClothEditor3DViewport.h
@@ -15,5 +15,5 @@ public:
 	// SAssetEditorViewport
 	virtual void BindCommands() override;
 	virtual TSharedPtr<SWidget> MakeViewportToolbar() override;
-
+	virtual void OnFocusViewportToSelection() override;
 };
