Cloth Editor: reduce FOV in both viewports to reduce wide angle lens distortion

diff --git a/Engine/Plugins/Experimental/ChaosClothAssetEditor/Source/ChaosClothAssetEditor/Private/ChaosClothAsset/ClothEditorToolkit.cpp b/Engine/Plugins/Experimental/ChaosClothAssetEditor/Source/ChaosClothAssetEditor/Private/ChaosClothAsset/ClothEditorToolkit.cpp
--- a/Engine/Plugins/Experimental/ChaosClothAssetEditor/Source/ChaosClothAssetEditor/Private/ChaosClothAsset/ClothEditorToolkit.cpp
+++ b/Engine/Plugins/Experimental/ChaosClothAssetEditor/Source/ChaosClothAssetEditor/Private/ChaosClothAsset/ClothEditorToolkit.cpp
@@ -358,6 +358,9 @@ void FChaosClothAssetEditorToolkit::PostInitAssetEditor()
 		// toward camera for z ordering because non-translucent materials start occluding themselves (once
 		// the component bounds are behind the displaced geometry).
 		Client->EngineShowFlags.SetDisableOcclusionQueries(true);
+
+		// Default FOV of 90 degrees causes a fair bit of lens distortion, especially noticeable with smaller viewports
+		Client->ViewFOV = 45.0;
 	};
 
 	{
