ClothEditor:
- Add an Options menu to both cloth viewport toolbars. Right now the only setting is the camera FOV for perspective views, but we can add more options as required
- Created a ViewportToolBarBase class that populates the Options menu for both viewports in the editor
- Set the default camera FOV to 45 degrees

diff --git a/Engine/Plugins/Experimental/ChaosClothAssetEditor/Source/ChaosClothAssetEditor/Private/ChaosClothAsset/SClothEditor3DViewport.cpp b/Engine/Plugins/Experimental/ChaosClothAssetEditor/Source/ChaosClothAssetEditor/Private/ChaosClothAsset/SClothEditor3DViewport.cpp
--- a/Engine/Plugins/Experimental/ChaosClothAssetEditor/Source/ChaosClothAssetEditor/Private/ChaosClothAsset/SClothEditor3DViewport.cpp
+++ b/Engine/Plugins/Experimental/ChaosClothAssetEditor/Source/ChaosClothAssetEditor/Private/ChaosClothAsset/SClothEditor3DViewport.cpp
@@ -88,4 +88,19 @@ void SChaosClothAssetEditor3DViewport::OnFocusViewportToSelection()
 	Client->FocusViewportOnBox(PreviewBoundingBox);
 }
 
+TSharedRef<class SEditorViewport> SChaosClothAssetEditor3DViewport::GetViewportWidget()
+{
+	return SharedThis(this);
+}
+
+TSharedPtr<FExtender> SChaosClothAssetEditor3DViewport::GetExtenders() const
+{
+	TSharedPtr<FExtender> Result(MakeShareable(new FExtender));
+	return Result;
+}
+
+void SChaosClothAssetEditor3DViewport::OnFloatingButtonClicked()
+{
+}
+
 #undef LOCTEXT_NAMESPACE
diff --git a/Engine/Plugins/Experimental/ChaosClothAssetEditor/Source/ChaosClothAssetEditor/Private/ChaosClothAsset/SClothEditor3DViewportToolBar.cpp b/Engine/Plugins/Experimental/ChaosClothAssetEditor/Source/ChaosClothAssetEditor/Private/ChaosClothAsset/SClothEditor3DViewportToolBar.cpp
--- a/Engine/Plugins/Experimental/ChaosClothAssetEditor/Source/ChaosClothAssetEditor/Private/ChaosClothAsset/SClothEditor3DViewportToolBar.cpp
+++ b/Engine/Plugins/Experimental/ChaosClothAssetEditor/Source/ChaosClothAssetEditor/Private/ChaosClothAsset/SClothEditor3DViewportToolBar.cpp
@@ -17,6 +17,8 @@
 
 void SChaosClothAssetEditor3DViewportToolBar::Construct(const FArguments& InArgs, TSharedPtr<class SChaosClothAssetEditor3DViewport> InChaosClothAssetEditor3DViewport)
 {
+	SCommonEditorViewportToolbarBase::Construct(SCommonEditorViewportToolbarBase::FArguments(), InChaosClothAssetEditor3DViewport);
+
 	ChaosClothAssetEditor3DViewportPtr = InChaosClothAssetEditor3DViewport;
 	CommandList = InArgs._CommandList;
 
@@ -33,6 +35,13 @@ void SChaosClothAssetEditor3DViewportToolBar::Construct(const FArguments& InArgs
 		]
 	];
 
+	MainBoxPtr->AddSlot()
+		.AutoWidth()
+		.Padding(ToolbarSlotPadding)
+		[
+			MakeOptionsMenu()
+		];
+
 	MainBoxPtr->AddSlot()
 		.Padding(ToolbarSlotPadding)
 		.HAlign(HAlign_Left)
@@ -46,8 +55,15 @@ void SChaosClothAssetEditor3DViewportToolBar::Construct(const FArguments& InArgs
 		[
 			MakeToolBar(InArgs._Extenders)
 		];
+}
 
-	SViewportToolBar::Construct(SViewportToolBar::FArguments());
+TSharedRef<SWidget> SChaosClothAssetEditor3DViewportToolBar::MakeOptionsMenu()
+{
+	return SNew(SEditorViewportToolbarMenu)
+		.ParentToolBar(SharedThis(this))
+		.Cursor(EMouseCursor::Default)
+		.Image("EditorViewportToolBar.OptionsDropdown")
+		.OnGetMenuContent(this, &SChaosClothAssetEditor3DViewportToolBar::GenerateClothViewportOptionsMenu);
 }
 
 TSharedRef<SWidget> SChaosClothAssetEditor3DViewportToolBar::MakeDisplayToolBar(const TSharedPtr<FExtender> InExtenders)
diff --git a/Engine/Plugins/Experimental/ChaosClothAssetEditor/Source/ChaosClothAssetEditor/Private/ChaosClothAsset/SClothEditorRestSpaceViewport.cpp b/Engine/Plugins/Experimental/ChaosClothAssetEditor/Source/ChaosClothAssetEditor/Private/ChaosClothAsset/SClothEditorRestSpaceViewport.cpp
--- a/Engine/Plugins/Experimental/ChaosClothAssetEditor/Source/ChaosClothAssetEditor/Private/ChaosClothAsset/SClothEditorRestSpaceViewport.cpp
+++ b/Engine/Plugins/Experimental/ChaosClothAssetEditor/Source/ChaosClothAssetEditor/Private/ChaosClothAsset/SClothEditorRestSpaceViewport.cpp
@@ -7,9 +7,6 @@
 #include "ChaosClothAsset/ClothEditorMode.h"
 #include "EditorModeManager.h"
 
-// TODO evntually:
-//#include "ClothEditorRestSpaceViewportClient.h"
-
 #define LOCTEXT_NAMESPACE "SChaosClothAssetEditorRestSpaceViewport"
 
 void SChaosClothAssetEditorRestSpaceViewport::BindCommands()
@@ -63,4 +60,19 @@ void SChaosClothAssetEditorRestSpaceViewport::OnFocusViewportToSelection()
 	}
 }
 
+TSharedRef<class SEditorViewport> SChaosClothAssetEditorRestSpaceViewport::GetViewportWidget()
+{
+	return SharedThis(this);
+}
+
+TSharedPtr<FExtender> SChaosClothAssetEditorRestSpaceViewport::GetExtenders() const
+{
+	TSharedPtr<FExtender> Result(MakeShareable(new FExtender));
+	return Result;
+}
+
+void SChaosClothAssetEditorRestSpaceViewport::OnFloatingButtonClicked()
+{
+}
+
 #undef LOCTEXT_NAMESPACE
diff --git a/Engine/Plugins/Experimental/ChaosClothAssetEditor/Source/ChaosClothAssetEditor/Private/ChaosClothAsset/SClothEditorRestSpaceViewportToolBar.cpp b/Engine/Plugins/Experimental/ChaosClothAssetEditor/Source/ChaosClothAssetEditor/Private/ChaosClothAsset/SClothEditorRestSpaceViewportToolBar.cpp
--- a/Engine/Plugins/Experimental/ChaosClothAssetEditor/Source/ChaosClothAssetEditor/Private/ChaosClothAsset/SClothEditorRestSpaceViewportToolBar.cpp
+++ b/Engine/Plugins/Experimental/ChaosClothAssetEditor/Source/ChaosClothAssetEditor/Private/ChaosClothAsset/SClothEditorRestSpaceViewportToolBar.cpp
@@ -15,6 +15,8 @@
 
 void SChaosClothAssetEditorRestSpaceViewportToolBar::Construct(const FArguments& InArgs, TSharedPtr<SChaosClothAssetEditorRestSpaceViewport> InChaosClothAssetEditorViewport)
 {
+	SCommonEditorViewportToolbarBase::Construct(SCommonEditorViewportToolbarBase::FArguments(), InChaosClothAssetEditorViewport);
+
 	ChaosClothAssetEditorRestSpaceViewportPtr = InChaosClothAssetEditorViewport;
 	CommandList = InArgs._CommandList;
 
@@ -31,6 +33,13 @@ void SChaosClothAssetEditorRestSpaceViewportToolBar::Construct(const FArguments&
 		]
 	];
 
+	MainBoxPtr->AddSlot()
+		.Padding(ToolbarSlotPadding)
+		.AutoWidth()
+		[
+			MakeOptionsMenu()
+		];
+
 	MainBoxPtr->AddSlot()
 		.Padding(ToolbarSlotPadding)
 		.HAlign(HAlign_Left)
@@ -44,8 +53,15 @@ void SChaosClothAssetEditorRestSpaceViewportToolBar::Construct(const FArguments&
 		[
 			MakeToolBar(InArgs._Extenders)
 		];
+}
 
-	SViewportToolBar::Construct(SViewportToolBar::FArguments());
+TSharedRef<SWidget> SChaosClothAssetEditorRestSpaceViewportToolBar::MakeOptionsMenu()
+{
+	return SNew(SEditorViewportToolbarMenu)
+		.ParentToolBar(SharedThis(this))
+		.Cursor(EMouseCursor::Default)
+		.Image("EditorViewportToolBar.OptionsDropdown")
+		.OnGetMenuContent(this, &SChaosClothAssetEditorRestSpaceViewportToolBar::GenerateClothViewportOptionsMenu);
 }
 
 TSharedRef<SWidget> SChaosClothAssetEditorRestSpaceViewportToolBar::MakeDisplayToolBar(const TSharedPtr<FExtender> InExtenders)
diff --git a/Engine/Plugins/Experimental/ChaosClothAssetEditor/Source/ChaosClothAssetEditor/Private/ChaosClothAsset/SClothEditorViewportToolBarBase.cpp b/Engine/Plugins/Experimental/ChaosClothAssetEditor/Source/ChaosClothAssetEditor/Private/ChaosClothAsset/SClothEditorViewportToolBarBase.cpp
new file mode 100644
--- /dev/null
+++ b/Engine/Plugins/Experimental/ChaosClothAssetEditor/Source/ChaosClothAssetEditor/Private/ChaosClothAsset/SClothEditorViewportToolBarBase.cpp
@@ -0,0 +1,76 @@
+// Copyright Epic Games, Inc. All Rights Reserved.
+
+#include "ChaosClothAsset/SClothEditorViewportToolBarBase.h"
+#include "SEditorViewport.h"
+#include "Widgets/Input/SSpinBox.h"
+#include "Widgets/Layout/SBorder.h"
+
+#define LOCTEXT_NAMESPACE "SChaosClothAssetEditorViewportToolBarBase"
+
+TSharedRef<SWidget> SChaosClothAssetEditorViewportToolBarBase::GenerateClothViewportOptionsMenu() const
+{
+	// (See also SCommonEditorViewportToolbarBase::GenerateOptionsMenu)
+
+	GetInfoProvider().OnFloatingButtonClicked();
+	TSharedRef<SEditorViewport> ViewportRef = GetInfoProvider().GetViewportWidget();
+	
+	const bool bIsPerspective = GetViewportClient().GetViewportType() == LVT_Perspective;
+	const bool bInShouldCloseWindowAfterMenuSelection = true;
+	FMenuBuilder OptionsMenuBuilder(bInShouldCloseWindowAfterMenuSelection, ViewportRef->GetCommandList());
+	{
+		OptionsMenuBuilder.BeginSection("ClothEditorViewportOptions", LOCTEXT("OptionsMenuHeader", "Viewport Options"));
+		{
+			if (bIsPerspective)
+			{
+				OptionsMenuBuilder.AddWidget(GenerateFOVMenu(), LOCTEXT("FOVAngle", "Field of View (H)"));
+			}
+		}
+		OptionsMenuBuilder.EndSection();
+
+		// TODO: Add any other options for our cloth viewports
+
+		ExtendOptionsMenu(OptionsMenuBuilder);
+	}
+
+	return OptionsMenuBuilder.MakeWidget();
+}
+
+TSharedRef<SWidget> SChaosClothAssetEditorViewportToolBarBase::GenerateFOVMenu() const
+{
+	// Re-implementation of inaccessible SCommonEditorViewportToolbarBase::GenerateFOVMenu()
+
+	const float FOVMin = 5.f;
+	const float FOVMax = 170.f;
+
+	return 
+		SNew(SBox)
+		.HAlign(HAlign_Right)
+		[
+			SNew(SBox)
+			.Padding(FMargin(4.0f, 0.0f, 0.0f, 0.0f))
+			.WidthOverride(100.0f)
+			[
+				SNew(SBorder)
+				.BorderImage(FAppStyle::Get().GetBrush("Menu.WidgetBorder"))
+				.Padding(FMargin(1.0f))
+				[
+					SNew(SSpinBox<float>)
+					.Style(&FAppStyle::Get(), "Menu.SpinBox")
+					.Font(FAppStyle::GetFontStyle(TEXT("MenuItem.Font")))
+					.MinValue(FOVMin)
+					.MaxValue(FOVMax)
+					.Value(this, &SChaosClothAssetEditorViewportToolBarBase::OnGetFOVValue)
+					.OnValueChanged(this, &SChaosClothAssetEditorViewportToolBarBase::OnFOVValueChanged)
+				]
+			]
+		];
+}
+
+float SChaosClothAssetEditorViewportToolBarBase::OnGetFOVValue() const
+{
+	// Re-implementation of inaccessible SCommonEditorViewportToolbarBase::OnGetFOVValue()
+
+	return GetViewportClient().ViewFOV;
+}
+
+#undef LOCTEXT_NAMESPACE
diff --git a/Engine/Plugins/Experimental/ChaosClothAssetEditor/Source/ChaosClothAssetEditor/Public/ChaosClothAsset/SClothEditor3DViewport.h b/Engine/Plugins/Experimental/ChaosClothAssetEditor/Source/ChaosClothAssetEditor/Public/ChaosClothAsset/SClothEditor3DViewport.h
--- a/Engine/Plugins/Experimental/ChaosClothAssetEditor/Source/ChaosClothAssetEditor/Public/ChaosClothAsset/SClothEditor3DViewport.h
+++ b/Engine/Plugins/Experimental/ChaosClothAssetEditor/Source/ChaosClothAssetEditor/Public/ChaosClothAsset/SClothEditor3DViewport.h
@@ -4,16 +4,24 @@
 
 #include "CoreMinimal.h"
 #include "SAssetEditorViewport.h"
+#include "SCommonEditorViewportToolbarBase.h"
 
 /**
  * Viewport used for 3D preview in cloth editor. Has a custom toolbar overlay at the top.
  */
-class CHAOSCLOTHASSETEDITOR_API SChaosClothAssetEditor3DViewport : public SAssetEditorViewport 
+class CHAOSCLOTHASSETEDITOR_API SChaosClothAssetEditor3DViewport : public SAssetEditorViewport, public ICommonEditorViewportToolbarInfoProvider
 {
 
 public:
 	// SAssetEditorViewport
 	virtual void BindCommands() override;
 	virtual TSharedPtr<SWidget> MakeViewportToolbar() override;
+
 	virtual void OnFocusViewportToSelection() override;
+
+	// ICommonEditorViewportToolbarInfoProvider
+	virtual TSharedRef<class SEditorViewport> GetViewportWidget() override;
+	virtual TSharedPtr<FExtender> GetExtenders() const override;
+	virtual void OnFloatingButtonClicked() override;
+
 };
diff --git a/Engine/Plugins/Experimental/ChaosClothAssetEditor/Source/ChaosClothAssetEditor/Public/ChaosClothAsset/SClothEditor3DViewportToolBar.h b/Engine/Plugins/Experimental/ChaosClothAssetEditor/Source/ChaosClothAssetEditor/Public/ChaosClothAsset/SClothEditor3DViewportToolBar.h
--- a/Engine/Plugins/Experimental/ChaosClothAssetEditor/Source/ChaosClothAssetEditor/Public/ChaosClothAsset/SClothEditor3DViewportToolBar.h
+++ b/Engine/Plugins/Experimental/ChaosClothAssetEditor/Source/ChaosClothAssetEditor/Public/ChaosClothAsset/SClothEditor3DViewportToolBar.h
@@ -2,7 +2,7 @@
 
 #pragma once
 
-#include "SViewportToolBar.h"
+#include "SClothEditorViewportToolBarBase.h"
 
 class FExtender;
 class FUICommandList;
@@ -11,7 +11,7 @@ class SChaosClothAssetEditor3DViewport;
 /**
  * Toolbar that shows up at the top of the 3d viewport
  */
-class CHAOSCLOTHASSETEDITOR_API SChaosClothAssetEditor3DViewportToolBar : public SViewportToolBar
+class CHAOSCLOTHASSETEDITOR_API SChaosClothAssetEditor3DViewportToolBar : public SChaosClothAssetEditorViewportToolBarBase
 {
 public:
 	SLATE_BEGIN_ARGS(SChaosClothAssetEditor3DViewportToolBar) {}
@@ -22,6 +22,7 @@ public:
 	void Construct(const FArguments& InArgs, TSharedPtr<class SChaosClothAssetEditor3DViewport> InChaosClothAssetEditor3DViewport);
 
 private:
+	TSharedRef<SWidget> MakeOptionsMenu();
 	TSharedRef<SWidget> MakeDisplayToolBar(const TSharedPtr<FExtender> InExtenders);
 	TSharedRef<SWidget> MakeToolBar(const TSharedPtr<FExtender> InExtenders);
 
diff --git a/Engine/Plugins/Experimental/ChaosClothAssetEditor/Source/ChaosClothAssetEditor/Public/ChaosClothAsset/SClothEditorRestSpaceViewport.h b/Engine/Plugins/Experimental/ChaosClothAssetEditor/Source/ChaosClothAssetEditor/Public/ChaosClothAsset/SClothEditorRestSpaceViewport.h
--- a/Engine/Plugins/Experimental/ChaosClothAssetEditor/Source/ChaosClothAssetEditor/Public/ChaosClothAsset/SClothEditorRestSpaceViewport.h
+++ b/Engine/Plugins/Experimental/ChaosClothAssetEditor/Source/ChaosClothAssetEditor/Public/ChaosClothAsset/SClothEditorRestSpaceViewport.h
@@ -4,8 +4,9 @@
 
 #include "CoreMinimal.h"
 #include "SBaseCharacterFXEditorViewport.h"
+#include "SCommonEditorViewportToolbarBase.h"
 
-class CHAOSCLOTHASSETEDITOR_API SChaosClothAssetEditorRestSpaceViewport : public SBaseCharacterFXEditorViewport
+class CHAOSCLOTHASSETEDITOR_API SChaosClothAssetEditorRestSpaceViewport : public SBaseCharacterFXEditorViewport, public ICommonEditorViewportToolbarInfoProvider
 {
 public:
 
@@ -13,4 +14,10 @@ public:
 	virtual void BindCommands() override;
 	virtual TSharedPtr<SWidget> MakeViewportToolbar() override;
 	virtual void OnFocusViewportToSelection() override;
+
+	// ICommonEditorViewportToolbarInfoProvider
+	virtual TSharedRef<class SEditorViewport> GetViewportWidget() override;
+	virtual TSharedPtr<FExtender> GetExtenders() const override;
+	virtual void OnFloatingButtonClicked() override;
+
 };
diff --git a/Engine/Plugins/Experimental/ChaosClothAssetEditor/Source/ChaosClothAssetEditor/Public/ChaosClothAsset/SClothEditorRestSpaceViewportToolBar.h b/Engine/Plugins/Experimental/ChaosClothAssetEditor/Source/ChaosClothAssetEditor/Public/ChaosClothAsset/SClothEditorRestSpaceViewportToolBar.h
--- a/Engine/Plugins/Experimental/ChaosClothAssetEditor/Source/ChaosClothAssetEditor/Public/ChaosClothAsset/SClothEditorRestSpaceViewportToolBar.h
+++ b/Engine/Plugins/Experimental/ChaosClothAssetEditor/Source/ChaosClothAssetEditor/Public/ChaosClothAsset/SClothEditorRestSpaceViewportToolBar.h
@@ -2,7 +2,7 @@
 
 #pragma once
 
-#include "SViewportToolBar.h"
+#include "SClothEditorViewportToolBarBase.h"
 
 class FExtender;
 class FUICommandList;
@@ -11,7 +11,7 @@ class SChaosClothAssetEditorRestSpaceViewport;
 /**
  * Toolbar that shows up at the top of the rest space viewport
  */
-class CHAOSCLOTHASSETEDITOR_API SChaosClothAssetEditorRestSpaceViewportToolBar : public SViewportToolBar
+class CHAOSCLOTHASSETEDITOR_API SChaosClothAssetEditorRestSpaceViewportToolBar : public SChaosClothAssetEditorViewportToolBarBase
 {
 public:
 	SLATE_BEGIN_ARGS(SChaosClothAssetEditorRestSpaceViewportToolBar) {}
@@ -22,6 +22,7 @@ public:
 	void Construct(const FArguments& InArgs, TSharedPtr<class SChaosClothAssetEditorRestSpaceViewport> InChaosClothAssetEditorViewport);
 
 private:
+	TSharedRef<SWidget> MakeOptionsMenu();
 	TSharedRef<SWidget> MakeDisplayToolBar(const TSharedPtr<FExtender> InExtenders);
 	TSharedRef<SWidget> MakeToolBar(const TSharedPtr<FExtender> InExtenders);
 
diff --git a/Engine/Plugins/Experimental/ChaosClothAssetEditor/Source/ChaosClothAssetEditor/Public/ChaosClothAsset/SClothEditorViewportToolBarBase.h b/Engine/Plugins/Experimental/ChaosClothAssetEditor/Source/ChaosClothAssetEditor/Public/ChaosClothAsset/SClothEditorViewportToolBarBase.h
new file mode 100644
--- /dev/null
+++ b/Engine/Plugins/Experimental/ChaosClothAssetEditor/Source/ChaosClothAssetEditor/Public/ChaosClothAsset/SClothEditorViewportToolBarBase.h
@@ -0,0 +1,17 @@
+// Copyright Epic Games, Inc. All Rights Reserved.
+
+#pragma once
+
+#include "SCommonEditorViewportToolbarBase.h"
+
+/// Common code for toolbars in the Cloth Editor
+class CHAOSCLOTHASSETEDITOR_API SChaosClothAssetEditorViewportToolBarBase : public SCommonEditorViewportToolbarBase
+{
+protected:
+	TSharedRef<SWidget> GenerateClothViewportOptionsMenu() const;
+
+private:
+	TSharedRef<SWidget> GenerateFOVMenu() const;
+	float OnGetFOVValue() const;
+};
+
