Cloth Editor: Adding a Cloth Collection "Outliner" widget. The outliner reads all Attribute names and Values from a Cloth Collection (for a specific attribute Group) and puts them into a ListView. The user can select which Group to view via a ComboBox in the Editor.

diff --git a/Engine/Plugins/Experimental/ChaosClothAssetEditor/Source/ChaosClothAssetEditor/Private/ChaosClothAsset/ClothEditorToolkit.cpp b/Engine/Plugins/Experimental/ChaosClothAssetEditor/Source/ChaosClothAssetEditor/Private/ChaosClothAsset/ClothEditorToolkit.cpp
--- a/Engine/Plugins/Experimental/ChaosClothAssetEditor/Source/ChaosClothAssetEditor/Private/ChaosClothAsset/ClothEditorToolkit.cpp
+++ b/Engine/Plugins/Experimental/ChaosClothAssetEditor/Source/ChaosClothAssetEditor/Private/ChaosClothAsset/ClothEditorToolkit.cpp
@@ -9,6 +9,9 @@
 #include "ChaosClothAsset/ClothEditor3DViewportClient.h"
 #include "ChaosClothAsset/SClothEditorRestSpaceViewport.h"
 #include "ChaosClothAsset/ClothEditorRestSpaceViewportClient.h"
+#include "ChaosClothAsset/ClothComponent.h"
+#include "ChaosClothAsset/ClothAsset.h"
+#include "ChaosClothAsset/ClothCollection.h"
 #include "Framework/Docking/LayoutExtender.h"
 #include "AssetEditorModeManager.h"
 #include "AdvancedPreviewScene.h"
@@ -17,10 +20,12 @@
 #include "Misc/MessageDialog.h"
 #include "Modules/ModuleManager.h"
 #include "Styling/AppStyle.h"
+#include "ChaosClothAsset/SClothCollectionOutliner.h"
 
 #define LOCTEXT_NAMESPACE "ChaosClothAssetEditorToolkit"
 
 const FName FChaosClothAssetEditorToolkit::ClothPreviewTabID(TEXT("ChaosClothAssetEditor_ClothPreviewTab"));
+const FName FChaosClothAssetEditorToolkit::OutlinerTabID(TEXT("ChaosClothAssetEditor_OutlinerTab"));
 
 FChaosClothAssetEditorToolkit::FChaosClothAssetEditorToolkit(UAssetEditor* InOwningAssetEditor)
 	: FBaseCharacterFXEditorToolkit(InOwningAssetEditor, FName("ChaosClothAssetEditor"))
@@ -71,6 +76,7 @@ FChaosClothAssetEditorToolkit::FChaosClothAssetEditorToolkit(UAssetEditor* InOwn
 					FTabManager::NewStack()
 					->SetSizeCoefficient(0.15f)
 					->AddTab(DetailsTabID, ETabState::OpenedTab)
+					->AddTab(OutlinerTabID, ETabState::OpenedTab)
 					->SetExtensionId("Details")
 					->SetHideTabWell(true)
 				)
@@ -189,6 +195,12 @@ void FChaosClothAssetEditorToolkit::RegisterTabSpawners(const TSharedRef<FTabMan
 		.SetDisplayName(LOCTEXT("Details", "Details"))
 		.SetGroup(AssetEditorTabsCategory.ToSharedRef())
 		.SetIcon(FSlateIcon(FAppStyle::GetAppStyleSetName(), "LevelEditor.Tabs.Details"));
+
+	InTabManager->RegisterTabSpawner(OutlinerTabID, FOnSpawnTab::CreateSP(this, &FChaosClothAssetEditorToolkit::SpawnTab_Outliner))
+		.SetDisplayName(LOCTEXT("Outliner", "Outliner"))
+		.SetGroup(AssetEditorTabsCategory.ToSharedRef())
+		.SetIcon(FSlateIcon(FAppStyle::GetAppStyleSetName(), "LevelEditor.Tabs.Outliner"));
+
 }
 
 void FChaosClothAssetEditorToolkit::UnregisterTabSpawners(const TSharedRef<FTabManager>& InTabManager)
@@ -281,6 +293,56 @@ TSharedRef<SDockTab> FChaosClothAssetEditorToolkit::SpawnTab_ClothPreview(const
 	return DockableTab;
 }
 
+TSharedRef<SDockTab> FChaosClothAssetEditorToolkit::SpawnTab_Outliner(const FSpawnTabArgs& Args)
+{
+
+	TSharedRef<SDockTab> DockableTab = SNew(SDockTab)
+	[
+		SNew(SVerticalBox)
+
+		+ SVerticalBox::Slot()
+		.AutoHeight()
+		[
+			SAssignNew(SelectedGroupNameComboBox, SComboBox<FName>)
+			.OptionsSource(&ClothCollectionGroupNames)
+			.OnSelectionChanged(SComboBox<FName>::FOnSelectionChanged::CreateLambda(
+				[this](FName SelectedName, ESelectInfo::Type)
+				{
+					if (OutlinerView)
+					{
+						OutlinerView->SetSelectedGroupName(SelectedName);	// this will also rebuild the table
+					}
+				}))
+			.OnGenerateWidget(SComboBox<FName>::FOnGenerateWidget::CreateLambda(
+				[](FName Item)
+				{
+					return SNew(STextBlock)
+						.Text(FText::FromName(Item));
+				}))
+			[
+				SNew(STextBlock)
+				.Text_Lambda([this]()
+				{
+					if (OutlinerView)
+					{
+						return FText::FromName(OutlinerView->GetSelectedGroupName());
+					}
+					else
+					{
+						return LOCTEXT("NoSelectedGroup", "");
+					}
+				})
+			]
+		]
+
+		+ SVerticalBox::Slot()
+		[
+			SAssignNew(OutlinerView, SClothCollectionOutliner)
+		]
+	];
+	return DockableTab;
+}
+
 //// This is bound in RegisterTabSpawners() to create the panel on the left. The panel is filled in by the mode.
 //TSharedRef<SDockTab> FChaosClothAssetEditorToolkit::SpawnTab_InteractiveToolsPanel(const FSpawnTabArgs& Args)
 //{
@@ -446,6 +508,35 @@ void FChaosClothAssetEditorToolkit::InitDetailsViewPanel()
 		ensure(ObjectToEditInDetailsView->HasAnyFlags(RF_Transactional));		// Ensure all objects are transactable for undo/redo in the details panel
 		SetEditingObject(ObjectToEditInDetailsView);
 	}
+
+	const UChaosClothAssetEditorMode* const ClothMode = CastChecked<UChaosClothAssetEditorMode>(EditorModeManager->GetActiveScriptableMode(UChaosClothAssetEditorMode::EM_ChaosClothAssetEditorModeId));
+	check(ClothMode);
+
+	const bool bHasOutliner = OutlinerView.IsValid();
+
+	bool bHasClothCollection = false;
+	if (ClothMode->ClothComponent)
+	{
+		if (ClothMode->ClothComponent->GetClothAsset())
+		{
+			if (ClothMode->ClothComponent->GetClothAsset()->GetClothCollection())
+			{
+				bHasClothCollection = true;
+			}
+		}
+	}
+
+	if (bHasOutliner && bHasClothCollection)
+	{
+		const TSharedPtr<UE::Chaos::ClothAsset::FClothCollection> ClothCollection = ClothMode->ClothComponent->GetClothAsset()->GetClothCollection();
+		const FName& SelectedGroupName = ClothCollection->SimVerticesGroup;
+		
+		OutlinerView->SetClothCollection(ClothCollection);
+		OutlinerView->SetSelectedGroupName(SelectedGroupName);
+
+		ClothCollectionGroupNames = ClothCollection->GroupNames();
+	}
+
 }
 
 #undef LOCTEXT_NAMESPACE
diff --git a/Engine/Plugins/Experimental/ChaosClothAssetEditor/Source/ChaosClothAssetEditor/Private/ChaosClothAsset/SClothCollectionOutliner.cpp b/Engine/Plugins/Experimental/ChaosClothAssetEditor/Source/ChaosClothAssetEditor/Private/ChaosClothAsset/SClothCollectionOutliner.cpp
new file mode 100644
--- /dev/null
+++ b/Engine/Plugins/Experimental/ChaosClothAssetEditor/Source/ChaosClothAssetEditor/Private/ChaosClothAsset/SClothCollectionOutliner.cpp
@@ -0,0 +1,253 @@
+// Copyright Epic Games, Inc. All Rights Reserved.
+
+#include "ChaosClothAsset/SClothCollectionOutliner.h"
+#include "ChaosClothAsset/ClothCollection.h"
+#include "Templates/EnableIf.h"
+
+#define LOCTEXT_NAMESPACE "ClothCollectionOutliner"
+
+const FName FClothCollectionHeaderData::ColumnZeroName = FName("ID");
+
+void SClothCollectionOutlinerRow::Construct(const FArguments& InArgs,
+	TSharedRef<STableViewBase> OwnerTableView,
+	const TSharedPtr<const FClothCollectionHeaderData>& InHeaderData,
+	const TSharedPtr<const FClothCollectionItem>& InItemToEdit)
+{
+	HeaderData = InHeaderData;
+	Item = InItemToEdit;
+
+	SMultiColumnTableRow<TSharedPtr<const FClothCollectionItem>>::Construct(
+		FSuperRowType::FArguments()
+		, OwnerTableView);
+}
+
+
+TSharedRef<SWidget> SClothCollectionOutlinerRow::GenerateWidgetForColumn(const FName& ColumnName)
+{
+	int32 FoundIndex;
+	if (HeaderData->AttributeNames.Find(ColumnName, FoundIndex))
+	{
+		const FString& AttrValue = Item->AttributeValues[FoundIndex];
+
+		return SNew(SHorizontalBox)
+			+ SHorizontalBox::Slot()
+			.AutoWidth()
+			[
+				SNew(STextBlock)
+				.Text(FText::FromString(AttrValue))
+			];
+	}
+
+	return SNullWidget::NullWidget;
+}
+
+
+void SClothCollectionOutliner::Construct(const FArguments& InArgs)
+{
+	ClothCollection = InArgs._ClothCollection;
+	SelectedGroupName = InArgs._SelectedGroupName;
+
+	HeaderRowWidget =
+		SNew(SHeaderRow)
+		.Visibility(EVisibility::Visible);
+
+	if (ClothCollection.IsValid())
+	{
+		RegenerateHeader();
+		RepopulateListView();
+	}
+
+	ChildSlot
+	[
+		SNew(SVerticalBox)
+		+ SVerticalBox::Slot()
+		.Padding(FMargin(0.0f, 3.f))
+		[
+			SNew(SOverlay)
+			+ SOverlay::Slot()
+			[
+				SAssignNew(ListView, SListView<TSharedPtr<const FClothCollectionItem>>)
+				.SelectionMode(ESelectionMode::None)
+				.ListItemsSource(&ListItems)
+				.OnGenerateRow(this, &SClothCollectionOutliner::GenerateRow)
+				.HeaderRow(HeaderRowWidget)
+			]
+		]
+	];
+}
+
+void SClothCollectionOutliner::SetClothCollection(TWeakPtr<UE::Chaos::ClothAsset::FClothCollection> InClothCollection)
+{
+	ClothCollection = InClothCollection;
+}
+
+void SClothCollectionOutliner::SetSelectedGroupName(const FName& InSelectedGroupName)
+{
+	SelectedGroupName = InSelectedGroupName;
+
+	RegenerateHeader();
+	RepopulateListView();
+}
+
+const FName& SClothCollectionOutliner::GetSelectedGroupName() const
+{
+	return SelectedGroupName;
+}
+
+void SClothCollectionOutliner::RegenerateHeader()
+{
+	constexpr float CustomFillWidth = 2.0f;
+
+	TSharedPtr<UE::Chaos::ClothAsset::FClothCollection> PinnedClothCollection = ClothCollection.Pin();
+
+	if (!PinnedClothCollection.IsValid())
+	{
+		return;
+	}
+
+	HeaderRowWidget->ClearColumns();
+
+	HeaderData = MakeShared<FClothCollectionHeaderData>();
+	HeaderData->AttributeNames.Add(FClothCollectionHeaderData::ColumnZeroName);
+	HeaderData->AttributeNames.Append(PinnedClothCollection->AttributeNames(SelectedGroupName));
+
+	for (int32 AttributeNameIndex = 0; AttributeNameIndex < HeaderData->AttributeNames.Num(); ++AttributeNameIndex)
+	{
+		const FName& AttrName = HeaderData->AttributeNames[AttributeNameIndex];
+
+		HeaderRowWidget->AddColumn(
+			SHeaderRow::Column(AttrName)
+			.DefaultLabel(FText::FromName(AttrName))
+			.FillWidth(CustomFillWidth)
+		);
+	}
+}
+
+namespace ClothCollectionOutlinerHelpers
+{
+	template<typename T>
+	FString AttributeValueToString(const T& Value)
+	{
+		return Value.ToString();
+	}
+
+	FString AttributeValueToString(float Value)
+	{
+		return FString::SanitizeFloat(Value);
+	}
+
+	FString AttributeValueToString(int32 Value)
+	{
+		return FString::FromInt(Value);
+	}
+
+	template<typename T>
+	FString AttributeValueToString(const TArray<T>& Array)
+	{
+		FString Out;
+		for ( int32 ArrayIndex = 0; ArrayIndex < Array.Num(); ++ArrayIndex )
+		{
+			Out += AttributeValueToString(Array[ArrayIndex]);
+
+			if (ArrayIndex != Array.Num() - 1)
+			{
+				Out += "; ";
+			}
+		}
+		return Out;
+	}
+
+	template<typename T>
+	FString AttributeValueToString(const UE::Chaos::ClothAsset::FClothCollection& ClothCollection, const FName& AttributeName, const FName& GroupName, int32 AttributeArrayIndex)
+	{
+		const TManagedArray<T>* const Array = ClothCollection.FindAttributeTyped<T>(AttributeName, GroupName);
+		if (Array == nullptr)
+		{
+			return FString("(Unknown Attribute)");
+		}
+
+		return AttributeValueToString((*Array)[AttributeArrayIndex]);
+	}
+
+	FString AttributeValueToString(const UE::Chaos::ClothAsset::FClothCollection& ClothCollection, const FName& AttributeName, const FName& GroupName, int32 AttributeArrayIndex)
+	{
+		const FManagedArrayCollection::EArrayType ArrayType = ClothCollection.GetAttributeType(AttributeName, GroupName);
+
+		FString ValueAsString;
+
+		switch (ArrayType)
+		{
+		case FManagedArrayCollection::EArrayType::FVectorType:
+			ValueAsString = AttributeValueToString<FVector3f>(ClothCollection, AttributeName, GroupName, AttributeArrayIndex);
+			break;
+		case FManagedArrayCollection::EArrayType::FVector2DType:
+			ValueAsString = AttributeValueToString<FVector2f>(ClothCollection, AttributeName, GroupName, AttributeArrayIndex);
+			break;
+		case FManagedArrayCollection::EArrayType::FFloatType:
+			ValueAsString = AttributeValueToString<float>(ClothCollection, AttributeName, GroupName, AttributeArrayIndex);
+			break;
+		case FManagedArrayCollection::EArrayType::FIntVectorType:
+			ValueAsString = AttributeValueToString<FIntVector>(ClothCollection, AttributeName, GroupName, AttributeArrayIndex);
+			break;
+		case FManagedArrayCollection::EArrayType::FVector2DArrayType:
+			ValueAsString = AttributeValueToString<TArray<FVector2f>>(ClothCollection, AttributeName, GroupName, AttributeArrayIndex);
+			break;
+		case FManagedArrayCollection::EArrayType::FLinearColorType:
+			ValueAsString = AttributeValueToString<FLinearColor>(ClothCollection, AttributeName, GroupName, AttributeArrayIndex);
+			break;
+		case FManagedArrayCollection::EArrayType::FInt32Type:
+			ValueAsString = AttributeValueToString<int32>(ClothCollection, AttributeName, GroupName, AttributeArrayIndex);
+			break;
+		default:
+			ensure(false);
+			ValueAsString = "(Unknown Data Type)";
+		}
+
+		return ValueAsString;
+	}
+}
+
+void SClothCollectionOutliner::RepopulateListView()
+{
+	ListItems.Empty();
+
+	const TSharedPtr<const UE::Chaos::ClothAsset::FClothCollection> PinnedClothCollection = ClothCollection.Pin();
+
+	if (!PinnedClothCollection.IsValid())
+	{
+		return;
+	}
+
+	const int32 NumElements = PinnedClothCollection->NumElements(SelectedGroupName);
+
+	for (int32 ElementIndex = 0; ElementIndex < NumElements; ++ElementIndex)
+	{
+		const TSharedPtr<FClothCollectionItem> NewItem = MakeShared<FClothCollectionItem>();
+		NewItem->AttributeValues.SetNum(HeaderData->AttributeNames.Num());
+
+		for (int32 AttributeNameIndex = 0; AttributeNameIndex < HeaderData->AttributeNames.Num(); ++AttributeNameIndex)
+		{
+			const FName& AttributeName = HeaderData->AttributeNames[AttributeNameIndex];
+			if (AttributeName == FClothCollectionHeaderData::ColumnZeroName)
+			{
+				NewItem->AttributeValues[AttributeNameIndex] = FString::FromInt(ElementIndex);
+			}
+			else
+			{
+				NewItem->AttributeValues[AttributeNameIndex] = ClothCollectionOutlinerHelpers::AttributeValueToString(*PinnedClothCollection, AttributeName, SelectedGroupName, ElementIndex);
+			}
+		}
+
+		ListItems.Add(NewItem);
+	}
+
+	ListView->RequestListRefresh();
+}
+
+TSharedRef<ITableRow> SClothCollectionOutliner::GenerateRow(TSharedPtr<const FClothCollectionItem> InItem, const TSharedRef<STableViewBase>& OwnerTable)
+{
+	return SNew(SClothCollectionOutlinerRow, OwnerTable, this->HeaderData, InItem);
+}
+
+
+#undef LOCTEXT_NAMESPACE
diff --git a/Engine/Plugins/Experimental/ChaosClothAssetEditor/Source/ChaosClothAssetEditor/Public/ChaosClothAsset/ClothEditorToolkit.h b/Engine/Plugins/Experimental/ChaosClothAssetEditor/Source/ChaosClothAssetEditor/Public/ChaosClothAsset/ClothEditorToolkit.h
--- a/Engine/Plugins/Experimental/ChaosClothAssetEditor/Source/ChaosClothAssetEditor/Public/ChaosClothAsset/ClothEditorToolkit.h
+++ b/Engine/Plugins/Experimental/ChaosClothAssetEditor/Source/ChaosClothAssetEditor/Public/ChaosClothAsset/ClothEditorToolkit.h
@@ -3,9 +3,11 @@
 #pragma once
 
 #include "BaseCharacterFXEditorToolkit.h"
+#include "Widgets/Input/SComboBox.h"
 
 class FAdvancedPreviewScene;
 class FChaosClothAssetEditor3DViewportClient;
+class SClothCollectionOutliner;
 
 /**
  * The toolkit is supposed to act as the UI manager for the asset editor. It's responsible 
@@ -23,7 +25,7 @@ public:
 	virtual ~FChaosClothAssetEditorToolkit();
 
 	static const FName ClothPreviewTabID;
-	static const FName InteractiveToolsPanelTabID;
+	static const FName OutlinerTabID;
 
 	// FAssetEditorToolkit
 	virtual void AddViewportOverlayWidget(TSharedRef<SWidget> InViewportOverlayWidget) override;
@@ -60,8 +62,6 @@ protected:
 	virtual void InitializeEdMode(UBaseCharacterFXEditorMode* EdMode) override;
 	virtual void CreateEditorModeUILayer() override;
 
-	TSharedRef<SDockTab> SpawnTab_ClothPreview(const FSpawnTabArgs& Args);
-
 	// Appearance customization points
 	// TODO: Implement these when we have an FChaosClothAssetEditorStyle class
 	//virtual const FSlateBrush* GetDefaultTabIcon() const override;
@@ -69,6 +69,9 @@ protected:
 
 private:
 
+	TSharedRef<SDockTab> SpawnTab_ClothPreview(const FSpawnTabArgs& Args);
+	TSharedRef<SDockTab> SpawnTab_Outliner(const FSpawnTabArgs& Args);
+
 	/** Scene in which the 3D sim space preview meshes live. */
 	TUniquePtr<FAdvancedPreviewScene> ClothPreviewScene;
 
@@ -80,9 +83,11 @@ private:
 	TWeakPtr<SEditorViewport> RestSpaceViewport;
 
 	void InitDetailsViewPanel();
+	
+	TSharedPtr<SClothCollectionOutliner> OutlinerView;
+
+	TSharedPtr<SComboBox<FName>> SelectedGroupNameComboBox;
+	TArray<FName> ClothCollectionGroupNames;		// Data source for SelectedGroupNameComboBox
 
-	// TODO as necessary:
-	//TObjectPtr<UInputRouter> ClothPreviewInputRouter = nullptr;
-	//UClothToolViewportButtonsAPI* ViewportButtonsAPI = nullptr;
 };
 
diff --git a/Engine/Plugins/Experimental/ChaosClothAssetEditor/Source/ChaosClothAssetEditor/Public/ChaosClothAsset/SClothCollectionOutliner.h b/Engine/Plugins/Experimental/ChaosClothAssetEditor/Source/ChaosClothAssetEditor/Public/ChaosClothAsset/SClothCollectionOutliner.h
new file mode 100644
--- /dev/null
+++ b/Engine/Plugins/Experimental/ChaosClothAssetEditor/Source/ChaosClothAssetEditor/Public/ChaosClothAsset/SClothCollectionOutliner.h
@@ -0,0 +1,77 @@
+// Copyright Epic Games, Inc. All Rights Reserved.
+
+#pragma once
+
+#include "Widgets/SCompoundWidget.h"
+#include "Widgets/Views/SListView.h"
+
+namespace UE::Chaos::ClothAsset{ class FClothCollection; }
+
+struct FClothCollectionHeaderData
+{
+	TArray<FName> AttributeNames;
+
+	// first column is used to display the numeric index of each element
+	static const FName ColumnZeroName;	
+};
+
+struct FClothCollectionItem
+{
+	TArray<FString> AttributeValues;
+};
+
+class SClothCollectionOutlinerRow : public SMultiColumnTableRow<TSharedPtr<const FClothCollectionItem>>
+{
+public:
+	SLATE_BEGIN_ARGS(SClothCollectionOutlinerRow) {}
+		SLATE_ARGUMENT(TSharedPtr<const FClothCollectionHeaderData>, HeaderData)
+		SLATE_ARGUMENT(TSharedPtr<const FClothCollectionItem>, Item)
+	SLATE_END_ARGS()
+
+	void Construct(const FArguments& InArgs, TSharedRef<STableViewBase> OwnerTableView, const TSharedPtr<const FClothCollectionHeaderData>& InHeaderData, const TSharedPtr<const FClothCollectionItem>& InItemToEdit);
+
+	// SMultiColumnTableRow
+	virtual TSharedRef<SWidget> GenerateWidgetForColumn(const FName& ColumnName) override;
+
+private:
+
+	TSharedPtr<const FClothCollectionHeaderData> HeaderData;
+	TSharedPtr<const FClothCollectionItem> Item;
+};
+
+/// Data visualizer for FClothCollection. Pulls out all Attribute names and values from a given Cloth Collection with a given Group name and adds them to a table
+class SClothCollectionOutliner : public SCompoundWidget
+{
+public:
+
+	SLATE_BEGIN_ARGS(SClothCollectionOutliner) {}
+		SLATE_ARGUMENT(TWeakPtr<UE::Chaos::ClothAsset::FClothCollection>, ClothCollection)
+		SLATE_ARGUMENT(FName, SelectedGroupName)
+	SLATE_END_ARGS()
+
+	void Construct(const FArguments& InArgs);
+
+	void SetClothCollection(TWeakPtr<UE::Chaos::ClothAsset::FClothCollection> ClothCollection);
+
+	///  Changing the selected group name will also automatically rebuild the header and all rows
+	void SetSelectedGroupName(const FName& SelectedGroupName);
+
+	const FName& GetSelectedGroupName() const;
+
+	TSharedRef<ITableRow> GenerateRow(TSharedPtr<const FClothCollectionItem> InItem, const TSharedRef<STableViewBase>& OwnerTable);
+
+private:
+
+	TWeakPtr<UE::Chaos::ClothAsset::FClothCollection> ClothCollection;
+	FName SelectedGroupName;
+
+	TSharedPtr<SListView<TSharedPtr<const FClothCollectionItem>>> ListView;
+	TArray<TSharedPtr<const FClothCollectionItem>> ListItems;
+
+	TSharedPtr<FClothCollectionHeaderData> HeaderData;
+	TSharedPtr<SHeaderRow> HeaderRowWidget;
+
+	void RegenerateHeader();
+	void RepopulateListView();
+
+};
