Cloth Editor: addressing a few review comments that were missed

diff --git a/Engine/Plugins/Experimental/ChaosClothAssetEditor/Source/ChaosClothAssetEditor/Private/ChaosClothAsset/ClothEditorToolkit.cpp b/Engine/Plugins/Experimental/ChaosClothAssetEditor/Source/ChaosClothAssetEditor/Private/ChaosClothAsset/ClothEditorToolkit.cpp
--- a/Engine/Plugins/Experimental/ChaosClothAssetEditor/Source/ChaosClothAssetEditor/Private/ChaosClothAsset/ClothEditorToolkit.cpp
+++ b/Engine/Plugins/Experimental/ChaosClothAssetEditor/Source/ChaosClothAssetEditor/Private/ChaosClothAsset/ClothEditorToolkit.cpp
@@ -21,6 +21,7 @@
 #include "Modules/ModuleManager.h"
 #include "Styling/AppStyle.h"
 #include "ChaosClothAsset/SClothCollectionOutliner.h"
+#include "Widgets/Input/SComboBox.h"
 
 #define LOCTEXT_NAMESPACE "ChaosClothAssetEditorToolkit"
 
@@ -510,25 +511,19 @@ void FChaosClothAssetEditorToolkit::InitDetailsViewPanel()
 	}
 
 	const UChaosClothAssetEditorMode* const ClothMode = CastChecked<UChaosClothAssetEditorMode>(EditorModeManager->GetActiveScriptableMode(UChaosClothAssetEditorMode::EM_ChaosClothAssetEditorModeId));
-	check(ClothMode);
 
-	const bool bHasOutliner = OutlinerView.IsValid();
+	TSharedPtr<UE::Chaos::ClothAsset::FClothCollection> ClothCollection;
 
-	bool bHasClothCollection = false;
-	if (ClothMode->ClothComponent)
+	if (const TObjectPtr<UChaosClothComponent> ClothComponent = ClothMode->ClothComponent)
 	{
-		if (ClothMode->ClothComponent->GetClothAsset())
+		if (UChaosClothAsset* const ClothAsset = ClothComponent->GetClothAsset())
 		{
-			if (ClothMode->ClothComponent->GetClothAsset()->GetClothCollection())
-			{
-				bHasClothCollection = true;
-			}
+			ClothCollection = ClothAsset->GetClothCollection();
 		}
 	}
-
-	if (bHasOutliner && bHasClothCollection)
+	
+	if (OutlinerView.IsValid() && ClothCollection.IsValid())
 	{
-		const TSharedPtr<UE::Chaos::ClothAsset::FClothCollection> ClothCollection = ClothMode->ClothComponent->GetClothAsset()->GetClothCollection();
 		const FName& SelectedGroupName = ClothCollection->SimVerticesGroup;
 		
 		OutlinerView->SetClothCollection(ClothCollection);
diff --git a/Engine/Plugins/Experimental/ChaosClothAssetEditor/Source/ChaosClothAssetEditor/Public/ChaosClothAsset/ClothEditorToolkit.h b/Engine/Plugins/Experimental/ChaosClothAssetEditor/Source/ChaosClothAssetEditor/Public/ChaosClothAsset/ClothEditorToolkit.h
--- a/Engine/Plugins/Experimental/ChaosClothAssetEditor/Source/ChaosClothAssetEditor/Public/ChaosClothAsset/ClothEditorToolkit.h
+++ b/Engine/Plugins/Experimental/ChaosClothAssetEditor/Source/ChaosClothAssetEditor/Public/ChaosClothAsset/ClothEditorToolkit.h
@@ -3,10 +3,10 @@
 #pragma once
 
 #include "BaseCharacterFXEditorToolkit.h"
-#include "Widgets/Input/SComboBox.h"
 
 class FAdvancedPreviewScene;
 class FChaosClothAssetEditor3DViewportClient;
+template<typename T> class SComboBox;
 class SClothCollectionOutliner;
 
 /**
