Chaos Cloth - Fix misplaced const in return type of UChaosClothAsset::GetClothCollection() const.

diff --git a/Engine/Plugins/Experimental/ChaosClothAsset/Source/ChaosClothAssetEngine/Public/ChaosClothAsset/ClothAsset.h b/Engine/Plugins/Experimental/ChaosClothAsset/Source/ChaosClothAssetEngine/Public/ChaosClothAsset/ClothAsset.h
--- a/Engine/Plugins/Experimental/ChaosClothAsset/Source/ChaosClothAssetEngine/Public/ChaosClothAsset/ClothAsset.h
+++ b/Engine/Plugins/Experimental/ChaosClothAsset/Source/ChaosClothAssetEngine/Public/ChaosClothAsset/ClothAsset.h
@@ -99,7 +99,7 @@ public:
 	TSharedPtr<UE::Chaos::ClothAsset::FClothCollection> GetClothCollection() { return ClothCollection; }
 
 	/** Return the enclosed Cloth Collection object, const version. */
-	const TSharedPtr<UE::Chaos::ClothAsset::FClothCollection> GetClothCollection() const { return ClothCollection; }
+	TSharedPtr<const UE::Chaos::ClothAsset::FClothCollection> GetClothCollection() const { return ClothCollection; }
 
 	/** Return the cloth simulation ready LOD model data. */
 	TSharedPtr<const FChaosClothSimulationModel> GetClothSimulationModel() const { return ClothSimulationModel; }
diff --git a/Engine/Plugins/Experimental/ChaosClothAssetEditor/Source/ChaosClothAssetTools/Private/ChaosClothAsset/ClothAssetBuilderEditor.cpp b/Engine/Plugins/Experimental/ChaosClothAssetEditor/Source/ChaosClothAssetTools/Private/ChaosClothAsset/ClothAssetBuilderEditor.cpp
--- a/Engine/Plugins/Experimental/ChaosClothAssetEditor/Source/ChaosClothAssetTools/Private/ChaosClothAsset/ClothAssetBuilderEditor.cpp
+++ b/Engine/Plugins/Experimental/ChaosClothAssetEditor/Source/ChaosClothAssetTools/Private/ChaosClothAsset/ClothAssetBuilderEditor.cpp
@@ -32,7 +32,7 @@ void UClothAssetBuilderEditor::BuildLod(FSkeletalMeshLODModel& LODModel, const U
 	TMap<int32, TArray<int32>> SectionFacesMap;
 	SectionFacesMap.Reserve(ClothAsset.GetMaterials().Num());
 
-	const TSharedPtr<UE::Chaos::ClothAsset::FClothCollection> ClothCollection = ClothAsset.GetClothCollection();
+	const TSharedPtr<const UE::Chaos::ClothAsset::FClothCollection> ClothCollection = ClothAsset.GetClothCollection();
 	const int32 PatternStart = ClothCollection->PatternStart[LodIndex];
 	const int32 PatternEnd = ClothCollection->PatternEnd[LodIndex];
 
diff --git a/Engine/Plugins/Experimental/ChaosClothAssetEditor/Source/ChaosClothAssetTools/Private/ChaosClothAsset/ClothPatternToDynamicMesh.cpp b/Engine/Plugins/Experimental/ChaosClothAssetEditor/Source/ChaosClothAssetTools/Private/ChaosClothAsset/ClothPatternToDynamicMesh.cpp
--- a/Engine/Plugins/Experimental/ChaosClothAssetEditor/Source/ChaosClothAssetTools/Private/ChaosClothAsset/ClothPatternToDynamicMesh.cpp
+++ b/Engine/Plugins/Experimental/ChaosClothAssetEditor/Source/ChaosClothAssetTools/Private/ChaosClothAsset/ClothPatternToDynamicMesh.cpp
@@ -244,9 +244,9 @@ private:
 
 void FClothPatternToDynamicMesh::Convert(const UChaosClothAsset* ClothAssetMeshIn, int32 LODIndex, int32 PatternIndex, bool bGet2DPattern, UE::Geometry::FDynamicMesh3& MeshOut)
 {
-	const TSharedPtr<UE::Chaos::ClothAsset::FClothCollection> ClothCollection = ClothAssetMeshIn->GetClothCollection();
+	const TSharedPtr<const UE::Chaos::ClothAsset::FClothCollection> ClothCollection = ClothAssetMeshIn->GetClothCollection();
 	check(ClothCollection.IsValid());
-	UE::Chaos::ClothAsset::FClothAdapter ClothAdapter(ClothCollection);
+	const UE::Chaos::ClothAsset::FClothConstAdapter ClothAdapter(ClothCollection);
 
 	const UE::Chaos::ClothAsset::FClothLodConstAdapter ClothLodAdapter = ClothAdapter.GetLod(LODIndex);
 	const UE::Chaos::ClothAsset::FClothPatternConstAdapter ClothPatternAdapter = ClothLodAdapter.GetPattern(PatternIndex);
