Cloth Editor: show the rest space patterns in 2D mode when the editor starts

diff --git a/Engine/Plugins/Experimental/ChaosClothAssetEditor/Source/ChaosClothAssetEditor/Private/ChaosClothAsset/ClothEditorMode.cpp b/Engine/Plugins/Experimental/ChaosClothAssetEditor/Source/ChaosClothAssetEditor/Private/ChaosClothAsset/ClothEditorMode.cpp
--- a/Engine/Plugins/Experimental/ChaosClothAssetEditor/Source/ChaosClothAssetEditor/Private/ChaosClothAsset/ClothEditorMode.cpp
+++ b/Engine/Plugins/Experimental/ChaosClothAssetEditor/Source/ChaosClothAssetEditor/Private/ChaosClothAsset/ClothEditorMode.cpp
@@ -633,7 +633,7 @@ void UChaosClothAssetEditorMode::InitializeTargets(const TArray<TObjectPtr<UObje
 	}
 
 	ReinitializeDynamicMeshComponents();
-	RefocusRestSpaceViewportClient();
+	bShouldFocusRestSpaceView = true;
 
 	// Prep things for layer/channel selection
 	AssetNames.Reset(ToolTargets.Num());
@@ -751,6 +751,16 @@ void UChaosClothAssetEditorMode::ModeTick(float DeltaTime)
 	}
 }
 
+void UChaosClothAssetEditorMode::RestSpaceViewportResized(FViewport* RestspaceViewport, uint32 /*Unused*/)
+{
+	// We'd like to call RefocusRestSpaceViewportClient() when the viewport is first created, however in Ortho mode the
+	// viewport needs to have non-zero size for FocusViewportOnBox() to work properly. So we wait until the viewport is resized here.
+	if (bShouldFocusRestSpaceView && RestspaceViewport && RestspaceViewport->GetSizeXY().X > 0 && RestspaceViewport->GetSizeXY().Y > 0)
+	{
+		RefocusRestSpaceViewportClient();
+		bShouldFocusRestSpaceView = false;
+	}
+}
 
 FBox UChaosClothAssetEditorMode::SceneBoundingBox() const
 {
@@ -859,6 +869,11 @@ void UChaosClothAssetEditorMode::SetRestSpaceViewportClient(TWeakPtr<FChaosCloth
 	{
 		VC->Set2DMode(bPattern2DMode);
 		VC->SetToolCommandList(ToolCommandList);
+
+		if (VC->Viewport)
+		{
+			VC->Viewport->ViewportResizedEvent.AddUObject(this, &UChaosClothAssetEditorMode::RestSpaceViewportResized);
+		}
 	}
 }
 
diff --git a/Engine/Plugins/Experimental/ChaosClothAssetEditor/Source/ChaosClothAssetEditor/Private/ChaosClothAsset/ClothEditorToolkit.cpp b/Engine/Plugins/Experimental/ChaosClothAssetEditor/Source/ChaosClothAssetEditor/Private/ChaosClothAsset/ClothEditorToolkit.cpp
--- a/Engine/Plugins/Experimental/ChaosClothAssetEditor/Source/ChaosClothAssetEditor/Private/ChaosClothAsset/ClothEditorToolkit.cpp
+++ b/Engine/Plugins/Experimental/ChaosClothAssetEditor/Source/ChaosClothAssetEditor/Private/ChaosClothAsset/ClothEditorToolkit.cpp
@@ -441,7 +441,6 @@ void FChaosClothAssetEditorToolkit::PostInitAssetEditor()
 	check(ClothMode);
 	const TWeakPtr<FViewportClient> WeakViewportClient(ViewportClient);
 	ClothMode->SetRestSpaceViewportClient(StaticCastWeakPtr<FChaosClothEditorRestSpaceViewportClient>(WeakViewportClient));
-	ClothMode->RefocusRestSpaceViewportClient();
 
 	// Note: We force the cloth preview viewport to open, since some ViewportClient functions are not robust to having no viewport.
 	// See UE-114649
diff --git a/Engine/Plugins/Experimental/ChaosClothAssetEditor/Source/ChaosClothAssetEditor/Public/ChaosClothAsset/ClothEditorMode.h b/Engine/Plugins/Experimental/ChaosClothAssetEditor/Source/ChaosClothAssetEditor/Public/ChaosClothAsset/ClothEditorMode.h
--- a/Engine/Plugins/Experimental/ChaosClothAssetEditor/Source/ChaosClothAssetEditor/Public/ChaosClothAsset/ClothEditorMode.h
+++ b/Engine/Plugins/Experimental/ChaosClothAssetEditor/Source/ChaosClothAssetEditor/Public/ChaosClothAsset/ClothEditorMode.h
@@ -25,6 +25,7 @@ class UDynamicMeshComponent;
 class UChaosClothComponent;
 class FEditorViewportClient;
 class FChaosClothEditorRestSpaceViewportClient;
+class FViewport;
 
 /**
  * The cloth editor mode is the mode used in the cloth asset editor. It holds most of the inter-tool state.
@@ -153,11 +154,15 @@ private:
 	FDelegateHandle SelectionModifiedEventHandle;
 
 	// Whether to display the 2D pattern or 3D rest configuration in the left viewport
-	bool bPattern2DMode = false;
+	bool bPattern2DMode = true;
 
 	// If we can switch between 2D and 3D rest configuration
 	bool bCanTogglePattern2DMode = true;
 
+	bool bShouldFocusRestSpaceView = true;
+
+	void RestSpaceViewportResized(FViewport* RestspaceViewport, uint32 Unused);
+
 	// Whether to combine all patterns into a single DynamicMeshComponent, or have separate components for each pattern
 	// TODO: Expose this to the user
 	bool bCombineAllPatterns = false;
