Cloth Editor: move headers to the private folder unless they need to be public

diff --git a/Engine/Plugins/Experimental/ChaosClothAssetEditor/Source/ChaosClothAssetEditor/Private/ChaosClothAsset/ChaosClothAssetEditorModule.h b/Engine/Plugins/Experimental/ChaosClothAssetEditor/Source/ChaosClothAssetEditor/Private/ChaosClothAsset/ChaosClothAssetEditorModule.h
new file mode 100644
--- /dev/null
+++ b/Engine/Plugins/Experimental/ChaosClothAssetEditor/Source/ChaosClothAssetEditor/Private/ChaosClothAsset/ChaosClothAssetEditorModule.h
@@ -0,0 +1,27 @@
+// Copyright Epic Games, Inc. All Rights Reserved.
+
+#pragma once
+
+#include "CoreMinimal.h"
+#include "BaseCharacterFXEditorModule.h"
+
+class FLayoutExtender;
+namespace UE::Chaos::ClothAsset
+{
+	class FAssetTypeActions_ClothAsset;
+	class FAssetTypeActions_ClothPreset;
+}
+
+class CHAOSCLOTHASSETEDITOR_API FChaosClothAssetEditorModule : public FBaseCharacterFXEditorModule
+{
+public:
+
+	/** IModuleInterface implementation */
+	virtual void StartupModule() override;
+	virtual void ShutdownModule() override;
+
+private:
+	UE::Chaos::ClothAsset::FAssetTypeActions_ClothAsset* AssetTypeActions_ClothAsset;
+	UE::Chaos::ClothAsset::FAssetTypeActions_ClothPreset* AssetTypeActions_ClothPreset;
+
+};
diff --git a/Engine/Plugins/Experimental/ChaosClothAssetEditor/Source/ChaosClothAssetEditor/Private/ChaosClothAsset/ClothEditor.h b/Engine/Plugins/Experimental/ChaosClothAssetEditor/Source/ChaosClothAssetEditor/Private/ChaosClothAsset/ClothEditor.h
new file mode 100644
--- /dev/null
+++ b/Engine/Plugins/Experimental/ChaosClothAssetEditor/Source/ChaosClothAssetEditor/Private/ChaosClothAsset/ClothEditor.h
@@ -0,0 +1,29 @@
+// Copyright Epic Games, Inc. All Rights Reserved.
+
+#pragma once
+
+#include "BaseCharacterFXEditor.h"
+#include "ClothEditor.generated.h"
+
+/** 
+ * The actual asset editor class doesn't have that much in it, intentionally. 
+ * 
+ * Our current asset editor guidelines ask us to place as little business logic as possible
+ * into the class, instead putting as much of the non-UI code into the subsystem as possible,
+ * and the UI code into the toolkit (which this class owns).
+ *
+ * However, since we're using a mode and the Interactive Tools Framework, a lot of our business logic
+ * ends up inside the mode and the tools, not the subsystem. The front-facing code is mostly in
+ * the asset editor toolkit, though the mode toolkit has most of the things that deal with the toolbar
+ * on the left.
+ */
+
+UCLASS()
+class CHAOSCLOTHASSETEDITOR_API UChaosClothAssetEditor : public UBaseCharacterFXEditor
+{
+	GENERATED_BODY()
+
+public:
+
+	virtual TSharedPtr<FBaseAssetToolkit> CreateToolkit() override;
+};
diff --git a/Engine/Plugins/Experimental/ChaosClothAssetEditor/Source/ChaosClothAssetEditor/Private/ChaosClothAsset/ClothEditor3DViewportClient.h b/Engine/Plugins/Experimental/ChaosClothAssetEditor/Source/ChaosClothAssetEditor/Private/ChaosClothAsset/ClothEditor3DViewportClient.h
new file mode 100644
--- /dev/null
+++ b/Engine/Plugins/Experimental/ChaosClothAssetEditor/Source/ChaosClothAssetEditor/Private/ChaosClothAsset/ClothEditor3DViewportClient.h
@@ -0,0 +1,59 @@
+// Copyright Epic Games, Inc. All Rights Reserved.
+
+#pragma once
+
+#include "CoreMinimal.h"
+#include "EditorViewportClient.h"
+
+class UChaosClothComponent;
+class UChaosClothAssetEditorMode;
+
+/**
+ * Viewport client for the 3d sim preview in the cloth editor. Currently same as editor viewport
+ * client but doesn't allow editor gizmos/widgets.
+ */
+class CHAOSCLOTHASSETEDITOR_API FChaosClothAssetEditor3DViewportClient : public FEditorViewportClient
+{
+public:
+
+	FChaosClothAssetEditor3DViewportClient(FEditorModeTools* InModeTools, FPreviewScene* InPreviewScene = nullptr,
+		const TWeakPtr<SEditorViewport>& InEditorViewportWidget = nullptr);
+
+	virtual ~FChaosClothAssetEditor3DViewportClient() = default;
+
+	virtual bool CanSetWidgetMode(UE::Widget::EWidgetMode NewMode) const override {	return false; }
+	virtual void SetWidgetMode(UE::Widget::EWidgetMode NewMode) override {}
+	virtual UE::Widget::EWidgetMode GetWidgetMode() const override { return UE::Widget::EWidgetMode::WM_None; }
+
+	void EnableSimMeshWireframe(bool bEnable ) { bSimMeshWireframe = bEnable; }
+	bool SimMeshWireframeEnabled() const { return bSimMeshWireframe; }
+
+	void EnableRenderMeshWireframe(bool bEnable);
+	bool RenderMeshWireframeEnabled() const { return bRenderMeshWireframe; }
+
+	void SetClothComponent(TObjectPtr<UChaosClothComponent> ClothComponent);
+	void SetClothEdMode(TObjectPtr<UChaosClothAssetEditorMode> ClothEdMode);
+
+	void SoftResetSimulation();
+	void HardResetSimulation();
+	void SuspendSimulation();
+	void ResumeSimulation();
+	bool IsSimulationSuspended() const;
+
+	FBox PreviewBoundingBox() const;
+
+private:
+
+	TObjectPtr<UChaosClothComponent> ClothComponent;
+
+	TObjectPtr<UChaosClothAssetEditorMode> ClothEdMode;
+
+	// Debug draw of simulation meshes
+	virtual void Draw(const FSceneView* View, FPrimitiveDrawInterface* PDI) override;
+
+	UChaosClothAssetEditorMode* GetClothEdMode();
+	const UChaosClothAssetEditorMode* GetClothEdMode() const;
+
+	bool bSimMeshWireframe = true;
+	bool bRenderMeshWireframe = false;
+};
diff --git a/Engine/Plugins/Experimental/ChaosClothAssetEditor/Source/ChaosClothAssetEditor/Private/ChaosClothAsset/ClothEditorCommands.h b/Engine/Plugins/Experimental/ChaosClothAssetEditor/Source/ChaosClothAssetEditor/Private/ChaosClothAsset/ClothEditorCommands.h
new file mode 100644
--- /dev/null
+++ b/Engine/Plugins/Experimental/ChaosClothAssetEditor/Source/ChaosClothAssetEditor/Private/ChaosClothAsset/ClothEditorCommands.h
@@ -0,0 +1,63 @@
+// Copyright Epic Games, Inc. All Rights Reserved.
+
+#pragma once
+
+#include "BaseCharacterFXEditorCommands.h"
+
+class CHAOSCLOTHASSETEDITOR_API FChaosClothAssetEditorCommands : public TBaseCharacterFXEditorCommands<FChaosClothAssetEditorCommands>
+{
+public:
+
+	FChaosClothAssetEditorCommands();
+
+	// TBaseCharacterFXEditorCommands<> interface
+	virtual void RegisterCommands() override;
+
+	// TInteractiveToolCommands<>
+	virtual void GetToolDefaultObjectList(TArray<UInteractiveTool*>& ToolCDOs) override;
+
+	/**
+	 * Add or remove commands relevant to Tool to the given UICommandList.
+	 * Call this when the active tool changes (eg on ToolManager.OnToolStarted / OnToolEnded)
+	 * @param bUnbind if true, commands are removed, otherwise added
+	 */
+	static void UpdateToolCommandBinding(UInteractiveTool* Tool, TSharedPtr<FUICommandList> UICommandList, bool bUnbind = false);
+
+
+public:
+
+	TSharedPtr<FUICommandInfo> OpenClothEditor;
+
+	const static FString BeginRemeshToolIdentifier;
+	TSharedPtr<FUICommandInfo> BeginRemeshTool;
+	const static FString BeginWeightMapPaintToolIdentifier;
+	TSharedPtr<FUICommandInfo> BeginWeightMapPaintTool;
+
+	const static FString BeginAttributeEditorToolIdentifier;
+	TSharedPtr<FUICommandInfo> BeginAttributeEditorTool;
+
+	const static FString BeginClothTrainingToolIdentifier;
+	TSharedPtr<FUICommandInfo> BeginClothTrainingTool;
+
+	// Rest space viewport commands
+	const static FString TogglePatternModeIdentifier;
+	TSharedPtr<FUICommandInfo> TogglePatternMode;
+
+	// Sim viewport commands
+	const static FString ToggleSimMeshWireframeIdentifier;
+	TSharedPtr<FUICommandInfo> ToggleSimMeshWireframe;
+
+	const static FString ToggleRenderMeshWireframeIdentifier;
+	TSharedPtr<FUICommandInfo> ToggleRenderMeshWireframe;
+
+
+	const static FString SoftResetSimulationIdentifier;
+	TSharedPtr<FUICommandInfo> SoftResetSimulation;
+
+	const static FString HardResetSimulationIdentifier;
+	TSharedPtr<FUICommandInfo> HardResetSimulation;
+
+	const static FString ToggleSimulationSuspendedIdentifier;
+	TSharedPtr<FUICommandInfo> ToggleSimulationSuspended;
+
+};
diff --git a/Engine/Plugins/Experimental/ChaosClothAssetEditor/Source/ChaosClothAssetEditor/Private/ChaosClothAsset/ClothEditorMode.h b/Engine/Plugins/Experimental/ChaosClothAssetEditor/Source/ChaosClothAssetEditor/Private/ChaosClothAsset/ClothEditorMode.h
new file mode 100644
--- /dev/null
+++ b/Engine/Plugins/Experimental/ChaosClothAssetEditor/Source/ChaosClothAssetEditor/Private/ChaosClothAsset/ClothEditorMode.h
@@ -0,0 +1,184 @@
+// Copyright Epic Games, Inc. All Rights Reserved.
+
+#pragma once
+
+#include "BaseCharacterFXEditorMode.h"
+#include "GeometryBase.h"
+#include "Delegates/IDelegateInstance.h"
+#include "ClothEditorMode.generated.h"
+
+PREDECLARE_GEOMETRY(class FDynamicMesh3);
+
+class FEditorViewportClient;
+class FAssetEditorModeManager;
+class FToolCommandChange;
+class UMeshElementsVisualizer;
+class UPreviewMesh;
+class UToolTarget;
+class FToolTargetTypeRequirements;
+class UWorld;
+class FChaosClothAssetEditorModeToolkit;
+class UInteractiveToolPropertySet; 
+class UMeshOpPreviewWithBackgroundCompute;
+class UClothToolViewportButtonsAPI;
+class UDynamicMeshComponent;
+class UChaosClothComponent;
+class FEditorViewportClient;
+class FChaosClothEditorRestSpaceViewportClient;
+class FViewport;
+
+/**
+ * The cloth editor mode is the mode used in the cloth asset editor. It holds most of the inter-tool state.
+ * We put things in a mode instead of directly into the asset editor in case we want to someday use the mode
+ * in multiple asset editors.
+ */
+UCLASS(Transient)
+class CHAOSCLOTHASSETEDITOR_API UChaosClothAssetEditorMode : public UBaseCharacterFXEditorMode
+{
+	GENERATED_BODY()
+public:
+	const static FEditorModeID EM_ChaosClothAssetEditorModeId;
+
+	UChaosClothAssetEditorMode();
+
+	virtual void Enter() override;
+
+	/**
+	 * Gets the tool target requirements for the mode. The resulting targets undergo further processing
+	 * to turn them into the input objects that tools get (since these need preview meshes, etc).
+	 */
+	static const FToolTargetTypeRequirements& GetToolTargetRequirements();
+
+	// Both initialization functions must be called for things to function properly. InitializeContexts should
+	// be done first so that the 3d preview world is ready for creating meshes in InitializeTargets.
+	void InitializeContexts(FEditorViewportClient& LivePreviewViewportClient, FAssetEditorModeManager& LivePreviewModeManager);
+
+	virtual void InitializeTargets(const TArray<TObjectPtr<UObject>>& AssetsIn);
+
+	// Asset management
+	bool HaveUnappliedChanges();
+	void GetAssetsWithUnappliedChanges(TArray<TObjectPtr<UObject>> UnappliedAssetsOut);
+	void ApplyChanges();
+
+	/** @return List of asset names, indexed by AssetID */
+	const TArray<FString>& GetAssetNames() const { return AssetNames; }
+
+	// UEdMode overrides
+	virtual bool ShouldToolStartBeAllowed(const FString& ToolIdentifier) const override;
+	virtual void Exit() override;
+	virtual void ModeTick(float DeltaTime) override;
+	virtual void OnToolStarted(UInteractiveToolManager* Manager, UInteractiveTool* Tool) override;
+	virtual void OnToolEnded(UInteractiveToolManager* Manager, UInteractiveTool* Tool) override;
+	virtual void PostUndo() override;
+
+	// We don't actually override MouseEnter, etc, because things get forwarded to the input
+	// router via FEditorModeTools, and we don't have any additional input handling to do at the mode level.
+
+	// Bounding box for rest-space meshes
+	virtual FBox SceneBoundingBox() const override;
+	FBox SelectionBoundingBox() const;		// only selected mesh components
+
+	// Bounding box for sim-space meshes
+	FBox PreviewBoundingBox() const;
+
+	// Toggle between 2D pattern and 3D rest space mesh view
+	bool IsPattern2DModeActive() const;
+	void TogglePatternMode();
+	bool CanTogglePatternMode() const;
+
+	// Simulation controls
+	void SoftResetSimulation();
+	void HardResetSimulation();
+	void SuspendSimulation();
+	void ResumeSimulation();
+	bool IsSimulationSuspended() const;
+
+private:
+
+	friend class FChaosClothAssetEditorToolkit;
+	void SetRestSpaceViewportClient(TWeakPtr<FChaosClothEditorRestSpaceViewportClient, ESPMode::ThreadSafe> ViewportClient);
+	void RefocusRestSpaceViewportClient();
+
+	// UBaseCharacterFXEditorMode
+	virtual void AddToolTargetFactories() override;
+	virtual void RegisterTools() override;
+	virtual void CreateToolTargets(const TArray<TObjectPtr<UObject>>& AssetsIn);
+
+	// UEdMode overrides
+	virtual void CreateToolkit() override;
+	// Not sure whether we need these yet
+	virtual void BindCommands() override;
+	
+	// Transforms that should be used for the 3D previews, 1:1 with OriginalObjectsToEdit
+	// and ToolTargets.
+	TArray<FTransform> Transforms;
+
+	// Preview simulation mesh
+	UPROPERTY()
+	TObjectPtr<UChaosClothComponent> ClothComponent;
+
+	// Rest-space wireframes. They have to get ticked to be able to respond to setting changes. 
+	UPROPERTY()
+	TArray<TObjectPtr<UMeshElementsVisualizer>> WireframesToTick;
+
+	// Authoritative list of targets that have changes that have not been baked back yet.
+	TSet<int32> ModifiedAssetIDs;
+
+	// 1:1 with ToolTargets, indexed by AssetID
+	TArray<FString> AssetNames;
+
+	// Here largely for convenience to avoid having to pass it around functions.
+	UPROPERTY()
+	TObjectPtr<UWorld> PreviewWorld = nullptr;
+	
+	// Mode-level property objects (visible or not) that get ticked.
+	UPROPERTY()
+	TArray<TObjectPtr<UInteractiveToolPropertySet>> PropertyObjectsToTick;
+
+	UPROPERTY()
+	TArray<TObjectPtr<UDynamicMeshComponent>> DynamicMeshComponents;
+
+	UPROPERTY()
+	TArray<TObjectPtr<AActor>> DynamicMeshComponentParentActors;
+
+	struct FDynamicMeshSourceInfo
+	{
+		int32 LodIndex;
+		int32 PatternIndex;
+	};
+
+	TArray<FDynamicMeshSourceInfo> DynamicMeshSourceInfos;
+
+	TWeakPtr<FChaosClothEditorRestSpaceViewportClient, ESPMode::ThreadSafe> RestSpaceViewportClient;
+
+	FDelegateHandle SelectionModifiedEventHandle;
+
+	// Whether to display the 2D pattern or 3D rest configuration in the left viewport
+	bool bPattern2DMode = true;
+
+	// If we can switch between 2D and 3D rest configuration
+	bool bCanTogglePattern2DMode = true;
+
+	bool bShouldFocusRestSpaceView = true;
+
+	void RestSpaceViewportResized(FViewport* RestspaceViewport, uint32 Unused);
+
+	// Whether to combine all patterns into a single DynamicMeshComponent, or have separate components for each pattern
+	// TODO: Expose this to the user
+	bool bCombineAllPatterns = false;
+
+	// Create dynamic mesh components from the cloth component's rest space info
+	void ReinitializeDynamicMeshComponents();
+
+	// Extract the rest space mesh from the given tool target
+	void GetRestSpaceMesh(UToolTarget* ToolTarget, UE::Geometry::FDynamicMesh3& RestSpaceMesh);
+
+	// Set up the preview simulation mesh from the given rest-space mesh
+	void UpdateSimulationMeshes();
+
+	bool bIsSimulationSuspended = false;
+	bool bShouldResetSimulation = false;
+	bool bHardReset = false;
+	bool bShouldClearTeleportFlag = false;
+};
+
diff --git a/Engine/Plugins/Experimental/ChaosClothAssetEditor/Source/ChaosClothAssetEditor/Private/ChaosClothAsset/ClothEditorModeToolkit.h b/Engine/Plugins/Experimental/ChaosClothAssetEditor/Source/ChaosClothAssetEditor/Private/ChaosClothAsset/ClothEditorModeToolkit.h
new file mode 100644
--- /dev/null
+++ b/Engine/Plugins/Experimental/ChaosClothAssetEditor/Source/ChaosClothAssetEditor/Private/ChaosClothAsset/ClothEditorModeToolkit.h
@@ -0,0 +1,23 @@
+// Copyright Epic Games, Inc. All Rights Reserved.
+
+#pragma once
+
+#include "BaseCharacterFXEditorModeToolkit.h"
+
+/**
+ * The cloth editor mode toolkit is responsible for the panel on the side in the cloth editor
+ * that shows mode and tool properties. Tool buttons would go in Init().
+ */
+class CHAOSCLOTHASSETEDITOR_API FChaosClothAssetEditorModeToolkit : public FBaseCharacterFXEditorModeToolkit
+{
+public:
+
+    /** For a specific tool palette category, construct and fill ToolbarBuilder with the category's tools **/
+	virtual void BuildToolPalette(FName PaletteName, class FToolBarBuilder& ToolbarBuilder) override;
+
+	virtual const FSlateBrush* GetActiveToolIcon(const FString& Identifier) const override;
+
+	// IToolkit
+	virtual FName GetToolkitFName() const override;
+	virtual FText GetBaseToolkitName() const override;
+};
diff --git a/Engine/Plugins/Experimental/ChaosClothAssetEditor/Source/ChaosClothAssetEditor/Private/ChaosClothAsset/ClothEditorModeUILayer.h b/Engine/Plugins/Experimental/ChaosClothAssetEditor/Source/ChaosClothAssetEditor/Private/ChaosClothAsset/ClothEditorModeUILayer.h
new file mode 100644
--- /dev/null
+++ b/Engine/Plugins/Experimental/ChaosClothAssetEditor/Source/ChaosClothAssetEditor/Private/ChaosClothAsset/ClothEditorModeUILayer.h
@@ -0,0 +1,30 @@
+// Copyright Epic Games, Inc. All Rights Reserved.
+
+#pragma once
+
+#include "BaseCharacterFXEditorModeUILayer.h"
+#include "ClothEditorModeUILayer.generated.h"
+
+/** Interchange layer to manage built in tab locations within the editor's layout. **/
+UCLASS()
+class CHAOSCLOTHASSETEDITOR_API UChaosClothAssetEditorUISubsystem : public UAssetEditorUISubsystem
+{
+	GENERATED_BODY()
+public:
+
+	virtual void Initialize(FSubsystemCollectionBase& Collection) override;
+	virtual void Deinitialize() override;
+	virtual void RegisterLayoutExtensions(FLayoutExtender& Extender) override;
+
+	// Identifier for the Layout extension in ChaosClothAssetEditorToolkit's StandaloneDefaultLayout
+	static const FName EditorSidePanelAreaName;
+
+};
+
+/** Handles the hosting of additional toolkits, such as the mode toolkit, within the ChaosClothAssetEditor's toolkit. **/
+class FChaosClothAssetEditorModeUILayer : public FBaseCharacterFXEditorModeUILayer
+{
+public:
+	FChaosClothAssetEditorModeUILayer(const IToolkitHost* InToolkitHost) : FBaseCharacterFXEditorModeUILayer(InToolkitHost) {}
+};
+
diff --git a/Engine/Plugins/Experimental/ChaosClothAssetEditor/Source/ChaosClothAssetEditor/Private/ChaosClothAsset/ClothEditorRestSpaceViewportClient.h b/Engine/Plugins/Experimental/ChaosClothAssetEditor/Source/ChaosClothAssetEditor/Private/ChaosClothAsset/ClothEditorRestSpaceViewportClient.h
new file mode 100644
--- /dev/null
+++ b/Engine/Plugins/Experimental/ChaosClothAssetEditor/Source/ChaosClothAssetEditor/Private/ChaosClothAsset/ClothEditorRestSpaceViewportClient.h
@@ -0,0 +1,37 @@
+// Copyright Epic Games, Inc. All Rights Reserved.
+
+#pragma once
+
+#include "CoreMinimal.h"
+#include "EditorViewportClient.h"
+#include "Behaviors/2DViewportBehaviorTargets.h" // FEditor2DScrollBehaviorTarget, FEditor2DMouseWheelZoomBehaviorTarget
+#include "InputBehaviorSet.h"
+
+class UInputBehaviorSet;
+
+class CHAOSCLOTHASSETEDITOR_API FChaosClothEditorRestSpaceViewportClient : public FEditorViewportClient
+{
+public:
+
+	FChaosClothEditorRestSpaceViewportClient(FEditorModeTools* InModeTools, FPreviewScene* InPreviewScene = nullptr, const TWeakPtr<SEditorViewport>& InEditorViewportWidget = nullptr);
+
+	virtual ~FChaosClothEditorRestSpaceViewportClient() = default;
+
+	// FEditorViewportClient
+	virtual void ProcessClick(FSceneView& View, HHitProxy* HitProxy, FKey Key, EInputEvent Event, uint32 HitX, uint32 HitY) override;
+	virtual bool InputKey(const FInputKeyEventArgs& EventArgs) override;
+
+	virtual bool ShouldOrbitCamera() const override;
+
+	void Set2DMode(bool In2DMode);
+
+	void SetEditorViewportWidget(TWeakPtr<SEditorViewport> InEditorViewportWidget);
+	void SetToolCommandList(TWeakPtr<FUICommandList> ToolCommandList);
+
+private:
+
+	bool b2DMode = false;
+
+	TWeakPtr<FUICommandList> ToolCommandList;
+};
+
diff --git a/Engine/Plugins/Experimental/ChaosClothAssetEditor/Source/ChaosClothAssetEditor/Private/ChaosClothAsset/ClothEditorStyle.h b/Engine/Plugins/Experimental/ChaosClothAssetEditor/Source/ChaosClothAssetEditor/Private/ChaosClothAsset/ClothEditorStyle.h
new file mode 100644
--- /dev/null
+++ b/Engine/Plugins/Experimental/ChaosClothAssetEditor/Source/ChaosClothAssetEditor/Private/ChaosClothAsset/ClothEditorStyle.h
@@ -0,0 +1,25 @@
+// Copyright Epic Games, Inc. All Rights Reserved.
+
+#pragma once
+
+#include "Styling/SlateStyle.h"
+#include "Styling/SlateStyleRegistry.h"
+
+/**
+ * Slate style set for Cloth Editor
+ */
+class CHAOSCLOTHASSETEDITOR_API FChaosClothAssetEditorStyle
+	: public FSlateStyleSet
+{
+public:
+	const static FName StyleName;
+
+	/** Access the singleton instance for this style set */
+	static FChaosClothAssetEditorStyle& Get();
+
+private:
+	static FString InContent(const FString& RelativePath, const ANSICHAR* Extension);
+
+	FChaosClothAssetEditorStyle();
+	~FChaosClothAssetEditorStyle();
+};
diff --git a/Engine/Plugins/Experimental/ChaosClothAssetEditor/Source/ChaosClothAssetEditor/Private/ChaosClothAsset/ClothEditorToolkit.h b/Engine/Plugins/Experimental/ChaosClothAssetEditor/Source/ChaosClothAssetEditor/Private/ChaosClothAsset/ClothEditorToolkit.h
new file mode 100644
--- /dev/null
+++ b/Engine/Plugins/Experimental/ChaosClothAssetEditor/Source/ChaosClothAssetEditor/Private/ChaosClothAsset/ClothEditorToolkit.h
@@ -0,0 +1,93 @@
+// Copyright Epic Games, Inc. All Rights Reserved.
+
+#pragma once
+
+#include "BaseCharacterFXEditorToolkit.h"
+
+class FAdvancedPreviewScene;
+class FChaosClothAssetEditor3DViewportClient;
+template<typename T> class SComboBox;
+class SClothCollectionOutliner;
+
+/**
+ * The toolkit is supposed to act as the UI manager for the asset editor. It's responsible 
+ * for setting up viewports and most toolbars, except for the internals of the mode panel.
+ * However, because the toolkit also sets up the mode manager, and much of the important
+ * state is held in the UChaosClothAssetEditorMode managed by the mode manager, the toolkit also ends up
+ * initializing the Cloth mode.
+ * Thus, the FChaosClothAssetEditorToolkit ends up being the central place for the Cloth Asset Editor setup.
+ */
+class CHAOSCLOTHASSETEDITOR_API FChaosClothAssetEditorToolkit : public FBaseCharacterFXEditorToolkit
+{
+public:
+
+	FChaosClothAssetEditorToolkit(UAssetEditor* InOwningAssetEditor);
+	virtual ~FChaosClothAssetEditorToolkit();
+
+	static const FName ClothPreviewTabID;
+	static const FName OutlinerTabID;
+
+	// FAssetEditorToolkit
+	virtual void AddViewportOverlayWidget(TSharedRef<SWidget> InViewportOverlayWidget) override;
+	virtual void RemoveViewportOverlayWidget(TSharedRef<SWidget> InViewportOverlayWidget) override;
+
+	virtual FText GetToolkitName() const override;
+	virtual FName GetToolkitFName() const override;
+	virtual FText GetBaseToolkitName() const override;
+	virtual FText GetToolkitToolTipText() const override;
+	virtual void RegisterTabSpawners(const TSharedRef<FTabManager>& InTabManager) override;
+	virtual void UnregisterTabSpawners(const TSharedRef<FTabManager>& TabManager) override;
+	virtual bool OnRequestClose() override;
+
+	// TODO: Implement this if we want to apply tool changes before saving an asset
+	//virtual void SaveAsset_Execute() override;
+
+	// IAssetEditorInstance
+	// This is important because if this returns true, attempting to edit a static mesh that is
+	// open in the cloth editor may open the cloth editor instead of opening the static mesh editor.
+	// TODO: Change this if we create a dedicated Cloth Asset
+	virtual bool IsPrimaryEditor() const override { return false; };
+
+protected:
+
+	// FBaseAssetToolkit
+	virtual AssetEditorViewportFactoryFunction GetViewportDelegate() override;
+	virtual TSharedPtr<FEditorViewportClient> CreateEditorViewportClient() const override;
+
+	// FAssetEditorToolkit
+	virtual void PostInitAssetEditor() override;
+
+	// FBaseCharacterFXEditorToolkit
+	virtual FEditorModeID GetEditorModeId() const override;
+	virtual void InitializeEdMode(UBaseCharacterFXEditorMode* EdMode) override;
+	virtual void CreateEditorModeUILayer() override;
+
+	// Appearance customization points
+	// TODO: Implement these when we have an FChaosClothAssetEditorStyle class
+	//virtual const FSlateBrush* GetDefaultTabIcon() const override;
+	//virtual FLinearColor GetDefaultTabColor() const override;
+
+private:
+
+	TSharedRef<SDockTab> SpawnTab_ClothPreview(const FSpawnTabArgs& Args);
+	TSharedRef<SDockTab> SpawnTab_Outliner(const FSpawnTabArgs& Args);
+
+	/** Scene in which the 3D sim space preview meshes live. */
+	TUniquePtr<FAdvancedPreviewScene> ClothPreviewScene;
+
+	TSharedPtr<class FEditorViewportTabContent> ClothPreviewTabContent;
+	AssetEditorViewportFactoryFunction ClothPreviewViewportDelegate;
+	TSharedPtr<FChaosClothAssetEditor3DViewportClient> ClothPreviewViewportClient;
+	TSharedPtr<FAssetEditorModeManager> ClothPreviewEditorModeManager;
+
+	TWeakPtr<SEditorViewport> RestSpaceViewport;
+
+	void InitDetailsViewPanel();
+	
+	TSharedPtr<SClothCollectionOutliner> OutlinerView;
+
+	TSharedPtr<SComboBox<FName>> SelectedGroupNameComboBox;
+	TArray<FName> ClothCollectionGroupNames;		// Data source for SelectedGroupNameComboBox
+
+};
+
diff --git a/Engine/Plugins/Experimental/ChaosClothAssetEditor/Source/ChaosClothAssetEditor/Private/ChaosClothAsset/SClothCollectionOutliner.h b/Engine/Plugins/Experimental/ChaosClothAssetEditor/Source/ChaosClothAssetEditor/Private/ChaosClothAsset/SClothCollectionOutliner.h
new file mode 100644
--- /dev/null
+++ b/Engine/Plugins/Experimental/ChaosClothAssetEditor/Source/ChaosClothAssetEditor/Private/ChaosClothAsset/SClothCollectionOutliner.h
@@ -0,0 +1,77 @@
+// Copyright Epic Games, Inc. All Rights Reserved.
+
+#pragma once
+
+#include "Widgets/SCompoundWidget.h"
+#include "Widgets/Views/SListView.h"
+
+namespace UE::Chaos::ClothAsset{ class FClothCollection; }
+
+struct FClothCollectionHeaderData
+{
+	TArray<FName> AttributeNames;
+
+	// first column is used to display the numeric index of each element
+	static const FName ColumnZeroName;	
+};
+
+struct FClothCollectionItem
+{
+	TArray<FString> AttributeValues;
+};
+
+class SClothCollectionOutlinerRow : public SMultiColumnTableRow<TSharedPtr<const FClothCollectionItem>>
+{
+public:
+	SLATE_BEGIN_ARGS(SClothCollectionOutlinerRow) {}
+		SLATE_ARGUMENT(TSharedPtr<const FClothCollectionHeaderData>, HeaderData)
+		SLATE_ARGUMENT(TSharedPtr<const FClothCollectionItem>, Item)
+	SLATE_END_ARGS()
+
+	void Construct(const FArguments& InArgs, TSharedRef<STableViewBase> OwnerTableView, const TSharedPtr<const FClothCollectionHeaderData>& InHeaderData, const TSharedPtr<const FClothCollectionItem>& InItemToEdit);
+
+	// SMultiColumnTableRow
+	virtual TSharedRef<SWidget> GenerateWidgetForColumn(const FName& ColumnName) override;
+
+private:
+
+	TSharedPtr<const FClothCollectionHeaderData> HeaderData;
+	TSharedPtr<const FClothCollectionItem> Item;
+};
+
+/// Data visualizer for FClothCollection. Pulls out all Attribute names and values from a given Cloth Collection with a given Group name and adds them to a table
+class SClothCollectionOutliner : public SCompoundWidget
+{
+public:
+
+	SLATE_BEGIN_ARGS(SClothCollectionOutliner) {}
+		SLATE_ARGUMENT(TWeakPtr<UE::Chaos::ClothAsset::FClothCollection>, ClothCollection)
+		SLATE_ARGUMENT(FName, SelectedGroupName)
+	SLATE_END_ARGS()
+
+	void Construct(const FArguments& InArgs);
+
+	void SetClothCollection(TWeakPtr<UE::Chaos::ClothAsset::FClothCollection> ClothCollection);
+
+	///  Changing the selected group name will also automatically rebuild the header and all rows
+	void SetSelectedGroupName(const FName& SelectedGroupName);
+
+	const FName& GetSelectedGroupName() const;
+
+	TSharedRef<ITableRow> GenerateRow(TSharedPtr<const FClothCollectionItem> InItem, const TSharedRef<STableViewBase>& OwnerTable);
+
+private:
+
+	TWeakPtr<UE::Chaos::ClothAsset::FClothCollection> ClothCollection;
+	FName SelectedGroupName;
+
+	TSharedPtr<SListView<TSharedPtr<const FClothCollectionItem>>> ListView;
+	TArray<TSharedPtr<const FClothCollectionItem>> ListItems;
+
+	TSharedPtr<FClothCollectionHeaderData> HeaderData;
+	TSharedPtr<SHeaderRow> HeaderRowWidget;
+
+	void RegenerateHeader();
+	void RepopulateListView();
+
+};
diff --git a/Engine/Plugins/Experimental/ChaosClothAssetEditor/Source/ChaosClothAssetEditor/Private/ChaosClothAsset/SClothEditor3DViewport.h b/Engine/Plugins/Experimental/ChaosClothAssetEditor/Source/ChaosClothAssetEditor/Private/ChaosClothAsset/SClothEditor3DViewport.h
new file mode 100644
--- /dev/null
+++ b/Engine/Plugins/Experimental/ChaosClothAssetEditor/Source/ChaosClothAssetEditor/Private/ChaosClothAsset/SClothEditor3DViewport.h
@@ -0,0 +1,27 @@
+// Copyright Epic Games, Inc. All Rights Reserved.
+
+#pragma once
+
+#include "CoreMinimal.h"
+#include "SAssetEditorViewport.h"
+#include "SCommonEditorViewportToolbarBase.h"
+
+/**
+ * Viewport used for 3D preview in cloth editor. Has a custom toolbar overlay at the top.
+ */
+class CHAOSCLOTHASSETEDITOR_API SChaosClothAssetEditor3DViewport : public SAssetEditorViewport, public ICommonEditorViewportToolbarInfoProvider
+{
+
+public:
+	// SAssetEditorViewport
+	virtual void BindCommands() override;
+	virtual TSharedPtr<SWidget> MakeViewportToolbar() override;
+
+	virtual void OnFocusViewportToSelection() override;
+
+	// ICommonEditorViewportToolbarInfoProvider
+	virtual TSharedRef<class SEditorViewport> GetViewportWidget() override;
+	virtual TSharedPtr<FExtender> GetExtenders() const override;
+	virtual void OnFloatingButtonClicked() override;
+
+};
diff --git a/Engine/Plugins/Experimental/ChaosClothAssetEditor/Source/ChaosClothAssetEditor/Private/ChaosClothAsset/SClothEditor3DViewportToolBar.h b/Engine/Plugins/Experimental/ChaosClothAssetEditor/Source/ChaosClothAssetEditor/Private/ChaosClothAsset/SClothEditor3DViewportToolBar.h
new file mode 100644
--- /dev/null
+++ b/Engine/Plugins/Experimental/ChaosClothAssetEditor/Source/ChaosClothAssetEditor/Private/ChaosClothAsset/SClothEditor3DViewportToolBar.h
@@ -0,0 +1,33 @@
+// Copyright Epic Games, Inc. All Rights Reserved.
+
+#pragma once
+
+#include "SClothEditorViewportToolBarBase.h"
+
+class FExtender;
+class FUICommandList;
+class SChaosClothAssetEditor3DViewport;
+
+/**
+ * Toolbar that shows up at the top of the 3d viewport
+ */
+class CHAOSCLOTHASSETEDITOR_API SChaosClothAssetEditor3DViewportToolBar : public SChaosClothAssetEditorViewportToolBarBase
+{
+public:
+	SLATE_BEGIN_ARGS(SChaosClothAssetEditor3DViewportToolBar) {}
+		SLATE_ARGUMENT(TSharedPtr<FUICommandList>, CommandList)
+		SLATE_ARGUMENT(TSharedPtr<FExtender>, Extenders)
+	SLATE_END_ARGS()
+
+	void Construct(const FArguments& InArgs, TSharedPtr<class SChaosClothAssetEditor3DViewport> InChaosClothAssetEditor3DViewport);
+
+private:
+	TSharedRef<SWidget> MakeOptionsMenu();
+	TSharedRef<SWidget> MakeDisplayToolBar(const TSharedPtr<FExtender> InExtenders);
+	TSharedRef<SWidget> MakeToolBar(const TSharedPtr<FExtender> InExtenders);
+
+	/** The viewport that we are in */
+	TWeakPtr<class SChaosClothAssetEditor3DViewport> ChaosClothAssetEditor3DViewportPtr;
+
+	TSharedPtr<FUICommandList> CommandList;
+};
diff --git a/Engine/Plugins/Experimental/ChaosClothAssetEditor/Source/ChaosClothAssetEditor/Private/ChaosClothAsset/SClothEditorRestSpaceViewport.h b/Engine/Plugins/Experimental/ChaosClothAssetEditor/Source/ChaosClothAssetEditor/Private/ChaosClothAsset/SClothEditorRestSpaceViewport.h
new file mode 100644
--- /dev/null
+++ b/Engine/Plugins/Experimental/ChaosClothAssetEditor/Source/ChaosClothAssetEditor/Private/ChaosClothAsset/SClothEditorRestSpaceViewport.h
@@ -0,0 +1,23 @@
+// Copyright Epic Games, Inc. All Rights Reserved.
+
+#pragma once
+
+#include "CoreMinimal.h"
+#include "SBaseCharacterFXEditorViewport.h"
+#include "SCommonEditorViewportToolbarBase.h"
+
+class CHAOSCLOTHASSETEDITOR_API SChaosClothAssetEditorRestSpaceViewport : public SBaseCharacterFXEditorViewport, public ICommonEditorViewportToolbarInfoProvider
+{
+public:
+
+	// SEditorViewport
+	virtual void BindCommands() override;
+	virtual TSharedPtr<SWidget> MakeViewportToolbar() override;
+	virtual void OnFocusViewportToSelection() override;
+
+	// ICommonEditorViewportToolbarInfoProvider
+	virtual TSharedRef<class SEditorViewport> GetViewportWidget() override;
+	virtual TSharedPtr<FExtender> GetExtenders() const override;
+	virtual void OnFloatingButtonClicked() override;
+
+};
diff --git a/Engine/Plugins/Experimental/ChaosClothAssetEditor/Source/ChaosClothAssetEditor/Private/ChaosClothAsset/SClothEditorRestSpaceViewportToolBar.h b/Engine/Plugins/Experimental/ChaosClothAssetEditor/Source/ChaosClothAssetEditor/Private/ChaosClothAsset/SClothEditorRestSpaceViewportToolBar.h
new file mode 100644
--- /dev/null
+++ b/Engine/Plugins/Experimental/ChaosClothAssetEditor/Source/ChaosClothAssetEditor/Private/ChaosClothAsset/SClothEditorRestSpaceViewportToolBar.h
@@ -0,0 +1,33 @@
+// Copyright Epic Games, Inc. All Rights Reserved.
+
+#pragma once
+
+#include "SClothEditorViewportToolBarBase.h"
+
+class FExtender;
+class FUICommandList;
+class SChaosClothAssetEditorRestSpaceViewport;
+
+/**
+ * Toolbar that shows up at the top of the rest space viewport
+ */
+class CHAOSCLOTHASSETEDITOR_API SChaosClothAssetEditorRestSpaceViewportToolBar : public SChaosClothAssetEditorViewportToolBarBase
+{
+public:
+	SLATE_BEGIN_ARGS(SChaosClothAssetEditorRestSpaceViewportToolBar) {}
+		SLATE_ARGUMENT(TSharedPtr<FUICommandList>, CommandList)
+		SLATE_ARGUMENT(TSharedPtr<FExtender>, Extenders)
+	SLATE_END_ARGS()
+
+	void Construct(const FArguments& InArgs, TSharedPtr<class SChaosClothAssetEditorRestSpaceViewport> InChaosClothAssetEditorViewport);
+
+private:
+	TSharedRef<SWidget> MakeOptionsMenu();
+	TSharedRef<SWidget> MakeDisplayToolBar(const TSharedPtr<FExtender> InExtenders);
+	TSharedRef<SWidget> MakeToolBar(const TSharedPtr<FExtender> InExtenders);
+
+	/** The viewport that we are in */
+	TWeakPtr<class SChaosClothAssetEditorRestSpaceViewport> ChaosClothAssetEditorRestSpaceViewportPtr;
+
+	TSharedPtr<FUICommandList> CommandList;
+};
diff --git a/Engine/Plugins/Experimental/ChaosClothAssetEditor/Source/ChaosClothAssetEditor/Private/ChaosClothAsset/SClothEditorViewportToolBarBase.h b/Engine/Plugins/Experimental/ChaosClothAssetEditor/Source/ChaosClothAssetEditor/Private/ChaosClothAsset/SClothEditorViewportToolBarBase.h
new file mode 100644
--- /dev/null
+++ b/Engine/Plugins/Experimental/ChaosClothAssetEditor/Source/ChaosClothAssetEditor/Private/ChaosClothAsset/SClothEditorViewportToolBarBase.h
@@ -0,0 +1,17 @@
+// Copyright Epic Games, Inc. All Rights Reserved.
+
+#pragma once
+
+#include "SCommonEditorViewportToolbarBase.h"
+
+/// Common code for toolbars in the Cloth Editor
+class CHAOSCLOTHASSETEDITOR_API SChaosClothAssetEditorViewportToolBarBase : public SCommonEditorViewportToolbarBase
+{
+protected:
+	TSharedRef<SWidget> GenerateClothViewportOptionsMenu() const;
+
+private:
+	TSharedRef<SWidget> GenerateFOVMenu() const;
+	float OnGetFOVValue() const;
+};
+
diff --git a/Engine/Plugins/Experimental/ChaosClothAssetEditor/Source/ChaosClothAssetEditor/Public/ChaosClothAsset/ChaosClothAssetEditorModule.h b/Engine/Plugins/Experimental/ChaosClothAssetEditor/Source/ChaosClothAssetEditor/Public/ChaosClothAsset/ChaosClothAssetEditorModule.h
deleted file mode 100644
diff --git a/Engine/Plugins/Experimental/ChaosClothAssetEditor/Source/ChaosClothAssetEditor/Public/ChaosClothAsset/ClothEditor.h b/Engine/Plugins/Experimental/ChaosClothAssetEditor/Source/ChaosClothAssetEditor/Public/ChaosClothAsset/ClothEditor.h
deleted file mode 100644
diff --git a/Engine/Plugins/Experimental/ChaosClothAssetEditor/Source/ChaosClothAssetEditor/Public/ChaosClothAsset/ClothEditor3DViewportClient.h b/Engine/Plugins/Experimental/ChaosClothAssetEditor/Source/ChaosClothAssetEditor/Public/ChaosClothAsset/ClothEditor3DViewportClient.h
deleted file mode 100644
diff --git a/Engine/Plugins/Experimental/ChaosClothAssetEditor/Source/ChaosClothAssetEditor/Public/ChaosClothAsset/ClothEditorCommands.h b/Engine/Plugins/Experimental/ChaosClothAssetEditor/Source/ChaosClothAssetEditor/Public/ChaosClothAsset/ClothEditorCommands.h
deleted file mode 100644
diff --git a/Engine/Plugins/Experimental/ChaosClothAssetEditor/Source/ChaosClothAssetEditor/Public/ChaosClothAsset/ClothEditorMode.h b/Engine/Plugins/Experimental/ChaosClothAssetEditor/Source/ChaosClothAssetEditor/Public/ChaosClothAsset/ClothEditorMode.h
deleted file mode 100644
diff --git a/Engine/Plugins/Experimental/ChaosClothAssetEditor/Source/ChaosClothAssetEditor/Public/ChaosClothAsset/ClothEditorModeToolkit.h b/Engine/Plugins/Experimental/ChaosClothAssetEditor/Source/ChaosClothAssetEditor/Public/ChaosClothAsset/ClothEditorModeToolkit.h
deleted file mode 100644
diff --git a/Engine/Plugins/Experimental/ChaosClothAssetEditor/Source/ChaosClothAssetEditor/Public/ChaosClothAsset/ClothEditorModeUILayer.h b/Engine/Plugins/Experimental/ChaosClothAssetEditor/Source/ChaosClothAssetEditor/Public/ChaosClothAsset/ClothEditorModeUILayer.h
deleted file mode 100644
diff --git a/Engine/Plugins/Experimental/ChaosClothAssetEditor/Source/ChaosClothAssetEditor/Public/ChaosClothAsset/ClothEditorRestSpaceViewportClient.h b/Engine/Plugins/Experimental/ChaosClothAssetEditor/Source/ChaosClothAssetEditor/Public/ChaosClothAsset/ClothEditorRestSpaceViewportClient.h
deleted file mode 100644
diff --git a/Engine/Plugins/Experimental/ChaosClothAssetEditor/Source/ChaosClothAssetEditor/Public/ChaosClothAsset/ClothEditorStyle.h b/Engine/Plugins/Experimental/ChaosClothAssetEditor/Source/ChaosClothAssetEditor/Public/ChaosClothAsset/ClothEditorStyle.h
deleted file mode 100644
diff --git a/Engine/Plugins/Experimental/ChaosClothAssetEditor/Source/ChaosClothAssetEditor/Public/ChaosClothAsset/ClothEditorToolkit.h b/Engine/Plugins/Experimental/ChaosClothAssetEditor/Source/ChaosClothAssetEditor/Public/ChaosClothAsset/ClothEditorToolkit.h
deleted file mode 100644
diff --git a/Engine/Plugins/Experimental/ChaosClothAssetEditor/Source/ChaosClothAssetEditor/Public/ChaosClothAsset/SClothCollectionOutliner.h b/Engine/Plugins/Experimental/ChaosClothAssetEditor/Source/ChaosClothAssetEditor/Public/ChaosClothAsset/SClothCollectionOutliner.h
deleted file mode 100644
diff --git a/Engine/Plugins/Experimental/ChaosClothAssetEditor/Source/ChaosClothAssetEditor/Public/ChaosClothAsset/SClothEditor3DViewport.h b/Engine/Plugins/Experimental/ChaosClothAssetEditor/Source/ChaosClothAssetEditor/Public/ChaosClothAsset/SClothEditor3DViewport.h
deleted file mode 100644
diff --git a/Engine/Plugins/Experimental/ChaosClothAssetEditor/Source/ChaosClothAssetEditor/Public/ChaosClothAsset/SClothEditor3DViewportToolBar.h b/Engine/Plugins/Experimental/ChaosClothAssetEditor/Source/ChaosClothAssetEditor/Public/ChaosClothAsset/SClothEditor3DViewportToolBar.h
deleted file mode 100644
diff --git a/Engine/Plugins/Experimental/ChaosClothAssetEditor/Source/ChaosClothAssetEditor/Public/ChaosClothAsset/SClothEditorRestSpaceViewport.h b/Engine/Plugins/Experimental/ChaosClothAssetEditor/Source/ChaosClothAssetEditor/Public/ChaosClothAsset/SClothEditorRestSpaceViewport.h
deleted file mode 100644
diff --git a/Engine/Plugins/Experimental/ChaosClothAssetEditor/Source/ChaosClothAssetEditor/Public/ChaosClothAsset/SClothEditorRestSpaceViewportToolBar.h b/Engine/Plugins/Experimental/ChaosClothAssetEditor/Source/ChaosClothAssetEditor/Public/ChaosClothAsset/SClothEditorRestSpaceViewportToolBar.h
deleted file mode 100644
diff --git a/Engine/Plugins/Experimental/ChaosClothAssetEditor/Source/ChaosClothAssetEditor/Public/ChaosClothAsset/SClothEditorViewportToolBarBase.h b/Engine/Plugins/Experimental/ChaosClothAssetEditor/Source/ChaosClothAssetEditor/Public/ChaosClothAsset/SClothEditorViewportToolBarBase.h
deleted file mode 100644
diff --git a/Engine/Plugins/Experimental/ChaosClothAssetEditor/Source/ChaosClothAssetEditorTools/Private/ChaosClothAsset/ChaosClothAssetEditorToolsModule.h b/Engine/Plugins/Experimental/ChaosClothAssetEditor/Source/ChaosClothAssetEditorTools/Private/ChaosClothAsset/ChaosClothAssetEditorToolsModule.h
new file mode 100644
--- /dev/null
+++ b/Engine/Plugins/Experimental/ChaosClothAssetEditor/Source/ChaosClothAssetEditorTools/Private/ChaosClothAsset/ChaosClothAssetEditorToolsModule.h
@@ -0,0 +1,15 @@
+// Copyright Epic Games, Inc. All Rights Reserved.
+
+#pragma once
+
+#include "CoreMinimal.h"
+#include "Modules/ModuleManager.h"
+
+class FChaosClothAssetEditorToolsModule : public IModuleInterface
+{
+public:
+
+	/** IModuleInterface implementation */
+	virtual void StartupModule() override;
+	virtual void ShutdownModule() override;
+};
diff --git a/Engine/Plugins/Experimental/ChaosClothAssetEditor/Source/ChaosClothAssetEditorTools/Private/ChaosClothAsset/ClothWeightMapPaintBrushOps.h b/Engine/Plugins/Experimental/ChaosClothAssetEditor/Source/ChaosClothAssetEditorTools/Private/ChaosClothAsset/ClothWeightMapPaintBrushOps.h
new file mode 100644
--- /dev/null
+++ b/Engine/Plugins/Experimental/ChaosClothAssetEditor/Source/ChaosClothAssetEditorTools/Private/ChaosClothAsset/ClothWeightMapPaintBrushOps.h
@@ -0,0 +1,125 @@
+// Copyright Epic Games, Inc. All Rights Reserved.
+
+#pragma once
+
+#include "Sculpting/MeshBrushOpBase.h"
+#include "DynamicMesh/DynamicMesh3.h"
+#include "MeshWeights.h"
+#include "Async/ParallelFor.h"
+#include "ClothWeightMapPaintBrushOps.generated.h"
+
+
+
+
+class CHAOSCLOTHASSETEDITORTOOLS_API FMeshVertexWeightMapEditBrushOp : public FMeshSculptBrushOp
+{
+public:
+
+	// not supported for this kind of brush op
+	virtual void ApplyStamp(
+		const FDynamicMesh3* Mesh,
+		const FSculptBrushStamp& Stamp,
+		const TArray<int32>& Vertices,
+		TArray<FVector3d>& NewColorsOut) override
+	{
+		check(false);
+	}
+
+
+	virtual void ApplyStampByVertices(
+		const FDynamicMesh3* Mesh,
+		const FSculptBrushStamp& Stamp,
+		const TArray<int32>& Vertices,
+		TArray<double>& NewAttributesOut) = 0;
+
+};
+
+
+
+
+
+
+UCLASS()
+class CHAOSCLOTHASSETEDITORTOOLS_API UWeightMapEraseBrushOpProps : public UMeshSculptBrushOpProps
+{
+	GENERATED_BODY()
+public:
+	/** Attribute to set as Erased value */
+	UPROPERTY(EditAnywhere, Category = EraseBrush, meta = (DisplayName = "Erase Attribute", UIMin = 0))
+	double AttributeValue = 0;
+
+	virtual double GetAttribute() { return AttributeValue; }
+};
+
+
+
+class FWeightMapEraseBrushOp : public FMeshVertexWeightMapEditBrushOp
+{
+public:
+
+	virtual void ApplyStampByVertices(
+		const FDynamicMesh3* Mesh,
+		const FSculptBrushStamp& Stamp,
+		const TArray<int32>& Vertices,
+		TArray<double>& NewAttributesOut) override
+	{
+		const FVector3d& StampPos = Stamp.LocalFrame.Origin;
+
+		UWeightMapEraseBrushOpProps* Props = GetPropertySetAs<UWeightMapEraseBrushOpProps>();
+		double EraseAttribute = (double)Props->GetAttribute();
+
+		// TODO: Add something here to get the old value so we can subtract (clamped) the AttributeValue from it.
+
+		// TODO: Handle the stamp's properties for fall off, etc..
+
+		int32 NumVertices = Vertices.Num();
+		for (int32 k = 0; k < NumVertices; ++k)
+		{
+			NewAttributesOut[k] = 0.0;
+		}
+	}
+};
+
+
+
+
+UCLASS()
+class CHAOSCLOTHASSETEDITORTOOLS_API UWeightMapPaintBrushOpProps : public UMeshSculptBrushOpProps
+{
+	GENERATED_BODY()
+public:
+	/** The Attribute that will be assigned to triangles within the brush region */
+	UPROPERTY(EditAnywhere, Category = PaintBrush, meta = (DisplayName = "Attribute", UIMin = 0))
+	double AttributeValue = 1;
+
+	virtual double GetAttribute() { return AttributeValue; }
+};
+
+
+
+class FWeightMapPaintBrushOp : public FMeshVertexWeightMapEditBrushOp
+{
+public:
+
+	virtual void ApplyStampByVertices(
+		const FDynamicMesh3* Mesh,
+		const FSculptBrushStamp& Stamp,
+		const TArray<int32>& Vertices,
+		TArray<double>& NewAttributesOut) override
+	{
+		const FVector3d& StampPos = Stamp.LocalFrame.Origin;
+
+		UWeightMapPaintBrushOpProps* Props = GetPropertySetAs<UWeightMapPaintBrushOpProps>();
+		double SetToAttribute = (double)Props->GetAttribute();
+
+		// TODO: Add something here to get the old value so we can add (clamped) the AttributeValue from it.
+
+		// TODO: Handle the stamp's properties for fall off, etc..
+
+		int32 NumVertices = Vertices.Num();
+		for ( int32 k = 0; k < NumVertices; ++k)
+		{
+			NewAttributesOut[k] = FMath::Clamp(SetToAttribute, 0.0, 1.0);
+		}
+	}
+};
diff --git a/Engine/Plugins/Experimental/ChaosClothAssetEditor/Source/ChaosClothAssetEditorTools/Public/ChaosClothAsset/ChaosClothAssetEditorToolsModule.h b/Engine/Plugins/Experimental/ChaosClothAssetEditor/Source/ChaosClothAssetEditorTools/Public/ChaosClothAsset/ChaosClothAssetEditorToolsModule.h
deleted file mode 100644
diff --git a/Engine/Plugins/Experimental/ChaosClothAssetEditor/Source/ChaosClothAssetEditorTools/Public/ChaosClothAsset/ClothWeightMapPaintBrushOps.h b/Engine/Plugins/Experimental/ChaosClothAssetEditor/Source/ChaosClothAssetEditorTools/Public/ChaosClothAsset/ClothWeightMapPaintBrushOps.h
deleted file mode 100644
