ClothEditor: tidy up the ClothEditorMode class
- make a bunch of functions private
- remove dead code carried over from the UV Editor that might not be relevant to our editor
- make naming conventions more consistent

Should not have any significant behavior changes

diff --git a/Engine/Plugins/Experimental/ChaosClothAssetEditor/Source/ChaosClothAssetEditor/Private/ChaosClothAsset/ClothEditorMode.cpp b/Engine/Plugins/Experimental/ChaosClothAssetEditor/Source/ChaosClothAssetEditor/Private/ChaosClothAsset/ClothEditorMode.cpp
--- a/Engine/Plugins/Experimental/ChaosClothAssetEditor/Source/ChaosClothAssetEditor/Private/ChaosClothAsset/ClothEditorMode.cpp
+++ b/Engine/Plugins/Experimental/ChaosClothAssetEditor/Source/ChaosClothAssetEditor/Private/ChaosClothAsset/ClothEditorMode.cpp
@@ -55,30 +55,54 @@
 
 const FEditorModeID UChaosClothAssetEditorMode::EM_ChaosClothAssetEditorModeId = TEXT("EM_ChaosClothAssetEditorMode");
 
-namespace ChaosClothAssetEditorModeLocals
+
+namespace ChaosClothAssetEditorModeHelpers
 {
-	void LogAndToastWarning(const FText& Error)
+	TArray<FName> GetClothAssetWeightMapNames(const UE::Chaos::ClothAsset::FClothAdapter& ClothAdapter)
 	{
-		FNotificationInfo Info(Error);
-		Info.ExpireDuration = 5.0f;
-		FSlateNotificationManager::Get().AddNotification(Info);
+		const TSharedPtr<const UE::Chaos::ClothAsset::FClothCollection> ClothCollection = ClothAdapter.FClothConstAdapter::GetClothCollection();
 
-		// TODO: use a non-temp category
-		UE_LOG(LogTemp, Warning, TEXT("%s"), *Error.ToString());
+		TArray<FName> OutWeightMapNames;
+		const TArray<FName> SimVerticesAttributeNames = ClothCollection->AttributeNames(UE::Chaos::ClothAsset::FClothCollection::SimVerticesGroup);
+		for (const FName& AttributeName : SimVerticesAttributeNames)
+		{
+			if (ClothCollection->FindAttributeTyped<float>(AttributeName, UE::Chaos::ClothAsset::FClothCollection::SimVerticesGroup))
+			{
+				OutWeightMapNames.Add(AttributeName);
+			}
+		}
+
+		return OutWeightMapNames;
 	}
-}		// namespace ChaosClothAssetEditorModeLocals
 
-const FToolTargetTypeRequirements& UChaosClothAssetEditorMode::GetToolTargetRequirements()
-{
-	static const FToolTargetTypeRequirements ToolTargetRequirements =
-		FToolTargetTypeRequirements({
-			UMaterialProvider::StaticClass(),
-			UDynamicMeshCommitter::StaticClass(),
-			UDynamicMeshProvider::StaticClass()
-			});
-	return ToolTargetRequirements;
+	void RemoveClothWeightMaps(UE::Chaos::ClothAsset::FClothAdapter& ClothAdapter, const TArray<FName>& WeightMapNames)
+	{
+		const TSharedPtr<const UE::Chaos::ClothAsset::FClothCollection> ClothCollection = ClothAdapter.GetClothCollection();
+		for (const FName& RemovedMap : WeightMapNames)
+		{
+			if (ClothCollection->FindAttributeTyped<float>(RemovedMap, UE::Chaos::ClothAsset::FClothCollection::SimVerticesGroup))
+			{
+				ClothAdapter.RemoveWeightMap(RemovedMap);
+			}
+		}
+	}
+
+	TArray<FName> GetDynamicMeshWeightMapNames(const UE::Geometry::FDynamicMesh3& DynamicMesh)
+	{
+		TArray<FName> OutWeightMapNames;
+
+		for (int32 DynamicMeshWeightMapIndex = 0; DynamicMeshWeightMapIndex < DynamicMesh.Attributes()->NumWeightLayers(); ++DynamicMeshWeightMapIndex)
+		{
+			const UE::Geometry::FDynamicMeshWeightAttribute* const WeightMapAttribute = DynamicMesh.Attributes()->GetWeightLayer(DynamicMeshWeightMapIndex);
+			const FName WeightMapName = WeightMapAttribute->GetName();
+			OutWeightMapNames.Add(WeightMapName);
+		}
+
+		return OutWeightMapNames;
+	}
 }
 
+
 UChaosClothAssetEditorMode::UChaosClothAssetEditorMode()
 {
 	Info = FEditorModeInfo(
@@ -88,6 +112,17 @@ UChaosClothAssetEditorMode::UChaosClothAssetEditorMode()
 		false);
 }
 
+const FToolTargetTypeRequirements& UChaosClothAssetEditorMode::GetToolTargetRequirements()
+{
+	static const FToolTargetTypeRequirements ToolTargetRequirements =
+		FToolTargetTypeRequirements({
+			UMaterialProvider::StaticClass(),
+			UDynamicMeshCommitter::StaticClass(),
+			UDynamicMeshProvider::StaticClass()
+			});
+	return ToolTargetRequirements;
+}
+
 void UChaosClothAssetEditorMode::Enter()
 {
 	UBaseCharacterFXEditorMode::Enter();
@@ -225,79 +260,11 @@ void UChaosClothAssetEditorMode::Exit()
 	Super::Exit();
 }
 
-void UChaosClothAssetEditorMode::InitializeContexts(FEditorViewportClient& LivePreviewViewportClient, FAssetEditorModeManager& LivePreviewModeManager)
-{
-	PreviewWorld = LivePreviewModeManager.GetPreviewScene()->GetWorld();
-
-	// UContextObjectStore* ContextStore = GetInteractiveToolsContext()->ToolManager->GetContextObjectStore();
-
-	// TODO: If we are going to support tools that operate on the 3D sim mesh, we'll need something like this to pass into those tools:
-	//UClothToolLivePreviewAPI* LivePreviewAPI = NewObject<UClothToolLivePreviewAPI>();
-	//LivePreviewAPI->Initialize(
-	//	LivePreviewWorld,
-	//	LivePreviewModeManager.GetInteractiveToolsContext()->InputRouter,
-	//	[this, LivePreviewViewportClientPtr = &LivePreviewViewportClient](FViewCameraState& CameraStateOut) {
-	//		GetCameraState(*LivePreviewViewportClientPtr, CameraStateOut); 
-	//	});
-	//ContextStore->AddContextObject(LivePreviewAPI);
-
-	// TODO: If we add buttons to the viewport that the tools will need to interact with (e.g. selection of triangles, edges, etc.), we will need something like this:
-	//ContextStore->AddContextObject(&ViewportButtonsAPI);
-}
-
-void UChaosClothAssetEditorMode::GetRestSpaceMesh(UToolTarget* ToolTarget, UE::Geometry::FDynamicMesh3& RestSpaceMesh)
+void UChaosClothAssetEditorMode::SetPreviewWorld(TObjectPtr<UWorld> InWorld)
 {
-	// TODO: Get the actual rest-space mesh. NOTE: Currently full 3D mesh (not 2D projection)
-	RestSpaceMesh = UE::ToolTarget::GetDynamicMeshCopy(ToolTarget);
+	PreviewWorld = InWorld;
 }
 
-namespace ChaosClothAssetEditorModeHelpers
-{
-	TArray<FName> GetClothAssetWeightMapNames(const UE::Chaos::ClothAsset::FClothAdapter& ClothAdapter)
-	{
-		const TSharedPtr<const UE::Chaos::ClothAsset::FClothCollection> ClothCollection = ClothAdapter.FClothConstAdapter::GetClothCollection();
-
-		TArray<FName> OutWeightMapNames;
-		const TArray<FName> SimVerticesAttributeNames = ClothCollection->AttributeNames(UE::Chaos::ClothAsset::FClothCollection::SimVerticesGroup);
-		for (const FName& AttributeName : SimVerticesAttributeNames)
-		{
-			if (ClothCollection->FindAttributeTyped<float>(AttributeName, UE::Chaos::ClothAsset::FClothCollection::SimVerticesGroup))
-			{
-				OutWeightMapNames.Add(AttributeName);
-			}
-		}
-
-		return OutWeightMapNames;
-	}
-
-	void RemoveClothWeightMaps(UE::Chaos::ClothAsset::FClothAdapter& ClothAdapter, const TArray<FName>& WeightMapNames)
-	{
-		const TSharedPtr<const UE::Chaos::ClothAsset::FClothCollection> ClothCollection = ClothAdapter.GetClothCollection();
-		for (const FName& RemovedMap : WeightMapNames)
-		{
-			if (ClothCollection->FindAttributeTyped<float>(RemovedMap, UE::Chaos::ClothAsset::FClothCollection::SimVerticesGroup))
-			{
-				ClothAdapter.RemoveWeightMap(RemovedMap);
-			}
-		}
-	}
-
-	TArray<FName> GetDynamicMeshWeightMapNames(const UE::Geometry::FDynamicMesh3& DynamicMesh)
-	{
-		TArray<FName> OutWeightMapNames;
-
-		for (int32 DynamicMeshWeightMapIndex = 0; DynamicMeshWeightMapIndex < DynamicMesh.Attributes()->NumWeightLayers(); ++DynamicMeshWeightMapIndex)
-		{
-			const UE::Geometry::FDynamicMeshWeightAttribute* const WeightMapAttribute = DynamicMesh.Attributes()->GetWeightLayer(DynamicMeshWeightMapIndex);
-			const FName WeightMapName = WeightMapAttribute->GetName();
-			OutWeightMapNames.Add(WeightMapName);
-		}
-
-		return OutWeightMapNames;
-	}
-}
-
-
 void UChaosClothAssetEditorMode::UpdateSimulationMeshes()
 {
 	GetToolManager()->BeginUndoTransaction(LOCTEXT("ChaosClothAssetEditorApplyChangesTransaction", "Cloth Editor Apply Changes"));
@@ -582,7 +549,7 @@ void UChaosClothAssetEditorMode::RefocusRestSpaceViewportClient()
 		PinnedVC->ToggleOrbitCamera(false);
 
 		const FBox SceneBounds = SceneBoundingBox();
-		if (IsPattern2DModeActive())
+		if (bPattern2DMode)
 		{
 			// 2D pattern
 			PinnedVC->SetInitialViewTransform(ELevelViewportType::LVT_OrthoXY, FVector(0, 0, 0), FRotator(0, 0, 0), DEFAULT_ORTHOZOOM);
@@ -605,9 +572,6 @@ void UChaosClothAssetEditorMode::InitializeTargets(const TArray<TObjectPtr<UObje
 
 	OriginalObjectsToEdit = AssetsIn;
 
-	// TODO: If we have multiple objects to initialize, position them relative to each other somehow
-	Transforms.Init(FTransform::Identity, AssetsIn.Num());
-
 	CreateToolTargets(AssetsIn);
 
 	// NOTE: This ensure will fire until we are able to build a tool target associated with ClothComponent
@@ -634,36 +598,8 @@ void UChaosClothAssetEditorMode::InitializeTargets(const TArray<TObjectPtr<UObje
 
 	ReinitializeDynamicMeshComponents();
 	bShouldFocusRestSpaceView = true;
-
-	// Prep things for layer/channel selection
-	AssetNames.Reset(ToolTargets.Num());
-	for (const TObjectPtr<UToolTarget>& ToolTarget : ToolTargets)
-	{
-		AssetNames.Add(UE::ToolTarget::GetHumanReadableName(ToolTarget));
-	}
 }
 
-bool UChaosClothAssetEditorMode::HaveUnappliedChanges()
-{
-	return ModifiedAssetIDs.Num() > 0;
-}
-
-void UChaosClothAssetEditorMode::GetAssetsWithUnappliedChanges(TArray<TObjectPtr<UObject>> UnappliedAssetsOut)
-{
-	for (int32 AssetID : ModifiedAssetIDs)
-	{
-		// The asset ID corresponds to the index into OriginalObjectsToEdit
-		UnappliedAssetsOut.Add(OriginalObjectsToEdit[AssetID]);
-	}
-}
-
-void UChaosClothAssetEditorMode::ApplyChanges()
-{
-	GetToolManager()->BeginUndoTransaction(LOCTEXT("ChaosClothAssetEditorApplyChangesTransaction", "Cloth Editor Apply Changes"));
-	GetToolManager()->EndUndoTransaction();
-}
-
-
 void UChaosClothAssetEditorMode::SoftResetSimulation()
 {
 	bShouldResetSimulation = true;
@@ -835,11 +771,6 @@ FBox UChaosClothAssetEditorMode::PreviewBoundingBox() const
 	return FBox(ForceInitToZero);
 }
 
-bool UChaosClothAssetEditorMode::IsPattern2DModeActive() const
-{
-	return bPattern2DMode;
-}
-
 void UChaosClothAssetEditorMode::TogglePatternMode()
 {
 	bPattern2DMode = !bPattern2DMode;
diff --git a/Engine/Plugins/Experimental/ChaosClothAssetEditor/Source/ChaosClothAssetEditor/Private/ChaosClothAsset/ClothEditorMode.h b/Engine/Plugins/Experimental/ChaosClothAssetEditor/Source/ChaosClothAssetEditor/Private/ChaosClothAsset/ClothEditorMode.h
--- a/Engine/Plugins/Experimental/ChaosClothAssetEditor/Source/ChaosClothAssetEditor/Private/ChaosClothAsset/ClothEditorMode.h
+++ b/Engine/Plugins/Experimental/ChaosClothAssetEditor/Source/ChaosClothAssetEditor/Private/ChaosClothAsset/ClothEditorMode.h
@@ -33,56 +33,23 @@ class FViewport;
  * in multiple asset editors.
  */
 UCLASS(Transient)
-class CHAOSCLOTHASSETEDITOR_API UChaosClothAssetEditorMode : public UBaseCharacterFXEditorMode
+class CHAOSCLOTHASSETEDITOR_API UChaosClothAssetEditorMode final : public UBaseCharacterFXEditorMode
 {
 	GENERATED_BODY()
+
 public:
+
 	const static FEditorModeID EM_ChaosClothAssetEditorModeId;
 
 	UChaosClothAssetEditorMode();
 
-	virtual void Enter() override;
-
-	/**
-	 * Gets the tool target requirements for the mode. The resulting targets undergo further processing
-	 * to turn them into the input objects that tools get (since these need preview meshes, etc).
-	 */
-	static const FToolTargetTypeRequirements& GetToolTargetRequirements();
-
-	// Both initialization functions must be called for things to function properly. InitializeContexts should
-	// be done first so that the 3d preview world is ready for creating meshes in InitializeTargets.
-	void InitializeContexts(FEditorViewportClient& LivePreviewViewportClient, FAssetEditorModeManager& LivePreviewModeManager);
-
-	virtual void InitializeTargets(const TArray<TObjectPtr<UObject>>& AssetsIn);
-
-	// Asset management
-	bool HaveUnappliedChanges();
-	void GetAssetsWithUnappliedChanges(TArray<TObjectPtr<UObject>> UnappliedAssetsOut);
-	void ApplyChanges();
-
-	/** @return List of asset names, indexed by AssetID */
-	const TArray<FString>& GetAssetNames() const { return AssetNames; }
-
-	// UEdMode overrides
-	virtual bool ShouldToolStartBeAllowed(const FString& ToolIdentifier) const override;
-	virtual void Exit() override;
-	virtual void ModeTick(float DeltaTime) override;
-	virtual void OnToolStarted(UInteractiveToolManager* Manager, UInteractiveTool* Tool) override;
-	virtual void OnToolEnded(UInteractiveToolManager* Manager, UInteractiveTool* Tool) override;
-	virtual void PostUndo() override;
-
-	// We don't actually override MouseEnter, etc, because things get forwarded to the input
-	// router via FEditorModeTools, and we don't have any additional input handling to do at the mode level.
-
-	// Bounding box for rest-space meshes
-	virtual FBox SceneBoundingBox() const override;
-	FBox SelectionBoundingBox() const;		// only selected mesh components
+	// Bounding box for selected rest space mesh components
+	FBox SelectionBoundingBox() const;
 
-	// Bounding box for sim-space meshes
+	// Bounding box for sim space meshes
 	FBox PreviewBoundingBox() const;
 
 	// Toggle between 2D pattern and 3D rest space mesh view
-	bool IsPattern2DModeActive() const;
 	void TogglePatternMode();
 	bool CanTogglePatternMode() const;
 
@@ -96,37 +63,49 @@ public:
 private:
 
 	friend class FChaosClothAssetEditorToolkit;
-	void SetRestSpaceViewportClient(TWeakPtr<FChaosClothEditorRestSpaceViewportClient, ESPMode::ThreadSafe> ViewportClient);
-	void RefocusRestSpaceViewportClient();
+
+	// UEdMode
+	virtual void Enter() final;
+	virtual void Exit() override;
+	virtual void ModeTick(float DeltaTime) override;
+	virtual bool ShouldToolStartBeAllowed(const FString& ToolIdentifier) const override;
+	virtual void OnToolStarted(UInteractiveToolManager* Manager, UInteractiveTool* Tool) override;
+	virtual void OnToolEnded(UInteractiveToolManager* Manager, UInteractiveTool* Tool) override;
+	virtual void PostUndo() override;
+	virtual void CreateToolkit() override;
+	virtual void BindCommands() override;
+
+	// (We don't actually override MouseEnter, etc, because things get forwarded to the input
+	// router via FEditorModeTools, and we don't have any additional input handling to do at the mode level.)
 
 	// UBaseCharacterFXEditorMode
 	virtual void AddToolTargetFactories() override;
 	virtual void RegisterTools() override;
-	virtual void CreateToolTargets(const TArray<TObjectPtr<UObject>>& AssetsIn);
+	virtual void CreateToolTargets(const TArray<TObjectPtr<UObject>>& AssetsIn) override;
+	virtual void InitializeTargets(const TArray<TObjectPtr<UObject>>& AssetsIn) override;
 
-	// UEdMode overrides
-	virtual void CreateToolkit() override;
-	// Not sure whether we need these yet
-	virtual void BindCommands() override;
-	
-	// Transforms that should be used for the 3D previews, 1:1 with OriginalObjectsToEdit
-	// and ToolTargets.
-	TArray<FTransform> Transforms;
+	// Gets the tool target requirements for the mode. The resulting targets undergo further processing
+	// to turn them into the input objects that tools get (since these need preview meshes, etc).
+	static const FToolTargetTypeRequirements& GetToolTargetRequirements();
+
+	// Both SetPreviewWorld and InitializeTargets functions must be called for things to function properly. SetPreviewWorld should
+	// be done first so that the 3d preview world is ready for creating meshes in InitializeTargets.
+	void SetPreviewWorld(TObjectPtr<UWorld> InWorld);
+
+	// Bounding box for rest space meshes
+	virtual FBox SceneBoundingBox() const override;
+
+	void SetRestSpaceViewportClient(TWeakPtr<FChaosClothEditorRestSpaceViewportClient, ESPMode::ThreadSafe> ViewportClient);
+	void RefocusRestSpaceViewportClient();
 
 	// Preview simulation mesh
 	UPROPERTY()
 	TObjectPtr<UChaosClothComponent> ClothComponent;
 
-	// Rest-space wireframes. They have to get ticked to be able to respond to setting changes. 
+	// Rest space wireframes. They have to get ticked to be able to respond to setting changes. 
 	UPROPERTY()
 	TArray<TObjectPtr<UMeshElementsVisualizer>> WireframesToTick;
 
-	// Authoritative list of targets that have changes that have not been baked back yet.
-	TSet<int32> ModifiedAssetIDs;
-
-	// 1:1 with ToolTargets, indexed by AssetID
-	TArray<FString> AssetNames;
-
 	// Here largely for convenience to avoid having to pass it around functions.
 	UPROPERTY()
 	TObjectPtr<UWorld> PreviewWorld = nullptr;
@@ -135,32 +114,35 @@ private:
 	UPROPERTY()
 	TArray<TObjectPtr<UInteractiveToolPropertySet>> PropertyObjectsToTick;
 
+	// Rest space editable meshes
 	UPROPERTY()
 	TArray<TObjectPtr<UDynamicMeshComponent>> DynamicMeshComponents;
 
+	// Actors required for hit testing DynamicMeshComponents
 	UPROPERTY()
 	TArray<TObjectPtr<AActor>> DynamicMeshComponentParentActors;
 
+	// Map back to original asset location for each DynamicMeshComponent
 	struct FDynamicMeshSourceInfo
 	{
 		int32 LodIndex;
 		int32 PatternIndex;
 	};
-
 	TArray<FDynamicMeshSourceInfo> DynamicMeshSourceInfos;
 
 	TWeakPtr<FChaosClothEditorRestSpaceViewportClient, ESPMode::ThreadSafe> RestSpaceViewportClient;
 
+	// Handle to a callback triggered when the current selection changes
 	FDelegateHandle SelectionModifiedEventHandle;
 
-	// Whether to display the 2D pattern or 3D rest configuration in the left viewport
+	// Whether to display the 2D pattern or 3D assembly in the rest space viewport
 	bool bPattern2DMode = true;
 
-	// If we can switch between 2D and 3D rest configuration
+	// Whether we can switch between 2D and 3D rest configuration
 	bool bCanTogglePattern2DMode = true;
 
+	// Whether the rest space viewport should focus on the rest space mesh on the next tick
 	bool bShouldFocusRestSpaceView = true;
-
 	void RestSpaceViewportResized(FViewport* RestspaceViewport, uint32 Unused);
 
 	// Whether to combine all patterns into a single DynamicMeshComponent, or have separate components for each pattern
@@ -170,12 +152,10 @@ private:
 	// Create dynamic mesh components from the cloth component's rest space info
 	void ReinitializeDynamicMeshComponents();
 
-	// Extract the rest space mesh from the given tool target
-	void GetRestSpaceMesh(UToolTarget* ToolTarget, UE::Geometry::FDynamicMesh3& RestSpaceMesh);
-
 	// Set up the preview simulation mesh from the given rest-space mesh
 	void UpdateSimulationMeshes();
 
+	// Simulation controls
 	bool bIsSimulationSuspended = false;
 	bool bShouldResetSimulation = false;
 	bool bHardReset = false;
diff --git a/Engine/Plugins/Experimental/ChaosClothAssetEditor/Source/ChaosClothAssetEditor/Private/ChaosClothAsset/ClothEditorToolkit.cpp b/Engine/Plugins/Experimental/ChaosClothAssetEditor/Source/ChaosClothAssetEditor/Private/ChaosClothAsset/ClothEditorToolkit.cpp
--- a/Engine/Plugins/Experimental/ChaosClothAssetEditor/Source/ChaosClothAssetEditor/Private/ChaosClothAsset/ClothEditorToolkit.cpp
+++ b/Engine/Plugins/Experimental/ChaosClothAssetEditor/Source/ChaosClothAssetEditor/Private/ChaosClothAsset/ClothEditorToolkit.cpp
@@ -227,31 +227,6 @@ bool FChaosClothAssetEditorToolkit::OnRequestClose()
 		return true; 
 	}
 
-	// Warn the user of any unapplied changes.
-	if (ClothEdMode->HaveUnappliedChanges())
-	{
-		TArray<TObjectPtr<UObject>> UnappliedAssets;
-		ClothEdMode->GetAssetsWithUnappliedChanges(UnappliedAssets);
-
-		EAppReturnType::Type YesNoCancelReply = FMessageDialog::Open(EAppMsgType::YesNoCancel,
-			NSLOCTEXT("ChaosClothAssetEditor", "Prompt_ChaosClothAssetEditorClose", "At least one of the assets has unapplied changes. Would you like to apply them? (Selecting 'No' will cause all changes to be lost!)"));
-
-		switch (YesNoCancelReply)
-		{
-		case EAppReturnType::Yes:
-			ClothEdMode->ApplyChanges();
-			break;
-
-		case EAppReturnType::No:
-			// exit
-			break;
-
-		case EAppReturnType::Cancel:
-			// don't exit
-			return false;
-		}
-	}
-
 	// Give any active modes a chance to shutdown while the toolkit host is still alive
     // This is super important to do, otherwise currently opened tabs won't be marked as "closed".
     // This results in tabs not being properly recycled upon reopening the editor and tab
@@ -481,13 +456,13 @@ void FChaosClothAssetEditorToolkit::InitializeEdMode(UBaseCharacterFXEditorMode*
 	UChaosClothAssetEditorMode* ClothMode = Cast<UChaosClothAssetEditorMode>(EdMode);
 	check(ClothMode);
 
-	// The mode will need to be able to get to the live preview world, camera, input router, and viewport buttons.
-	ClothMode->InitializeContexts(*ClothPreviewViewportClient, *ClothPreviewEditorModeManager);
+	// The mode will need to be able to get to the live preview world
+	ClothMode->SetPreviewWorld(ClothPreviewEditorModeManager->GetWorld());
 
 	TArray<TObjectPtr<UObject>> ObjectsToEdit;
 	OwningAssetEditor->GetObjectsToEdit(ObjectsToEdit);
 
-	EdMode->InitializeTargets(ObjectsToEdit);
+	ClothMode->InitializeTargets(ObjectsToEdit);
 }
 
 void FChaosClothAssetEditorToolkit::CreateEditorModeUILayer()
