Cloth Editor:
- add Dataflow graph editor widget to the Cloth Editor
- add a Dataflow property to ClothAsset
- create a do-nothing ClothAsset Terminal Dataflow node for testing

diff --git a/Engine/Plugins/Experimental/ChaosClothAsset/Source/ChaosClothAssetEngine/ChaosClothAssetEngine.Build.cs b/Engine/Plugins/Experimental/ChaosClothAsset/Source/ChaosClothAssetEngine/ChaosClothAssetEngine.Build.cs
--- a/Engine/Plugins/Experimental/ChaosClothAsset/Source/ChaosClothAssetEngine/ChaosClothAssetEngine.Build.cs
+++ b/Engine/Plugins/Experimental/ChaosClothAsset/Source/ChaosClothAssetEngine/ChaosClothAssetEngine.Build.cs
@@ -20,7 +20,8 @@ public class ChaosClothAssetEngine : ModuleRules
 				"ChaosClothAsset",
 				"ChaosCloth",
 				"TaskGraph",
-				"ClothingSystemRuntimeInterface"
+				"ClothingSystemRuntimeInterface",
+				"DataflowEngine"
 			}
 		);
 
diff --git a/Engine/Plugins/Experimental/ChaosClothAsset/Source/ChaosClothAssetEngine/Public/ChaosClothAsset/ClothAsset.h b/Engine/Plugins/Experimental/ChaosClothAsset/Source/ChaosClothAssetEngine/Public/ChaosClothAsset/ClothAsset.h
--- a/Engine/Plugins/Experimental/ChaosClothAsset/Source/ChaosClothAssetEngine/Public/ChaosClothAsset/ClothAsset.h
+++ b/Engine/Plugins/Experimental/ChaosClothAsset/Source/ChaosClothAssetEngine/Public/ChaosClothAsset/ClothAsset.h
@@ -13,6 +13,7 @@ class FSkeletalMeshRenderData;
 class FSkeletalMeshModel;
 struct FChaosClothSimulationModel;
 struct FSkeletalMeshLODInfo;
+class UDataflow;
 
 UENUM()
 enum class EClothAssetAsyncProperties : uint64
@@ -120,6 +121,16 @@ public:
 	/** Set the bone hierachy to use for this cloth. */
 	void SetReferenceSkeleton(const FReferenceSkeleton& InReferenceSkeleton, bool bRebuildClothSimulationModel = true) { RefSkeleton = InReferenceSkeleton; UpdateSkeleton(bRebuildClothSimulationModel); }
 
+	//
+	// Dataflow
+	//
+	UPROPERTY(EditAnywhere, Category = "Dataflow")
+	TObjectPtr<UDataflow> DataflowAsset;
+
+	UPROPERTY(EditAnywhere, Category = "Dataflow")
+	FString DataflowTerminal = "ClothAssetTerminal";
+
+
 private:
 	//~ Begin USkinnedAsset interface
 	/** Initial step for the Post Load process - Can't be done in parallel. */
diff --git a/Engine/Plugins/Experimental/ChaosClothAssetEditor/ChaosClothAssetEditor.uplugin b/Engine/Plugins/Experimental/ChaosClothAssetEditor/ChaosClothAssetEditor.uplugin
--- a/Engine/Plugins/Experimental/ChaosClothAssetEditor/ChaosClothAssetEditor.uplugin
+++ b/Engine/Plugins/Experimental/ChaosClothAssetEditor/ChaosClothAssetEditor.uplugin
@@ -55,6 +55,10 @@
 			"Name": "BaseCharacterFXEditor",
 			"Enabled": true,
 			"TargetAllowList": [ "Editor" ]
+		},
+		{
+			"Name": "Dataflow",
+			"Enabled": true
 		}
 
 	]
diff --git a/Engine/Plugins/Experimental/ChaosClothAssetEditor/Source/ChaosClothAssetEditor/ChaosClothAssetEditor.Build.cs b/Engine/Plugins/Experimental/ChaosClothAssetEditor/Source/ChaosClothAssetEditor/ChaosClothAssetEditor.Build.cs
--- a/Engine/Plugins/Experimental/ChaosClothAssetEditor/Source/ChaosClothAssetEditor/ChaosClothAssetEditor.Build.cs
+++ b/Engine/Plugins/Experimental/ChaosClothAssetEditor/Source/ChaosClothAssetEditor/ChaosClothAssetEditor.Build.cs
@@ -54,7 +54,10 @@ public class ChaosClothAssetEditor : ModuleRules
 				"ChaosClothAssetEngine",
 				"ClothingSystemEditor",
 				"ChaosClothAssetTools",
-				"ChaosClothAssetEditorTools"
+				"ChaosClothAssetEditorTools",
+				"DataflowCore",
+				"DataflowEditor",
+				"DataflowEngine"
 				// ... add private dependencies that you statically link with here ...	
 			}
 			);
diff --git a/Engine/Plugins/Experimental/ChaosClothAssetEditor/Source/ChaosClothAssetEditor/Private/ChaosClothAsset/ChaosClothAssetEditorModule.cpp b/Engine/Plugins/Experimental/ChaosClothAssetEditor/Source/ChaosClothAssetEditor/Private/ChaosClothAsset/ChaosClothAssetEditorModule.cpp
--- a/Engine/Plugins/Experimental/ChaosClothAssetEditor/Source/ChaosClothAssetEditor/Private/ChaosClothAsset/ChaosClothAssetEditorModule.cpp
+++ b/Engine/Plugins/Experimental/ChaosClothAssetEditor/Source/ChaosClothAssetEditor/Private/ChaosClothAsset/ChaosClothAssetEditorModule.cpp
@@ -4,6 +4,7 @@
 #include "ChaosClothAsset/ClothEditorMode.h"
 #include "ChaosClothAsset/ClothEditorCommands.h"
 #include "ChaosClothAsset/ClothEditorStyle.h"
+#include "ChaosClothAsset/ClothDataflowNodes.h"
 #include "ContentBrowserMenuContexts.h"
 #include "EditorModeRegistry.h"
 #include "Selection.h"
@@ -34,6 +35,8 @@ void FChaosClothAssetEditorModule::StartupModule()
 	AssetTypeActions_ClothPreset = new FAssetTypeActions_ClothPreset();
 	AssetTools.RegisterAssetTypeActions(MakeShareable(AssetTypeActions_ClothPreset));
 
+	Dataflow::RegisterClothDataflowNodes();
+
 	// TODO: Register details view customizations
 }
 
diff --git a/Engine/Plugins/Experimental/ChaosClothAssetEditor/Source/ChaosClothAssetEditor/Private/ChaosClothAsset/ClothDataflowNodes.cpp b/Engine/Plugins/Experimental/ChaosClothAssetEditor/Source/ChaosClothAssetEditor/Private/ChaosClothAsset/ClothDataflowNodes.cpp
new file mode 100644
--- /dev/null
+++ b/Engine/Plugins/Experimental/ChaosClothAssetEditor/Source/ChaosClothAssetEditor/Private/ChaosClothAsset/ClothDataflowNodes.cpp
@@ -0,0 +1,17 @@
+
+#include "ChaosClothAsset/ClothDataflowNodes.h"
+#include "Logging/LogMacros.h"
+
+#include UE_INLINE_GENERATED_CPP_BY_NAME(ClothDataflowNodes)
+
+#define LOCTEXT_NAMESPACE "ClothDataflowNodes"
+
+namespace Dataflow
+{
+	void RegisterClothDataflowNodes()
+	{
+		DATAFLOW_NODE_REGISTER_CREATION_FACTORY(FClothAssetTerminalDataflowNode);
+	}
+}
+
+#undef LOCTEXT_NAMESPACE
diff --git a/Engine/Plugins/Experimental/ChaosClothAssetEditor/Source/ChaosClothAssetEditor/Private/ChaosClothAsset/ClothDataflowNodes.h b/Engine/Plugins/Experimental/ChaosClothAssetEditor/Source/ChaosClothAssetEditor/Private/ChaosClothAsset/ClothDataflowNodes.h
new file mode 100644
--- /dev/null
+++ b/Engine/Plugins/Experimental/ChaosClothAssetEditor/Source/ChaosClothAssetEditor/Private/ChaosClothAsset/ClothDataflowNodes.h
@@ -0,0 +1,47 @@
+// Copyright Epic Games, Inc. All Rights Reserved.
+
+#pragma once 
+
+#include "Dataflow/DataflowCore.h"
+#include "ChaosClothAsset/ClothAsset.h"
+#include "ChaosClothAsset/ClothCollection.h"
+#include "ClothDataflowNodes.generated.h"
+
+USTRUCT()
+struct FClothAssetTerminalDataflowNode : public FDataflowTerminalNode
+{
+	GENERATED_USTRUCT_BODY()
+	DATAFLOW_NODE_DEFINE_INTERNAL(FClothAssetTerminalDataflowNode, "ClothAssetTerminal", "Cloth", "")
+
+public:
+
+	UPROPERTY(meta = (DataflowInput, DataflowOutput, DisplayName = "Collection", DataflowPassthrough = "Collection"))
+	FManagedArrayCollection Collection;		// TODO: Replace with FClothCollection
+
+	FClothAssetTerminalDataflowNode(const Dataflow::FNodeParameters& InParam, FGuid InGuid = FGuid::NewGuid())
+		: FDataflowTerminalNode(InParam, InGuid)
+	{
+		RegisterInputConnection(&Collection);
+		RegisterOutputConnection(&Collection, &Collection);
+	}
+
+	virtual void SetAssetValue(TObjectPtr<UObject> Asset, Dataflow::FContext& Context) const override
+	{
+		if (UChaosClothAsset* ClothAsset = Cast<UChaosClothAsset>(Asset.Get()))
+		{
+			// TODO: Modify ClothAsset directly
+		}
+	}
+
+	virtual void Evaluate(Dataflow::FContext& Context, const FDataflowOutput* Out) const override
+	{
+		const FManagedArrayCollection& InCollection = GetValue<FManagedArrayCollection>(Context, &Collection);
+		SetValue<FManagedArrayCollection>(Context, InCollection, &Collection);
+	}
+};
+
+namespace Dataflow
+{
+	void RegisterClothDataflowNodes();
+}
+
diff --git a/Engine/Plugins/Experimental/ChaosClothAssetEditor/Source/ChaosClothAssetEditor/Private/ChaosClothAsset/ClothEditorToolkit.cpp b/Engine/Plugins/Experimental/ChaosClothAssetEditor/Source/ChaosClothAssetEditor/Private/ChaosClothAsset/ClothEditorToolkit.cpp
--- a/Engine/Plugins/Experimental/ChaosClothAssetEditor/Source/ChaosClothAssetEditor/Private/ChaosClothAsset/ClothEditorToolkit.cpp
+++ b/Engine/Plugins/Experimental/ChaosClothAssetEditor/Source/ChaosClothAssetEditor/Private/ChaosClothAsset/ClothEditorToolkit.cpp
@@ -22,11 +22,49 @@
 #include "Styling/AppStyle.h"
 #include "ChaosClothAsset/SClothCollectionOutliner.h"
 #include "Widgets/Input/SComboBox.h"
+#include "Dataflow/DataflowEditorActions.h"
+#include "Dataflow/DataflowGraphEditor.h"
+#include "Dataflow/DataflowObject.h"
+#include "Dataflow/DataflowSchema.h"
+#include "IStructureDetailsView.h"
+#include "Engine/Canvas.h"
+#include "PropertyEditorModule.h"
 
 #define LOCTEXT_NAMESPACE "ChaosClothAssetEditorToolkit"
 
 const FName FChaosClothAssetEditorToolkit::ClothPreviewTabID(TEXT("ChaosClothAssetEditor_ClothPreviewTab"));
 const FName FChaosClothAssetEditorToolkit::OutlinerTabID(TEXT("ChaosClothAssetEditor_OutlinerTab"));
+const FName FChaosClothAssetEditorToolkit::GraphCanvasTabId(TEXT("ChaosClothAssetEditor_GraphCanvas"));
+const FName FChaosClothAssetEditorToolkit::NodeDetailsTabId(TEXT("ChaosClothAssetEditor_NodeDetails"));
+
+
+namespace ClothEditorToolkitHelpers
+{
+	UDataflow* GetDataflowFrom(const UObject* InObject)
+	{
+		if (const UClass* Class = InObject->GetClass())
+		{
+			if (const FProperty* Property = Class->FindPropertyByName(FName("DataflowAsset")))
+			{
+				return *Property->ContainerPtrToValuePtr<UDataflow*>(InObject);
+			}
+		}
+		return nullptr;
+	}
+
+	FString GetDataflowTerminalFrom(const UObject* InObject)
+	{
+		if (const UClass* Class = InObject->GetClass())
+		{
+			if (const FProperty* Property = Class->FindPropertyByName(FName("DataflowTerminal")))
+			{
+				return *Property->ContainerPtrToValuePtr<FString>(InObject);
+			}
+		}
+		return FString();
+	}
+}
+
 
 FChaosClothAssetEditorToolkit::FChaosClothAssetEditorToolkit(UAssetEditor* InOwningAssetEditor)
 	: FBaseCharacterFXEditorToolkit(InOwningAssetEditor, FName("ChaosClothAssetEditor"))
@@ -42,7 +80,7 @@ FChaosClothAssetEditorToolkit::FChaosClothAssetEditorToolkit(UAssetEditor* InOwn
 	// Note: Changes to the layout should include a increment to the layout's ID, i.e.
 	// ChaosClothAssetEditorLayout[X] -> ChaosClothAssetEditorLayout[X+1]. Otherwise, layouts may be messed up
 	// without a full reset to layout defaults inside the editor.
-	StandaloneDefaultLayout = FTabManager::NewLayout(FName("ChaosClothAssetEditorLayout3"))
+	StandaloneDefaultLayout = FTabManager::NewLayout(FName("ChaosClothAssetEditorLayout4"))
 		->AddArea
 		(
 			FTabManager::NewPrimaryArea()->SetOrientation(Orient_Vertical)
@@ -52,33 +90,61 @@ FChaosClothAssetEditorToolkit::FChaosClothAssetEditorToolkit(UAssetEditor* InOwn
 				->Split
 				(
 					FTabManager::NewStack()
-					->SetSizeCoefficient(0.15f)
+					->SetSizeCoefficient(0.1f)
 					->SetExtensionId(UChaosClothAssetEditorUISubsystem::EditorSidePanelAreaName)
 					->SetHideTabWell(true)
 				)
 				->Split
 				(
-					FTabManager::NewStack()
-					->SetSizeCoefficient(0.35f)
-					->AddTab(ViewportTabID, ETabState::OpenedTab)
-					->SetExtensionId("RestSpaceViewportArea")
-					->SetHideTabWell(true)
-				)
-				->Split
-				(
-					FTabManager::NewStack()
-					->SetSizeCoefficient(0.35f)
-					->AddTab(ClothPreviewTabID, ETabState::OpenedTab)
-					->SetExtensionId("Viewport3DArea")
-					->SetHideTabWell(true)
+					FTabManager::NewSplitter()->SetOrientation(Orient_Vertical)
+					->Split
+					(
+						FTabManager::NewSplitter()->SetOrientation(Orient_Horizontal)
+						->Split
+						(
+							FTabManager::NewStack()
+							->SetSizeCoefficient(0.5f)
+							->AddTab(ViewportTabID, ETabState::OpenedTab)
+							->SetExtensionId("RestSpaceViewportArea")
+							->SetHideTabWell(true)
+						)
+						->Split
+						(
+							FTabManager::NewStack()
+							->SetSizeCoefficient(0.5f)
+							->AddTab(ClothPreviewTabID, ETabState::OpenedTab)
+							->SetExtensionId("Viewport3DArea")
+							->SetHideTabWell(true)
+						)
+					)
+					->Split
+					(
+						FTabManager::NewSplitter()->SetOrientation(Orient_Horizontal)
+						->Split
+						(
+							FTabManager::NewStack()
+							->SetSizeCoefficient(0.7f)
+							->AddTab(GraphCanvasTabId, ETabState::OpenedTab)
+							->SetExtensionId("GraphEditorArea")
+							->SetHideTabWell(true)
+						)
+						->Split
+						(
+							FTabManager::NewStack()
+							->SetSizeCoefficient(0.3f)
+							->AddTab(NodeDetailsTabId, ETabState::OpenedTab)
+							->SetExtensionId("NodeDetailsArea")
+							->SetHideTabWell(true)
+						)
+					)
 				)
 				->Split
 				(
 					FTabManager::NewStack()
-					->SetSizeCoefficient(0.15f)
+					->SetSizeCoefficient(0.3f)
 					->AddTab(DetailsTabID, ETabState::OpenedTab)
 					->AddTab(OutlinerTabID, ETabState::OpenedTab)
-					->SetExtensionId("Details")
+					->SetExtensionId("DetailsArea")
 					->SetHideTabWell(true)
 				)
 			)
@@ -130,6 +196,63 @@ FChaosClothAssetEditorToolkit::~FChaosClothAssetEditorToolkit()
 	EditorModeManager->DestroyMode(UChaosClothAssetEditorMode::EM_ChaosClothAssetEditorModeId);
 }
 
+UChaosClothAsset* FChaosClothAssetEditorToolkit::GetAsset() const
+{
+	TArray<TObjectPtr<UObject>> ObjectsToEdit;
+	OwningAssetEditor->GetObjectsToEdit(ObjectsToEdit);
+	
+	UObject* ObjectToEdit = nullptr;
+	if (ensure(ObjectsToEdit.Num() == 1))
+	{
+		ObjectToEdit = ObjectsToEdit[0];
+	}
+
+	return Cast<UChaosClothAsset>(ObjectToEdit);
+}
+
+void FChaosClothAssetEditorToolkit::CreateWidgets()
+{
+	FBaseCharacterFXEditorToolkit::CreateWidgets();
+
+	UChaosClothAsset* const ClothAsset = GetAsset();
+
+	if (ClothAsset)
+	{
+		Dataflow = ClothEditorToolkitHelpers::GetDataflowFrom(ClothAsset);
+
+		if (Dataflow)
+		{
+			DataflowTerminalPath = ClothEditorToolkitHelpers::GetDataflowTerminalFrom(ClothAsset);
+
+			Dataflow->Schema = UDataflowSchema::StaticClass();
+
+			NodeDetailsEditor = CreateNodeDetailsEditorWidget(ClothAsset);
+			GraphEditor = CreateGraphEditorWidget(Dataflow, NodeDetailsEditor);
+		}
+	}
+}
+
+void FChaosClothAssetEditorToolkit::Tick(float DeltaTime)
+{
+	UChaosClothAsset* const ClothAsset = GetAsset();
+
+	if (Dataflow && ClothAsset)
+	{
+		if (!DataflowContext)
+		{
+			DataflowContext = TSharedPtr<Dataflow::FEngineContext>(new Dataflow::FClothAssetDataflowContext(ClothAsset, Dataflow, Dataflow::FTimestamp::Invalid));
+			LastDataflowNodeTimestamp = Dataflow::FTimestamp::Invalid;
+		}
+		DataflowTerminalPath = ClothEditorToolkitHelpers::GetDataflowTerminalFrom(ClothAsset);
+		FDataflowEditorCommands::EvaluateTerminalNode(*DataflowContext.Get(), LastDataflowNodeTimestamp, Dataflow, nullptr, nullptr, ClothAsset, DataflowTerminalPath);
+	}
+}
+
+TStatId FChaosClothAssetEditorToolkit::GetStatId() const
+{
+	RETURN_QUICK_DECLARE_CYCLE_STAT(FChaosClothAssetEditorToolkit, STATGROUP_Tickables);
+}
+
 // This gets used to label the editor's tab in the window that opens.
 FText FChaosClothAssetEditorToolkit::GetToolkitName() const
 {
@@ -202,6 +325,16 @@ void FChaosClothAssetEditorToolkit::RegisterTabSpawners(const TSharedRef<FTabMan
 		.SetGroup(AssetEditorTabsCategory.ToSharedRef())
 		.SetIcon(FSlateIcon(FAppStyle::GetAppStyleSetName(), "LevelEditor.Tabs.Outliner"));
 
+	InTabManager->RegisterTabSpawner(GraphCanvasTabId, FOnSpawnTab::CreateSP(this, &FChaosClothAssetEditorToolkit::SpawnTab_GraphCanvas))
+		.SetDisplayName(LOCTEXT("DataflowTab", "Graph"))
+		.SetGroup(AssetEditorTabsCategory.ToSharedRef())
+		.SetIcon(FSlateIcon(FAppStyle::GetAppStyleSetName(), "GraphEditor.EventGraph_16x"));
+
+	InTabManager->RegisterTabSpawner(NodeDetailsTabId, FOnSpawnTab::CreateSP(this, &FChaosClothAssetEditorToolkit::SpawnTab_NodeDetails))
+		.SetDisplayName(LOCTEXT("NodeDetailsTab", "Node Details"))
+		.SetGroup(AssetEditorTabsCategory.ToSharedRef())
+		.SetIcon(FSlateIcon(FAppStyle::GetAppStyleSetName(), "LevelEditor.Tabs.Details"));
+
 }
 
 void FChaosClothAssetEditorToolkit::UnregisterTabSpawners(const TSharedRef<FTabManager>& InTabManager)
@@ -211,6 +344,8 @@ void FChaosClothAssetEditorToolkit::UnregisterTabSpawners(const TSharedRef<FTabM
 	InTabManager->UnregisterTabSpawner(ClothPreviewTabID);
 	InTabManager->UnregisterTabSpawner(ViewportTabID);
 	InTabManager->UnregisterTabSpawner(DetailsTabID);
+	InTabManager->UnregisterTabSpawner(GraphCanvasTabId);
+	InTabManager->UnregisterTabSpawner(NodeDetailsTabId);
 }
 
 bool FChaosClothAssetEditorToolkit::OnRequestClose()
@@ -319,31 +454,30 @@ TSharedRef<SDockTab> FChaosClothAssetEditorToolkit::SpawnTab_Outliner(const FSpa
 	return DockableTab;
 }
 
-//// This is bound in RegisterTabSpawners() to create the panel on the left. The panel is filled in by the mode.
-//TSharedRef<SDockTab> FChaosClothAssetEditorToolkit::SpawnTab_InteractiveToolsPanel(const FSpawnTabArgs& Args)
-//{
-//	TSharedRef<SDockTab> ToolsPanel = SNew(SDockTab);
-//
-//	UChaosClothAssetEditorMode* ClothMode = Cast<UChaosClothAssetEditorMode>(EditorModeManager->GetActiveScriptableMode(UChaosClothAssetEditorMode::EM_ChaosClothAssetEditorModeId));
-//	if (!ClothMode)
-//	{
-//		// This is where we will drop out on the first call to this callback, when the mode does not yet
-//		// exist. There is probably a place where we could safely intialize the mode to make sure that
-//		// it is around before the first call, but it seems cleaner to just do the mode initialization 
-//		// in PostInitAssetEditor and fill in the tab at that time.
-//		// Later calls to this callback will occur if the user closes and restores the tab, and they
-//		// will continue past this point to allow the tab to be refilled.
-//		return ToolsPanel;
-//	}
-//	TSharedPtr<FModeToolkit> ClothModeToolkit = ClothMode->GetToolkit().Pin();
-//	if (!ClothModeToolkit.IsValid())
-//	{
-//		return ToolsPanel;
-//	}
-//
-//	ToolsPanel->SetContent(ClothModeToolkit->GetInlineContent().ToSharedRef());
-//	return ToolsPanel;
-//}
+TSharedRef<SDockTab> FChaosClothAssetEditorToolkit::SpawnTab_GraphCanvas(const FSpawnTabArgs& Args)
+{
+	check(Args.GetTabId() == GraphCanvasTabId);
+
+	TSharedRef<SDockTab> SpawnedTab = SNew(SDockTab)
+		.Label(LOCTEXT("DataflowEditor_Dataflow_TabTitle", "Graph"))
+		[
+			GraphEditor.ToSharedRef()
+		];
+
+	return SpawnedTab;
+}
+
+TSharedRef<SDockTab> FChaosClothAssetEditorToolkit::SpawnTab_NodeDetails(const FSpawnTabArgs& Args)
+{
+	check(Args.GetTabId() == NodeDetailsTabId);
+
+	return SNew(SDockTab)
+		.Label(LOCTEXT("DataflowEditor_NodeDetails_TabTitle", "Node Details"))
+		[
+			NodeDetailsEditor->GetWidget()->AsShared()
+		];
+}
+
 
 // Called from FBaseAssetToolkit::CreateWidgets to populate ViewportClient, but otherwise only used 
 // in our own viewport delegate.
@@ -508,4 +642,84 @@ void FChaosClothAssetEditorToolkit::InitDetailsViewPanel()
 
 }
 
+TSharedRef<SGraphEditor> FChaosClothAssetEditorToolkit::CreateGraphEditorWidget(UDataflow* DataflowToEdit, TSharedPtr<IStructureDetailsView> InNodeDetailsEditor)
+{
+	ensure(DataflowToEdit);
+	using namespace Dataflow;
+
+	FDataflowEditorCommands::FGraphEvaluationCallback Evaluate = [this, DataflowToEdit](FDataflowNode* Node, FDataflowOutput* Out)
+	{
+		if (DataflowToEdit)
+		{
+			UObject* const Asset = GetAsset();
+
+			if (!DataflowContext)
+			{
+				DataflowContext = TSharedPtr<Dataflow::FEngineContext>(new Dataflow::FClothAssetDataflowContext(Asset, Dataflow, Dataflow::FTimestamp::Invalid));
+			}
+			LastDataflowNodeTimestamp = Dataflow::FTimestamp::Invalid;
+			
+			FDataflowEditorCommands::EvaluateTerminalNode(*DataflowContext.Get(), LastDataflowNodeTimestamp, Dataflow, Node, Out, Asset, DataflowTerminalPath);
+		}
+	};
+
+	SGraphEditor::FGraphEditorEvents InEvents;
+	InEvents.OnVerifyTextCommit = FOnNodeVerifyTextCommit::CreateSP(this, &FChaosClothAssetEditorToolkit::OnNodeVerifyTitleCommit);
+	InEvents.OnTextCommitted = FOnNodeTextCommitted::CreateSP(this, &FChaosClothAssetEditorToolkit::OnNodeTitleCommitted);
+
+	return SNew(SDataflowGraphEditor, DataflowToEdit)
+		.GraphToEdit(DataflowToEdit)
+		.GraphEvents(InEvents)
+		.DetailsView(InNodeDetailsEditor)
+		.EvaluateGraph(Evaluate);
+
+}
+
+TSharedPtr<IStructureDetailsView> FChaosClothAssetEditorToolkit::CreateNodeDetailsEditorWidget(UObject* ObjectToEdit)
+{
+	ensure(ObjectToEdit);
+	FPropertyEditorModule& PropertyEditorModule = FModuleManager::LoadModuleChecked<FPropertyEditorModule>(TEXT("PropertyEditor"));
+
+	FDetailsViewArgs DetailsViewArgs;
+	{
+		DetailsViewArgs.bAllowSearch = false;
+		DetailsViewArgs.bHideSelectionTip = true;
+		DetailsViewArgs.bLockable = false;
+		DetailsViewArgs.bSearchInitialKeyFocus = true;
+		DetailsViewArgs.bUpdatesFromSelection = false;
+		DetailsViewArgs.NotifyHook = nullptr;
+		DetailsViewArgs.bShowOptions = true;
+		DetailsViewArgs.bShowModifiedPropertiesOption = false;
+		DetailsViewArgs.bShowScrollBar = false;
+	}
+
+	FStructureDetailsViewArgs StructureViewArgs;
+	{
+		StructureViewArgs.bShowObjects = true;
+		StructureViewArgs.bShowAssets = true;
+		StructureViewArgs.bShowClasses = true;
+		StructureViewArgs.bShowInterfaces = true;
+	}
+	TSharedPtr<IStructureDetailsView> NodeDetailsView = PropertyEditorModule.CreateStructureDetailView(DetailsViewArgs, StructureViewArgs, nullptr);
+	NodeDetailsView->GetDetailsView()->SetObject(ObjectToEdit);
+	NodeDetailsView->GetOnFinishedChangingPropertiesDelegate().AddSP(this, &FChaosClothAssetEditorToolkit::OnPropertyValueChanged);
+
+	return NodeDetailsView;
+}
+
+void FChaosClothAssetEditorToolkit::OnPropertyValueChanged(const FPropertyChangedEvent& PropertyChangedEvent)
+{
+	FDataflowEditorCommands::OnPropertyValueChanged(Dataflow, DataflowContext, LastDataflowNodeTimestamp, PropertyChangedEvent);
+}
+
+bool FChaosClothAssetEditorToolkit::OnNodeVerifyTitleCommit(const FText& NewText, UEdGraphNode* GraphNode, FText& OutErrorMessage) const
+{
+	return FDataflowEditorCommands::OnNodeVerifyTitleCommit(NewText, GraphNode, OutErrorMessage);
+}
+
+void FChaosClothAssetEditorToolkit::OnNodeTitleCommitted(const FText& InNewText, ETextCommit::Type InCommitType, UEdGraphNode* GraphNode) const
+{
+	FDataflowEditorCommands::OnNodeTitleCommitted(InNewText, InCommitType, GraphNode);
+}
+
 #undef LOCTEXT_NAMESPACE
diff --git a/Engine/Plugins/Experimental/ChaosClothAssetEditor/Source/ChaosClothAssetEditor/Private/ChaosClothAsset/ClothEditorToolkit.h b/Engine/Plugins/Experimental/ChaosClothAssetEditor/Source/ChaosClothAssetEditor/Private/ChaosClothAsset/ClothEditorToolkit.h
--- a/Engine/Plugins/Experimental/ChaosClothAssetEditor/Source/ChaosClothAssetEditor/Private/ChaosClothAsset/ClothEditorToolkit.h
+++ b/Engine/Plugins/Experimental/ChaosClothAssetEditor/Source/ChaosClothAssetEditor/Private/ChaosClothAsset/ClothEditorToolkit.h
@@ -3,11 +3,30 @@
 #pragma once
 
 #include "BaseCharacterFXEditorToolkit.h"
+#include "Dataflow/DataflowObjectInterface.h"
 
 class FAdvancedPreviewScene;
 class FChaosClothAssetEditor3DViewportClient;
 template<typename T> class SComboBox;
 class SClothCollectionOutliner;
+class UDataflow;
+class UChaosClothAsset;
+class SGraphEditor;
+class IStructureDetailsView;
+class UEdGraphNode;
+
+namespace Dataflow
+{
+	class CHAOSCLOTHASSETEDITOR_API FClothAssetDataflowContext : public TEngineContext<FContextSingle>
+	{
+	public:
+		DATAFLOW_CONTEXT_INTERNAL(TEngineContext<FContextSingle>, FClothAssetDataflowContext);
+
+		FClothAssetDataflowContext(UObject* InOwner, UDataflow* InGraph, FTimestamp InTimestamp)
+			: Super(InOwner, InGraph, InTimestamp)
+		{}
+	};
+}
 
 /**
  * The toolkit is supposed to act as the UI manager for the asset editor. It's responsible 
@@ -17,7 +36,7 @@ class SClothCollectionOutliner;
  * initializing the Cloth mode.
  * Thus, the FChaosClothAssetEditorToolkit ends up being the central place for the Cloth Asset Editor setup.
  */
-class CHAOSCLOTHASSETEDITOR_API FChaosClothAssetEditorToolkit : public FBaseCharacterFXEditorToolkit
+class CHAOSCLOTHASSETEDITOR_API FChaosClothAssetEditorToolkit : public FBaseCharacterFXEditorToolkit, public FTickableEditorObject
 {
 public:
 
@@ -39,8 +58,8 @@ public:
 	virtual void UnregisterTabSpawners(const TSharedRef<FTabManager>& TabManager) override;
 	virtual bool OnRequestClose() override;
 
-	// TODO: Implement this if we want to apply tool changes before saving an asset
-	//virtual void SaveAsset_Execute() override;
+	// FBaseAssetToolkit
+	virtual void CreateWidgets() override;
 
 	// IAssetEditorInstance
 	// This is important because if this returns true, attempting to edit a static mesh that is
@@ -48,6 +67,11 @@ public:
 	// TODO: Change this if we create a dedicated Cloth Asset
 	virtual bool IsPrimaryEditor() const override { return false; };
 
+	// FTickableEditorObject
+	virtual void Tick(float DeltaTime) override;
+	virtual bool IsTickable() const override { return true; }
+	virtual TStatId GetStatId() const override;
+
 protected:
 
 	// FBaseAssetToolkit
@@ -69,8 +93,13 @@ protected:
 
 private:
 
+	// Return the cloth asset held by the Cloth Editor
+	UChaosClothAsset* GetAsset() const;
+
 	TSharedRef<SDockTab> SpawnTab_ClothPreview(const FSpawnTabArgs& Args);
 	TSharedRef<SDockTab> SpawnTab_Outliner(const FSpawnTabArgs& Args);
+	TSharedRef<SDockTab> SpawnTab_GraphCanvas(const FSpawnTabArgs& Args);
+	TSharedRef<SDockTab> SpawnTab_NodeDetails(const FSpawnTabArgs& Args);
 
 	/** Scene in which the 3D sim space preview meshes live. */
 	TUniquePtr<FAdvancedPreviewScene> ClothPreviewScene;
@@ -89,5 +118,28 @@ private:
 	TSharedPtr<SComboBox<FName>> SelectedGroupNameComboBox;
 	TArray<FName> ClothCollectionGroupNames;		// Data source for SelectedGroupNameComboBox
 
+	//~ Begin Dataflow support
+	 
+	UDataflow* Dataflow = nullptr;
+
+	static const FName GraphCanvasTabId;
+	TSharedPtr<SGraphEditor> GraphEditor;
+	TSharedRef<SGraphEditor> CreateGraphEditorWidget(UDataflow* ObjectToEdit, TSharedPtr<IStructureDetailsView> PropertiesEditor);
+
+	static const FName NodeDetailsTabId;
+	TSharedPtr<IStructureDetailsView> NodeDetailsEditor;
+	TSharedPtr<IStructureDetailsView> CreateNodeDetailsEditorWidget(UObject* ObjectToEdit);
+
+	FString DataflowTerminalPath = "";
+	TSharedPtr<Dataflow::FEngineContext> DataflowContext;
+	Dataflow::FTimestamp LastDataflowNodeTimestamp = Dataflow::FTimestamp::Invalid;
+
+	// DataflowEditorActions
+	void OnPropertyValueChanged(const FPropertyChangedEvent& PropertyChangedEvent);
+	bool OnNodeVerifyTitleCommit(const FText& NewText, UEdGraphNode* GraphNode, FText& OutErrorMessage) const;
+	void OnNodeTitleCommitted(const FText& InNewText, ETextCommit::Type InCommitType, UEdGraphNode* GraphNode) const;
+
+	//~ End Dataflow support
+
 };
 
