ClothEditor: tidy up the ClothEditorToolkit class
- make as much private as possible
- remove dead code carried over from the UV Editor that might not be relevant to our editor
- make naming conventions more consistent

Should not have any significant behavior changes

diff --git a/Engine/Plugins/Experimental/ChaosClothAssetEditor/Source/ChaosClothAssetEditor/Private/ChaosClothAsset/ClothEditorToolkit.cpp b/Engine/Plugins/Experimental/ChaosClothAssetEditor/Source/ChaosClothAssetEditor/Private/ChaosClothAsset/ClothEditorToolkit.cpp
--- a/Engine/Plugins/Experimental/ChaosClothAssetEditor/Source/ChaosClothAssetEditor/Private/ChaosClothAsset/ClothEditorToolkit.cpp
+++ b/Engine/Plugins/Experimental/ChaosClothAssetEditor/Source/ChaosClothAssetEditor/Private/ChaosClothAsset/ClothEditorToolkit.cpp
@@ -177,12 +177,6 @@
 			.EditorViewportClient(ClothPreviewViewportClient);
 	};
 
-	// This API object serves as a communication point between the viewport toolbars and the tools.
-	// We create it here so that we can pass it both into the rest space viewport and when we initialize 
-	// the mode.
-	//ViewportButtonsAPI = NewObject<UClothToolViewportButtonsAPI>();
-
-
 	FPreviewScene::ConstructionValues SceneArgs;
 	ObjectScene = MakeUnique<FPreviewScene>(SceneArgs);
 
@@ -196,20 +190,68 @@
 	EditorModeManager->DestroyMode(UChaosClothAssetEditorMode::EM_ChaosClothAssetEditorModeId);
 }
 
-UChaosClothAsset* FChaosClothAssetEditorToolkit::GetAsset() const
+
+//~ Begin FTickableEditorObject overrides
+
+void FChaosClothAssetEditorToolkit::Tick(float DeltaTime)
 {
-	TArray<TObjectPtr<UObject>> ObjectsToEdit;
-	OwningAssetEditor->GetObjectsToEdit(ObjectsToEdit);
-	
-	UObject* ObjectToEdit = nullptr;
-	if (ensure(ObjectsToEdit.Num() == 1))
+	UChaosClothAsset* const ClothAsset = GetAsset();
+
+	if (Dataflow && ClothAsset)
 	{
-		ObjectToEdit = ObjectsToEdit[0];
+		if (!DataflowContext)
+		{
+			DataflowContext = TSharedPtr<Dataflow::FEngineContext>(new Dataflow::FClothAssetDataflowContext(ClothAsset, Dataflow, Dataflow::FTimestamp::Invalid));
+			LastDataflowNodeTimestamp = Dataflow::FTimestamp::Invalid;
+		}
+		DataflowTerminalPath = ClothEditorToolkitHelpers::GetDataflowTerminalFrom(ClothAsset);
+		FDataflowEditorCommands::EvaluateTerminalNode(*DataflowContext.Get(), LastDataflowNodeTimestamp, Dataflow, nullptr, nullptr, ClothAsset, DataflowTerminalPath);
 	}
-
-	return Cast<UChaosClothAsset>(ObjectToEdit);
 }
 
+TStatId FChaosClothAssetEditorToolkit::GetStatId() const
+{
+	RETURN_QUICK_DECLARE_CYCLE_STAT(FChaosClothAssetEditorToolkit, STATGROUP_Tickables);
+}
+
+
+//~ End FTickableEditorObject overrides
+
+
+//~ Begin FBaseCharacterFXEditorToolkit overrides
+
+FEditorModeID FChaosClothAssetEditorToolkit::GetEditorModeId() const
+{
+	return UChaosClothAssetEditorMode::EM_ChaosClothAssetEditorModeId;
+}
+
+
+void FChaosClothAssetEditorToolkit::InitializeEdMode(UBaseCharacterFXEditorMode* EdMode)
+{
+	UChaosClothAssetEditorMode* ClothMode = Cast<UChaosClothAssetEditorMode>(EdMode);
+	check(ClothMode);
+
+	// The mode will need to be able to get to the live preview world
+	ClothMode->SetPreviewWorld(ClothPreviewEditorModeManager->GetWorld());
+
+	TArray<TObjectPtr<UObject>> ObjectsToEdit;
+	OwningAssetEditor->GetObjectsToEdit(ObjectsToEdit);
+
+	ClothMode->InitializeTargets(ObjectsToEdit);
+}
+
+void FChaosClothAssetEditorToolkit::CreateEditorModeUILayer()
+{
+	TSharedPtr<class IToolkitHost> PinnedToolkitHost = ToolkitHost.Pin();
+	check(PinnedToolkitHost.IsValid());
+	ModeUILayer = MakeShared<FChaosClothAssetEditorModeUILayer>(PinnedToolkitHost.Get());
+}
+
+//~ End FBaseCharacterFXEditorToolkit overrides
+
+
+//~ Begin FBaseAssetToolkit overrides
+
 void FChaosClothAssetEditorToolkit::CreateWidgets()
 {
 	FBaseCharacterFXEditorToolkit::CreateWidgets();
@@ -232,27 +274,149 @@
 	}
 }
 
-void FChaosClothAssetEditorToolkit::Tick(float DeltaTime)
+// Called from FBaseAssetToolkit::CreateWidgets. The delegate call path goes through FAssetEditorToolkit::InitAssetEditor
+// and FBaseAssetToolkit::SpawnTab_Viewport.
+AssetEditorViewportFactoryFunction FChaosClothAssetEditorToolkit::GetViewportDelegate()
 {
-	UChaosClothAsset* const ClothAsset = GetAsset();
+	AssetEditorViewportFactoryFunction TempViewportDelegate = [this](FAssetEditorViewportConstructionArgs InArgs)
+	{
+		return SAssignNew(RestSpaceViewport, SChaosClothAssetEditorRestSpaceViewport, InArgs)
+			.EditorViewportClient(ViewportClient);
+	};
 
-	if (Dataflow && ClothAsset)
-	{
-		if (!DataflowContext)
-		{
-			DataflowContext = TSharedPtr<Dataflow::FEngineContext>(new Dataflow::FClothAssetDataflowContext(ClothAsset, Dataflow, Dataflow::FTimestamp::Invalid));
-			LastDataflowNodeTimestamp = Dataflow::FTimestamp::Invalid;
-		}
-		DataflowTerminalPath = ClothEditorToolkitHelpers::GetDataflowTerminalFrom(ClothAsset);
-		FDataflowEditorCommands::EvaluateTerminalNode(*DataflowContext.Get(), LastDataflowNodeTimestamp, Dataflow, nullptr, nullptr, ClothAsset, DataflowTerminalPath);
-	}
+	return TempViewportDelegate;
 }
 
-TStatId FChaosClothAssetEditorToolkit::GetStatId() const
+// Called from FBaseAssetToolkit::CreateWidgets to populate ViewportClient, but otherwise only used 
+// in our own viewport delegate.
+TSharedPtr<FEditorViewportClient> FChaosClothAssetEditorToolkit::CreateEditorViewportClient() const
 {
-	RETURN_QUICK_DECLARE_CYCLE_STAT(FChaosClothAssetEditorToolkit, STATGROUP_Tickables);
+	// Note that we can't reliably adjust the viewport client here because we will be passing it
+	// into the viewport created by the viewport delegate we get from GetViewportDelegate(), and
+	// that delegate may (will) affect the settings based on FAssetEditorViewportConstructionArgs,
+	// namely ViewportType.
+	// Instead, we do viewport client adjustment in PostInitAssetEditor().
+	check(EditorModeManager.IsValid());
+	return MakeShared<FChaosClothEditorRestSpaceViewportClient>(EditorModeManager.Get(), ObjectScene.Get());
 }
 
+//~ End FBaseAssetToolkit overrides
+
+
+//~ Begin FAssetEditorToolkit overrides
+
+void FChaosClothAssetEditorToolkit::AddViewportOverlayWidget(TSharedRef<SWidget> InViewportOverlayWidget)
+{
+	TSharedPtr<SChaosClothAssetEditorRestSpaceViewport> ViewportWidget = StaticCastSharedPtr<SChaosClothAssetEditorRestSpaceViewport>(ViewportTabContent->GetFirstViewport());
+	ViewportWidget->AddOverlayWidget(InViewportOverlayWidget);
+}
+
+void FChaosClothAssetEditorToolkit::RemoveViewportOverlayWidget(TSharedRef<SWidget> InViewportOverlayWidget)
+{
+	TSharedPtr<SChaosClothAssetEditorRestSpaceViewport> ViewportWidget = StaticCastSharedPtr<SChaosClothAssetEditorRestSpaceViewport>(ViewportTabContent->GetFirstViewport());
+	ViewportWidget->RemoveOverlayWidget(InViewportOverlayWidget);
+}
+
+bool FChaosClothAssetEditorToolkit::OnRequestClose()
+{
+	// Note: This needs a bit of adjusting, because currently OnRequestClose seems to be 
+	// called multiple times when the editor itself is being closed. We can take the route 
+	// of NiagaraScriptToolkit and remember when changes are discarded, but this can cause
+	// issues if the editor close sequence is interrupted due to some other asset editor.
+
+	UChaosClothAssetEditorMode* ClothEdMode = Cast<UChaosClothAssetEditorMode>(EditorModeManager->GetActiveScriptableMode(UChaosClothAssetEditorMode::EM_ChaosClothAssetEditorModeId));
+	if (!ClothEdMode) {
+		// If we don't have a valid mode, because the OnRequestClose is currently being called multiple times,
+		// simply return true because there's nothing left to do.
+		return true;
+	}
+
+	// Give any active modes a chance to shutdown while the toolkit host is still alive
+	// This is super important to do, otherwise currently opened tabs won't be marked as "closed".
+	// This results in tabs not being properly recycled upon reopening the editor and tab
+	// duplication for each opening event.
+	GetEditorModeManager().ActivateDefaultMode();
+
+	return FAssetEditorToolkit::OnRequestClose();
+}
+
+void FChaosClothAssetEditorToolkit::PostInitAssetEditor()
+{
+	FBaseCharacterFXEditorToolkit::PostInitAssetEditor();
+
+	// Custom viewport setup
+
+	auto SetCommonViewportClientOptions = [](FEditorViewportClient* Client)
+	{
+		// Normally the bIsRealtime flag is determined by whether the connection is remote, but our
+		// tools require always being ticked.
+		Client->SetRealtime(true);
+
+		// Disable motion blur effects that cause our renders to "fade in" as things are moved
+		Client->EngineShowFlags.SetTemporalAA(false);
+		Client->EngineShowFlags.SetAntiAliasing(true);
+		Client->EngineShowFlags.SetMotionBlur(false);
+
+		// Disable the dithering of occluded portions of gizmos.
+		Client->EngineShowFlags.SetOpaqueCompositeEditorPrimitives(true);
+
+		// Disable hardware occlusion queries, which make it harder to use vertex shaders to pull materials
+		// toward camera for z ordering because non-translucent materials start occluding themselves (once
+		// the component bounds are behind the displaced geometry).
+		Client->EngineShowFlags.SetDisableOcclusionQueries(true);
+
+		// Default FOV of 90 degrees causes a fair bit of lens distortion, especially noticeable with smaller viewports
+		Client->ViewFOV = 45.0;
+	};
+
+	{
+		// when CreateEditorViewportClient() is called, RestSpaceViewport is null. Set it here instead
+		TSharedPtr<FChaosClothEditorRestSpaceViewportClient> VC = StaticCastSharedPtr<FChaosClothEditorRestSpaceViewportClient>(ViewportClient);
+		VC->SetEditorViewportWidget(RestSpaceViewport);
+	}
+
+	SetCommonViewportClientOptions(ViewportClient.Get());
+
+	// Lit gives us the most options in terms of the materials we can use.
+	ViewportClient->SetViewMode(EViewModeIndex::VMI_Lit);
+
+	UChaosClothAssetEditorMode* const ClothMode = CastChecked<UChaosClothAssetEditorMode>(EditorModeManager->GetActiveScriptableMode(UChaosClothAssetEditorMode::EM_ChaosClothAssetEditorModeId));
+	
+	const TWeakPtr<FViewportClient> WeakViewportClient(ViewportClient);
+	ClothMode->SetRestSpaceViewportClient(StaticCastWeakPtr<FChaosClothEditorRestSpaceViewportClient>(WeakViewportClient));
+
+	// Note: We force the cloth preview viewport to open, since some ViewportClient functions are not robust to having no viewport.
+	// See UE-114649
+	if (!TabManager->FindExistingLiveTab(ClothPreviewTabID))
+	{
+		TabManager->TryInvokeTab(ClothPreviewTabID);
+	}
+
+	// We need the viewport client to start out focused, or else it won't get ticked until
+	// we click inside it.
+	ViewportClient->ReceivedFocus(ViewportClient->Viewport);
+
+	// Set up 3D viewport
+	ClothPreviewViewportClient->SetClothComponent(ClothMode->ClothComponent);
+	ClothPreviewViewportClient->SetClothEdMode(ClothMode);
+
+	SetCommonViewportClientOptions(ClothPreviewViewportClient.Get());
+	ClothPreviewViewportClient->SetInitialViewTransform(ELevelViewportType::LVT_Perspective, FVector(0, -100, 100), FRotator(0, 90, 0), DEFAULT_ORTHOZOOM);
+
+	if (ClothPreviewViewportClient->Viewport != nullptr)
+	{
+		FBox PreviewBounds = ClothMode->PreviewBoundingBox();
+		ClothPreviewViewportClient->FocusViewportOnBox(PreviewBounds);
+	}
+
+	InitDetailsViewPanel();
+}
+
+//~ End FAssetEditorToolkit overrides
+
+
+//~ Begin IToolkit overrides
+
 // This gets used to label the editor's tab in the window that opens.
 FText FChaosClothAssetEditorToolkit::GetToolkitName() const
 {
@@ -269,7 +433,7 @@
 	return FName("Cloth Editor");
 }
 
-// TODO: What is this actually used for?
+// Used to create a section in the Help menu for the cloth editor
 FText FChaosClothAssetEditorToolkit::GetBaseToolkitName() const
 {
 	return LOCTEXT("ChaosClothAssetEditorBaseName", "Cloth Editor");
@@ -348,54 +512,23 @@
 	InTabManager->UnregisterTabSpawner(NodeDetailsTabId);
 }
 
-bool FChaosClothAssetEditorToolkit::OnRequestClose()
+//~ End IToolkit overrides
+
+
+UChaosClothAsset* FChaosClothAssetEditorToolkit::GetAsset() const
 {
-	// Note: This needs a bit of adjusting, because currently OnRequestClose seems to be 
-	// called multiple times when the editor itself is being closed. We can take the route 
-	// of NiagaraScriptToolkit and remember when changes are discarded, but this can cause
-	// issues if the editor close sequence is interrupted due to some other asset editor.
-
-	UChaosClothAssetEditorMode* ClothEdMode = Cast<UChaosClothAssetEditorMode>(EditorModeManager->GetActiveScriptableMode(UChaosClothAssetEditorMode::EM_ChaosClothAssetEditorModeId));
-	if (!ClothEdMode) {
-		// If we don't have a valid mode, because the OnRequestClose is currently being called multiple times,
-		// simply return true because there's nothing left to do.
-		return true; 
+	TArray<TObjectPtr<UObject>> ObjectsToEdit;
+	OwningAssetEditor->GetObjectsToEdit(ObjectsToEdit);
+	
+	UObject* ObjectToEdit = nullptr;
+	if (ensure(ObjectsToEdit.Num() == 1))
+	{
+		ObjectToEdit = ObjectsToEdit[0];
 	}
 
-	// Give any active modes a chance to shutdown while the toolkit host is still alive
-    // This is super important to do, otherwise currently opened tabs won't be marked as "closed".
-    // This results in tabs not being properly recycled upon reopening the editor and tab
-    // duplication for each opening event.
-	GetEditorModeManager().ActivateDefaultMode();
-
-	return FAssetEditorToolkit::OnRequestClose();
+	return Cast<UChaosClothAsset>(ObjectToEdit);
 }
 
-void FChaosClothAssetEditorToolkit::AddViewportOverlayWidget(TSharedRef<SWidget> InViewportOverlayWidget) 
-{
-	TSharedPtr<SChaosClothAssetEditorRestSpaceViewport> ViewportWidget = StaticCastSharedPtr<SChaosClothAssetEditorRestSpaceViewport>(ViewportTabContent->GetFirstViewport());
-	ViewportWidget->AddOverlayWidget(InViewportOverlayWidget);
-}
-
-void FChaosClothAssetEditorToolkit::RemoveViewportOverlayWidget(TSharedRef<SWidget> InViewportOverlayWidget)
-{
-	TSharedPtr<SChaosClothAssetEditorRestSpaceViewport> ViewportWidget = StaticCastSharedPtr<SChaosClothAssetEditorRestSpaceViewport>(ViewportTabContent->GetFirstViewport());
-	ViewportWidget->RemoveOverlayWidget(InViewportOverlayWidget);
-}
-
-//// We override the "Save" button behavior slightly to apply our changes before saving the asset.
-//void FChaosClothAssetEditorToolkit::SaveAsset_Execute()
-//{
-//	UChaosClothAssetEditorMode* ClothMode = Cast<UChaosClothAssetEditorMode>(EditorModeManager->GetActiveScriptableMode(UChaosClothAssetEditorMode::EM_ChaosClothAssetEditorModeId));
-//	check(ClothMode);
-//	if (ClothMode->HaveUnappliedChanges())
-//	{
-//		ClothMode->ApplyChanges();
-//	}
-//
-//	FAssetEditorToolkit::SaveAsset_Execute();
-//}
-
 TSharedRef<SDockTab> FChaosClothAssetEditorToolkit::SpawnTab_ClothPreview(const FSpawnTabArgs& Args)
 {
 	TSharedRef<SDockTab> DockableTab = SNew(SDockTab);
@@ -419,9 +552,9 @@
 			.OnSelectionChanged(SComboBox<FName>::FOnSelectionChanged::CreateLambda(
 				[this](FName SelectedName, ESelectInfo::Type)
 				{
-					if (OutlinerView)
+					if (Outliner)
 					{
-						OutlinerView->SetSelectedGroupName(SelectedName);	// this will also rebuild the table
+						Outliner->SetSelectedGroupName(SelectedName);	// this will also rebuild the table
 					}
 				}))
 			.OnGenerateWidget(SComboBox<FName>::FOnGenerateWidget::CreateLambda(
@@ -434,9 +567,9 @@
 				SNew(STextBlock)
 				.Text_Lambda([this]()
 				{
-					if (OutlinerView)
+					if (Outliner)
 					{
-						return FText::FromName(OutlinerView->GetSelectedGroupName());
+						return FText::FromName(Outliner->GetSelectedGroupName());
 					}
 					else
 					{
@@ -448,7 +581,7 @@
 
 		+ SVerticalBox::Slot()
 		[
-			SAssignNew(OutlinerView, SClothCollectionOutliner)
+			SAssignNew(Outliner, SClothCollectionOutliner)
 		]
 	];
 	return DockableTab;
@@ -478,134 +611,6 @@
 		];
 }
 
-
-// Called from FBaseAssetToolkit::CreateWidgets to populate ViewportClient, but otherwise only used 
-// in our own viewport delegate.
-TSharedPtr<FEditorViewportClient> FChaosClothAssetEditorToolkit::CreateEditorViewportClient() const
-{
-	// Note that we can't reliably adjust the viewport client here because we will be passing it
-	// into the viewport created by the viewport delegate we get from GetViewportDelegate(), and
-	// that delegate may (will) affect the settings based on FAssetEditorViewportConstructionArgs,
-	// namely ViewportType.
-	// Instead, we do viewport client adjustment in PostInitAssetEditor().
-	check(EditorModeManager.IsValid());
-	return MakeShared<FChaosClothEditorRestSpaceViewportClient>(EditorModeManager.Get(), ObjectScene.Get());
-}
-
-// Called from FBaseAssetToolkit::CreateWidgets. The delegate call path goes through FAssetEditorToolkit::InitAssetEditor
-// and FBaseAssetToolkit::SpawnTab_Viewport.
-AssetEditorViewportFactoryFunction FChaosClothAssetEditorToolkit::GetViewportDelegate()
-{
-	AssetEditorViewportFactoryFunction TempViewportDelegate = [this](FAssetEditorViewportConstructionArgs InArgs)
-	{
-		return SAssignNew(RestSpaceViewport, SChaosClothAssetEditorRestSpaceViewport, InArgs)
-			.EditorViewportClient(ViewportClient);
-	};
-
-	return TempViewportDelegate;
-}
-
-
-void FChaosClothAssetEditorToolkit::PostInitAssetEditor()
-{
-	FBaseCharacterFXEditorToolkit::PostInitAssetEditor();
-
-	// Custom viewport setup
-
-	auto SetCommonViewportClientOptions = [](FEditorViewportClient* Client)
-	{
-		// Normally the bIsRealtime flag is determined by whether the connection is remote, but our
-		// tools require always being ticked.
-		Client->SetRealtime(true);
-
-		// Disable motion blur effects that cause our renders to "fade in" as things are moved
-		Client->EngineShowFlags.SetTemporalAA(false);
-		Client->EngineShowFlags.SetAntiAliasing(true);
-		Client->EngineShowFlags.SetMotionBlur(false);
-
-		// Disable the dithering of occluded portions of gizmos.
-		Client->EngineShowFlags.SetOpaqueCompositeEditorPrimitives(true);
-
-		// Disable hardware occlusion queries, which make it harder to use vertex shaders to pull materials
-		// toward camera for z ordering because non-translucent materials start occluding themselves (once
-		// the component bounds are behind the displaced geometry).
-		Client->EngineShowFlags.SetDisableOcclusionQueries(true);
-
-		// Default FOV of 90 degrees causes a fair bit of lens distortion, especially noticeable with smaller viewports
-		Client->ViewFOV = 45.0;
-	};
-
-	{
-		// when CreateEditorViewportClient() is called, RestSpaceViewport is null. Set it here instead
-		TSharedPtr<FChaosClothEditorRestSpaceViewportClient> VC = StaticCastSharedPtr<FChaosClothEditorRestSpaceViewportClient>(ViewportClient);
-		VC->SetEditorViewportWidget(RestSpaceViewport);
-	}
-	
-	SetCommonViewportClientOptions(ViewportClient.Get());
-
-	// Lit gives us the most options in terms of the materials we can use.
-	ViewportClient->SetViewMode(EViewModeIndex::VMI_Lit);
-
-	UChaosClothAssetEditorMode* ClothMode = Cast<UChaosClothAssetEditorMode>(EditorModeManager->GetActiveScriptableMode(UChaosClothAssetEditorMode::EM_ChaosClothAssetEditorModeId));
-	check(ClothMode);
-	const TWeakPtr<FViewportClient> WeakViewportClient(ViewportClient);
-	ClothMode->SetRestSpaceViewportClient(StaticCastWeakPtr<FChaosClothEditorRestSpaceViewportClient>(WeakViewportClient));
-
-	// Note: We force the cloth preview viewport to open, since some ViewportClient functions are not robust to having no viewport.
-	// See UE-114649
-	if (!TabManager->FindExistingLiveTab(ClothPreviewTabID))
-	{
-		TabManager->TryInvokeTab(ClothPreviewTabID);
-	}
-
-	// We need the viewport client to start out focused, or else it won't get ticked until
-	// we click inside it.
-	ViewportClient->ReceivedFocus(ViewportClient->Viewport);
-
-	// Set up 3D viewport
-	ClothPreviewViewportClient->SetClothComponent(ClothMode->ClothComponent);
-	ClothPreviewViewportClient->SetClothEdMode(ClothMode);
-
-	SetCommonViewportClientOptions(ClothPreviewViewportClient.Get());
-	ClothPreviewViewportClient->SetInitialViewTransform(ELevelViewportType::LVT_Perspective, FVector(0, -100, 100), FRotator(0, 90, 0), DEFAULT_ORTHOZOOM);
-
-	if (ClothPreviewViewportClient->Viewport != nullptr)
-	{
-		FBox PreviewBounds = ClothMode->PreviewBoundingBox();
-		ClothPreviewViewportClient->FocusViewportOnBox(PreviewBounds);
-	}
-
-	InitDetailsViewPanel();
-}
-
-
-FEditorModeID FChaosClothAssetEditorToolkit::GetEditorModeId() const
-{
-	return UChaosClothAssetEditorMode::EM_ChaosClothAssetEditorModeId;
-}
-
-
-void FChaosClothAssetEditorToolkit::InitializeEdMode(UBaseCharacterFXEditorMode* EdMode)
-{
-	UChaosClothAssetEditorMode* ClothMode = Cast<UChaosClothAssetEditorMode>(EdMode);
-	check(ClothMode);
-
-	// The mode will need to be able to get to the live preview world
-	ClothMode->SetPreviewWorld(ClothPreviewEditorModeManager->GetWorld());
-
-	TArray<TObjectPtr<UObject>> ObjectsToEdit;
-	OwningAssetEditor->GetObjectsToEdit(ObjectsToEdit);
-
-	ClothMode->InitializeTargets(ObjectsToEdit);
-}
-
-void FChaosClothAssetEditorToolkit::CreateEditorModeUILayer()
-{
-	TSharedPtr<class IToolkitHost> PinnedToolkitHost = ToolkitHost.Pin();
-	check(PinnedToolkitHost.IsValid());
-	ModeUILayer = MakeShared<FChaosClothAssetEditorModeUILayer>(PinnedToolkitHost.Get());
-}
-
 void FChaosClothAssetEditorToolkit::InitDetailsViewPanel()
 {
 	TArray<TObjectPtr<UObject>> ObjectsToEdit;
@@ -630,12 +635,12 @@
 		}
 	}
 	
-	if (OutlinerView.IsValid() && ClothCollection.IsValid())
+	if (Outliner.IsValid() && ClothCollection.IsValid())
 	{
 		const FName& SelectedGroupName = ClothCollection->SimVerticesGroup;
 		
-		OutlinerView->SetClothCollection(ClothCollection);
-		OutlinerView->SetSelectedGroupName(SelectedGroupName);
+		Outliner->SetClothCollection(ClothCollection);
+		Outliner->SetSelectedGroupName(SelectedGroupName);
 
 		ClothCollectionGroupNames = ClothCollection->GroupNames();
 	}
@@ -707,6 +712,8 @@
 	return NodeDetailsView;
 }
 
+//~ Begin DataflowEditorActions
+
 void FChaosClothAssetEditorToolkit::OnPropertyValueChanged(const FPropertyChangedEvent& PropertyChangedEvent)
 {
 	FDataflowEditorCommands::OnPropertyValueChanged(Dataflow, DataflowContext, LastDataflowNodeTimestamp, PropertyChangedEvent);
@@ -722,4 +729,6 @@
 	FDataflowEditorCommands::OnNodeTitleCommitted(InNewText, InCommitType, GraphNode);
 }
 
+//~ Ends DataflowEditorActions
+
 #undef LOCTEXT_NAMESPACE
diff --git a/Engine/Plugins/Experimental/ChaosClothAssetEditor/Source/ChaosClothAssetEditor/Private/ChaosClothAsset/ClothEditorToolkit.h b/Engine/Plugins/Experimental/ChaosClothAssetEditor/Source/ChaosClothAssetEditor/Private/ChaosClothAsset/ClothEditorToolkit.h
--- a/Engine/Plugins/Experimental/ChaosClothAssetEditor/Source/ChaosClothAssetEditor/Private/ChaosClothAsset/ClothEditorToolkit.h
+++ b/Engine/Plugins/Experimental/ChaosClothAssetEditor/Source/ChaosClothAssetEditor/Private/ChaosClothAsset/ClothEditorToolkit.h
@@ -17,7 +17,7 @@ class UEdGraphNode;
 
 namespace Dataflow
 {
-	class CHAOSCLOTHASSETEDITOR_API FClothAssetDataflowContext : public TEngineContext<FContextSingle>
+	class CHAOSCLOTHASSETEDITOR_API FClothAssetDataflowContext final : public TEngineContext<FContextSingle>
 	{
 	public:
 		DATAFLOW_CONTEXT_INTERNAL(TEngineContext<FContextSingle>, FClothAssetDataflowContext);
@@ -36,62 +36,49 @@ namespace Dataflow
  * initializing the Cloth mode.
  * Thus, the FChaosClothAssetEditorToolkit ends up being the central place for the Cloth Asset Editor setup.
  */
-class CHAOSCLOTHASSETEDITOR_API FChaosClothAssetEditorToolkit : public FBaseCharacterFXEditorToolkit, public FTickableEditorObject
+class CHAOSCLOTHASSETEDITOR_API FChaosClothAssetEditorToolkit final : public FBaseCharacterFXEditorToolkit, public FTickableEditorObject
 {
 public:
 
-	FChaosClothAssetEditorToolkit(UAssetEditor* InOwningAssetEditor);
+	explicit FChaosClothAssetEditorToolkit(UAssetEditor* InOwningAssetEditor);
 	virtual ~FChaosClothAssetEditorToolkit();
 
+private:
+
 	static const FName ClothPreviewTabID;
 	static const FName OutlinerTabID;
 
-	// FAssetEditorToolkit
-	virtual void AddViewportOverlayWidget(TSharedRef<SWidget> InViewportOverlayWidget) override;
-	virtual void RemoveViewportOverlayWidget(TSharedRef<SWidget> InViewportOverlayWidget) override;
-
-	virtual FText GetToolkitName() const override;
-	virtual FName GetToolkitFName() const override;
-	virtual FText GetBaseToolkitName() const override;
-	virtual FText GetToolkitToolTipText() const override;
-	virtual void RegisterTabSpawners(const TSharedRef<FTabManager>& InTabManager) override;
-	virtual void UnregisterTabSpawners(const TSharedRef<FTabManager>& TabManager) override;
-	virtual bool OnRequestClose() override;
-
-	// FBaseAssetToolkit
-	virtual void CreateWidgets() override;
-
-	// IAssetEditorInstance
-	// This is important because if this returns true, attempting to edit a static mesh that is
-	// open in the cloth editor may open the cloth editor instead of opening the static mesh editor.
-	// TODO: Change this if we create a dedicated Cloth Asset
-	virtual bool IsPrimaryEditor() const override { return false; };
-
 	// FTickableEditorObject
 	virtual void Tick(float DeltaTime) override;
 	virtual bool IsTickable() const override { return true; }
 	virtual TStatId GetStatId() const override;
 
-protected:
+	// FBaseCharacterFXEditorToolkit
+	virtual FEditorModeID GetEditorModeId() const override;
+	virtual void InitializeEdMode(UBaseCharacterFXEditorMode* EdMode) override;
+	virtual void CreateEditorModeUILayer() override;
 
 	// FBaseAssetToolkit
+	virtual void CreateWidgets() override;
 	virtual AssetEditorViewportFactoryFunction GetViewportDelegate() override;
 	virtual TSharedPtr<FEditorViewportClient> CreateEditorViewportClient() const override;
 
 	// FAssetEditorToolkit
+	virtual void AddViewportOverlayWidget(TSharedRef<SWidget> InViewportOverlayWidget) override;
+	virtual void RemoveViewportOverlayWidget(TSharedRef<SWidget> InViewportOverlayWidget) override;
+	virtual bool OnRequestClose() override;
 	virtual void PostInitAssetEditor() override;
 
-	// FBaseCharacterFXEditorToolkit
-	virtual FEditorModeID GetEditorModeId() const override;
-	virtual void InitializeEdMode(UBaseCharacterFXEditorMode* EdMode) override;
-	virtual void CreateEditorModeUILayer() override;
-
-	// Appearance customization points
-	// TODO: Implement these when we have an FChaosClothAssetEditorStyle class
-	//virtual const FSlateBrush* GetDefaultTabIcon() const override;
-	//virtual FLinearColor GetDefaultTabColor() const override;
+	// IAssetEditorInstance
+	virtual bool IsPrimaryEditor() const override { return true; };
 
-private:
+	// IToolkit
+	virtual FText GetToolkitName() const override;
+	virtual FName GetToolkitFName() const override;
+	virtual FText GetBaseToolkitName() const override;
+	virtual FText GetToolkitToolTipText() const override;
+	virtual void RegisterTabSpawners(const TSharedRef<FTabManager>& InTabManager) override;
+	virtual void UnregisterTabSpawners(const TSharedRef<FTabManager>& TabManager) override;
 
 	// Return the cloth asset held by the Cloth Editor
 	UChaosClothAsset* GetAsset() const;
@@ -101,6 +88,8 @@ private:
 	TSharedRef<SDockTab> SpawnTab_GraphCanvas(const FSpawnTabArgs& Args);
 	TSharedRef<SDockTab> SpawnTab_NodeDetails(const FSpawnTabArgs& Args);
 
+	void InitDetailsViewPanel();
+
 	/** Scene in which the 3D sim space preview meshes live. */
 	TUniquePtr<FAdvancedPreviewScene> ClothPreviewScene;
 
@@ -111,9 +100,7 @@ private:
 
 	TWeakPtr<SEditorViewport> RestSpaceViewport;
 
-	void InitDetailsViewPanel();
-	
-	TSharedPtr<SClothCollectionOutliner> OutlinerView;
+	TSharedPtr<SClothCollectionOutliner> Outliner;
 
 	TSharedPtr<SComboBox<FName>> SelectedGroupNameComboBox;
 	TArray<FName> ClothCollectionGroupNames;		// Data source for SelectedGroupNameComboBox
