Cloth Editor: hook up Dataflow node rendering to the cloth preview viewport
- Also enable rendering of currently selected nodes in preview viewport

diff --git a/Engine/Plugins/Experimental/ChaosClothAssetEditor/Source/ChaosClothAssetEditor/ChaosClothAssetEditor.Build.cs b/Engine/Plugins/Experimental/ChaosClothAssetEditor/Source/ChaosClothAssetEditor/ChaosClothAssetEditor.Build.cs
--- a/Engine/Plugins/Experimental/ChaosClothAssetEditor/Source/ChaosClothAssetEditor/ChaosClothAssetEditor.Build.cs
+++ b/Engine/Plugins/Experimental/ChaosClothAssetEditor/Source/ChaosClothAssetEditor/ChaosClothAssetEditor.Build.cs
@@ -57,7 +57,9 @@ public class ChaosClothAssetEditor : ModuleRules
 				"ChaosClothAssetEditorTools",
 				"DataflowCore",
 				"DataflowEditor",
-				"DataflowEngine"
+				"DataflowEngine",
+				"DataflowEnginePlugin",
+				"RenderCore"
 				// ... add private dependencies that you statically link with here ...	
 			}
 			);
diff --git a/Engine/Plugins/Experimental/ChaosClothAssetEditor/Source/ChaosClothAssetEditor/Private/ChaosClothAsset/ClothDataflowNodes.h b/Engine/Plugins/Experimental/ChaosClothAssetEditor/Source/ChaosClothAssetEditor/Private/ChaosClothAsset/ClothDataflowNodes.h
--- a/Engine/Plugins/Experimental/ChaosClothAssetEditor/Source/ChaosClothAssetEditor/Private/ChaosClothAsset/ClothDataflowNodes.h
+++ b/Engine/Plugins/Experimental/ChaosClothAssetEditor/Source/ChaosClothAssetEditor/Private/ChaosClothAsset/ClothDataflowNodes.h
@@ -5,6 +5,7 @@
 #include "Dataflow/DataflowCore.h"
 #include "ChaosClothAsset/ClothAsset.h"
 #include "ChaosClothAsset/ClothCollection.h"
+#include "GeometryCollection/GeometryCollection.h"
 #include "ClothDataflowNodes.generated.h"
 
 USTRUCT()
@@ -12,6 +13,7 @@ struct FClothAssetTerminalDataflowNode : public FDataflowTerminalNode
 {
 	GENERATED_USTRUCT_BODY()
 	DATAFLOW_NODE_DEFINE_INTERNAL(FClothAssetTerminalDataflowNode, "ClothAssetTerminal", "Cloth", "")
+	DATAFLOW_NODE_RENDER_TYPE(FGeometryCollection::StaticType(), "Collection")
 
 public:
 
diff --git a/Engine/Plugins/Experimental/ChaosClothAssetEditor/Source/ChaosClothAssetEditor/Private/ChaosClothAsset/ClothEditor3DViewportClient.cpp b/Engine/Plugins/Experimental/ChaosClothAssetEditor/Source/ChaosClothAssetEditor/Private/ChaosClothAsset/ClothEditor3DViewportClient.cpp
--- a/Engine/Plugins/Experimental/ChaosClothAssetEditor/Source/ChaosClothAssetEditor/Private/ChaosClothAsset/ClothEditor3DViewportClient.cpp
+++ b/Engine/Plugins/Experimental/ChaosClothAssetEditor/Source/ChaosClothAssetEditor/Private/ChaosClothAsset/ClothEditor3DViewportClient.cpp
@@ -2,8 +2,13 @@
 
 #include "ChaosClothAsset/ClothEditor3DViewportClient.h"
 #include "ChaosClothAsset/ClothEditorMode.h"
+#include "ChaosClothAsset/ClothEditorToolkit.h"
 #include "EditorModeManager.h"
 #include "ChaosClothAsset/ClothComponent.h"
+#include "Dataflow/DataflowEdNode.h"
+#include "Dataflow/DataflowObject.h"
+#include "Dataflow/DataflowActor.h"
+#include "Dataflow/DataflowRenderingFactory.h"
 
 FChaosClothAssetEditor3DViewportClient::FChaosClothAssetEditor3DViewportClient(FEditorModeTools* InModeTools,
 	FPreviewScene* InPreviewScene, 
@@ -18,6 +23,58 @@ FChaosClothAssetEditor3DViewportClient::FChaosClothAssetEditor3DViewportClient(F
 }
 
 
+void FChaosClothAssetEditor3DViewportClient::Tick(float DeltaSeconds)
+{
+	FEditorViewportClient::Tick(DeltaSeconds);
+
+	auto GetLatestTimestamp = [](const UDataflow* Dataflow, const Dataflow::FContext* Context) -> Dataflow::FTimestamp
+	{
+		if (Dataflow && Context)
+		{
+			return FMath::Max(Dataflow->GetRenderingTimestamp().Value, Context->GetTimestamp().Value);
+		}
+		return Dataflow::FTimestamp::Invalid;
+	};
+
+	if (ClothToolkit.IsValid())
+	{
+		if (const TSharedPtr<Dataflow::FContext> Context = ClothToolkit->GetDataflowContext())
+		{
+			if (const UDataflow* const Dataflow = ClothToolkit->GetDataflow())
+			{
+				if (UDataflowComponent* const DataflowComponent = ClothEdMode->GetDataflowComponent())
+				{
+					const Dataflow::FTimestamp SystemTimestamp = GetLatestTimestamp(Dataflow, Context.Get());
+
+					if (SystemTimestamp >= LastModifiedTimestamp)
+					{
+						if (Dataflow->GetRenderTargets().Num())
+						{
+							// Component Object Rendering
+							DataflowComponent->ResetRenderTargets();
+							DataflowComponent->SetDataflow(Dataflow);
+							DataflowComponent->SetContext(Context);
+							for (const UDataflowEdNode* const Node : Dataflow->GetRenderTargets())
+							{
+								DataflowComponent->AddRenderTarget(Node);
+							}
+						}
+						else
+						{
+							DataflowComponent->ResetRenderTargets();
+						}
+
+						LastModifiedTimestamp = GetLatestTimestamp(Dataflow, Context.Get()).Value + 1;
+					}
+				}
+			}
+		}
+	}
+
+	// Note: we don't tick the PreviewWorld here, that is done in UChaosClothAssetEditorMode::ModeTick()
+
+}
+
 void FChaosClothAssetEditor3DViewportClient::EnableRenderMeshWireframe(bool bEnable)
 {
 	bRenderMeshWireframe = bEnable;
@@ -38,6 +95,11 @@ void FChaosClothAssetEditor3DViewportClient::SetClothEdMode(TObjectPtr<UChaosClo
 	ClothEdMode = InClothEdMode;
 }
 
+void FChaosClothAssetEditor3DViewportClient::SetClothEditorToolkit(TSharedPtr<const FChaosClothAssetEditorToolkit> InClothToolkit)
+{
+	ClothToolkit = InClothToolkit;
+}
+
 void FChaosClothAssetEditor3DViewportClient::SoftResetSimulation()
 {
 	if (ClothEdMode)
diff --git a/Engine/Plugins/Experimental/ChaosClothAssetEditor/Source/ChaosClothAssetEditor/Private/ChaosClothAsset/ClothEditor3DViewportClient.h b/Engine/Plugins/Experimental/ChaosClothAssetEditor/Source/ChaosClothAssetEditor/Private/ChaosClothAsset/ClothEditor3DViewportClient.h
--- a/Engine/Plugins/Experimental/ChaosClothAssetEditor/Source/ChaosClothAssetEditor/Private/ChaosClothAsset/ClothEditor3DViewportClient.h
+++ b/Engine/Plugins/Experimental/ChaosClothAssetEditor/Source/ChaosClothAssetEditor/Private/ChaosClothAsset/ClothEditor3DViewportClient.h
@@ -2,11 +2,13 @@
 
 #pragma once
 
-#include "CoreMinimal.h"
 #include "EditorViewportClient.h"
+#include "Dataflow/DataflowNodeParameters.h"
+#include "GeometryCollection/ManagedArrayCollection.h"
 
 class UChaosClothComponent;
 class UChaosClothAssetEditorMode;
+class FChaosClothAssetEditorToolkit;
 
 /**
  * Viewport client for the 3d sim preview in the cloth editor. Currently same as editor viewport
@@ -25,6 +27,8 @@ public:
 	virtual void SetWidgetMode(UE::Widget::EWidgetMode NewMode) override {}
 	virtual UE::Widget::EWidgetMode GetWidgetMode() const override { return UE::Widget::EWidgetMode::WM_None; }
 
+	virtual void Tick(float DeltaSeconds) override;
+
 	void EnableSimMeshWireframe(bool bEnable ) { bSimMeshWireframe = bEnable; }
 	bool SimMeshWireframeEnabled() const { return bSimMeshWireframe; }
 
@@ -33,6 +37,7 @@ public:
 
 	void SetClothComponent(TObjectPtr<UChaosClothComponent> ClothComponent);
 	void SetClothEdMode(TObjectPtr<UChaosClothAssetEditorMode> ClothEdMode);
+	void SetClothEditorToolkit(TSharedPtr<const FChaosClothAssetEditorToolkit> ClothToolkit);
 
 	void SoftResetSimulation();
 	void HardResetSimulation();
@@ -48,12 +53,17 @@ private:
 
 	TObjectPtr<UChaosClothAssetEditorMode> ClothEdMode;
 
+	TSharedPtr<const FChaosClothAssetEditorToolkit> ClothToolkit;
+
 	// Debug draw of simulation meshes
 	virtual void Draw(const FSceneView* View, FPrimitiveDrawInterface* PDI) override;
-
-	UChaosClothAssetEditorMode* GetClothEdMode();
-	const UChaosClothAssetEditorMode* GetClothEdMode() const;
-
+	
 	bool bSimMeshWireframe = true;
 	bool bRenderMeshWireframe = false;
+
+	//~ Dataflow render support
+
+	Dataflow::FTimestamp LastModifiedTimestamp = Dataflow::FTimestamp::Invalid;
+
+	//~ Dataflow render support
 };
diff --git a/Engine/Plugins/Experimental/ChaosClothAssetEditor/Source/ChaosClothAssetEditor/Private/ChaosClothAsset/ClothEditorMode.cpp b/Engine/Plugins/Experimental/ChaosClothAssetEditor/Source/ChaosClothAssetEditor/Private/ChaosClothAsset/ClothEditorMode.cpp
--- a/Engine/Plugins/Experimental/ChaosClothAssetEditor/Source/ChaosClothAssetEditor/Private/ChaosClothAsset/ClothEditorMode.cpp
+++ b/Engine/Plugins/Experimental/ChaosClothAssetEditor/Source/ChaosClothAssetEditor/Private/ChaosClothAsset/ClothEditorMode.cpp
@@ -48,6 +48,7 @@
 #include "Widgets/Notifications/SNotificationList.h"
 #include "BoneWeights.h"
 #include "ToolSetupUtil.h"
+#include "Dataflow/DataflowComponent.h"
 
 #include UE_INLINE_GENERATED_CPP_BY_NAME(ClothEditorMode)
 
@@ -597,6 +598,10 @@ void UChaosClothAssetEditorMode::InitializeTargets(const TArray<TObjectPtr<UObje
 	}
 
 	ReinitializeDynamicMeshComponents();
+
+	DataflowComponent = NewObject<UDataflowComponent>();
+	DataflowComponent->RegisterComponentWithWorld(PreviewWorld);
+
 	bShouldFocusRestSpaceView = true;
 }
 
@@ -629,6 +634,11 @@ bool UChaosClothAssetEditorMode::IsSimulationSuspended() const
 	return bIsSimulationSuspended;
 }
 
+UDataflowComponent* UChaosClothAssetEditorMode::GetDataflowComponent() const
+{
+	return DataflowComponent;
+}
+
 void UChaosClothAssetEditorMode::ModeTick(float DeltaTime)
 {
 	Super::ModeTick(DeltaTime);
diff --git a/Engine/Plugins/Experimental/ChaosClothAssetEditor/Source/ChaosClothAssetEditor/Private/ChaosClothAsset/ClothEditorMode.h b/Engine/Plugins/Experimental/ChaosClothAssetEditor/Source/ChaosClothAssetEditor/Private/ChaosClothAsset/ClothEditorMode.h
--- a/Engine/Plugins/Experimental/ChaosClothAssetEditor/Source/ChaosClothAssetEditor/Private/ChaosClothAsset/ClothEditorMode.h
+++ b/Engine/Plugins/Experimental/ChaosClothAssetEditor/Source/ChaosClothAssetEditor/Private/ChaosClothAsset/ClothEditorMode.h
@@ -26,6 +26,7 @@ class UChaosClothComponent;
 class FEditorViewportClient;
 class FChaosClothEditorRestSpaceViewportClient;
 class FViewport;
+class UDataflowComponent;
 
 /**
  * The cloth editor mode is the mode used in the cloth asset editor. It holds most of the inter-tool state.
@@ -60,6 +61,8 @@ public:
 	void ResumeSimulation();
 	bool IsSimulationSuspended() const;
 
+	UDataflowComponent* GetDataflowComponent() const;
+
 private:
 
 	friend class FChaosClothAssetEditorToolkit;
@@ -160,5 +163,8 @@ private:
 	bool bShouldResetSimulation = false;
 	bool bHardReset = false;
 	bool bShouldClearTeleportFlag = false;
+
+	UPROPERTY()
+	TObjectPtr<UDataflowComponent> DataflowComponent = nullptr;
 };
 
diff --git a/Engine/Plugins/Experimental/ChaosClothAssetEditor/Source/ChaosClothAssetEditor/Private/ChaosClothAsset/ClothEditorToolkit.cpp b/Engine/Plugins/Experimental/ChaosClothAssetEditor/Source/ChaosClothAssetEditor/Private/ChaosClothAsset/ClothEditorToolkit.cpp
--- a/Engine/Plugins/Experimental/ChaosClothAssetEditor/Source/ChaosClothAssetEditor/Private/ChaosClothAsset/ClothEditorToolkit.cpp
+++ b/Engine/Plugins/Experimental/ChaosClothAssetEditor/Source/ChaosClothAssetEditor/Private/ChaosClothAsset/ClothEditorToolkit.cpp
@@ -26,9 +26,11 @@
 #include "Dataflow/DataflowGraphEditor.h"
 #include "Dataflow/DataflowObject.h"
 #include "Dataflow/DataflowSchema.h"
+#include "Dataflow/DataFlowEdNode.h"
 #include "IStructureDetailsView.h"
 #include "Engine/Canvas.h"
 #include "PropertyEditorModule.h"
+#include "Algo/RemoveIf.h"
 
 #define LOCTEXT_NAMESPACE "ChaosClothAssetEditorToolkit"
 
@@ -190,6 +192,15 @@ FChaosClothAssetEditorToolkit::~FChaosClothAssetEditorToolkit()
 	EditorModeManager->DestroyMode(UChaosClothAssetEditorMode::EM_ChaosClothAssetEditorModeId);
 }
 
+TSharedPtr<Dataflow::FEngineContext> FChaosClothAssetEditorToolkit::GetDataflowContext() const
+{
+	return DataflowContext;
+}
+
+const UDataflow* FChaosClothAssetEditorToolkit::GetDataflow() const
+{
+	return Dataflow;
+}
 
 //~ Begin FTickableEditorObject overrides
 
@@ -262,6 +273,7 @@ void FChaosClothAssetEditorToolkit::CreateWidgets()
 	{
 		Dataflow = ClothEditorToolkitHelpers::GetDataflowFrom(ClothAsset);
 
+		// TODO: Figure out how to create the GraphEditor widgets when the ClothAsset doesn't have a Dataflow property set
 		if (Dataflow)
 		{
 			DataflowTerminalPath = ClothEditorToolkitHelpers::GetDataflowTerminalFrom(ClothAsset);
@@ -399,6 +411,7 @@ void FChaosClothAssetEditorToolkit::PostInitAssetEditor()
 	// Set up 3D viewport
 	ClothPreviewViewportClient->SetClothComponent(ClothMode->ClothComponent);
 	ClothPreviewViewportClient->SetClothEdMode(ClothMode);
+	ClothPreviewViewportClient->SetClothEditorToolkit(StaticCastSharedRef<FChaosClothAssetEditorToolkit>(this->AsShared()));
 
 	SetCommonViewportClientOptions(ClothPreviewViewportClient.Get());
 	ClothPreviewViewportClient->SetInitialViewTransform(ELevelViewportType::LVT_Perspective, FVector(0, -100, 100), FRotator(0, 90, 0), DEFAULT_ORTHOZOOM);
@@ -591,24 +604,42 @@ TSharedRef<SDockTab> FChaosClothAssetEditorToolkit::SpawnTab_GraphCanvas(const F
 {
 	check(Args.GetTabId() == GraphCanvasTabId);
 
-	TSharedRef<SDockTab> SpawnedTab = SNew(SDockTab)
-		.Label(LOCTEXT("DataflowEditor_Dataflow_TabTitle", "Graph"))
-		[
-			GraphEditor.ToSharedRef()
-		];
+	if (GraphEditor)
+	{
+		TSharedRef<SDockTab> SpawnedTab = SNew(SDockTab)
+			.Label(LOCTEXT("DataflowEditor_Dataflow_TabTitle", "Graph"))
+			[
+				GraphEditor.ToSharedRef()
+			];
 
-	return SpawnedTab;
+		return SpawnedTab;
+	}
+	else
+	{
+		TSharedRef<SDockTab> SpawnedTab = SNew(SDockTab)
+			.Label(LOCTEXT("DataflowEditor_Dataflow_TabTitle", "Graph"));
+
+		return SpawnedTab;
+	}
 }
 
 TSharedRef<SDockTab> FChaosClothAssetEditorToolkit::SpawnTab_NodeDetails(const FSpawnTabArgs& Args)
 {
 	check(Args.GetTabId() == NodeDetailsTabId);
 
-	return SNew(SDockTab)
-		.Label(LOCTEXT("DataflowEditor_NodeDetails_TabTitle", "Node Details"))
-		[
-			NodeDetailsEditor->GetWidget()->AsShared()
-		];
+	if (NodeDetailsEditor)
+	{
+		return SNew(SDockTab)
+			.Label(LOCTEXT("DataflowEditor_NodeDetails_TabTitle", "Node Details"))
+			[
+				NodeDetailsEditor->GetWidget()->AsShared()
+			];
+	}
+	else
+	{
+		return SNew(SDockTab)
+			.Label(LOCTEXT("DataflowEditor_NodeDetails_TabTitle", "Node Details"));
+	}
 }
 
 void FChaosClothAssetEditorToolkit::InitDetailsViewPanel()
@@ -672,12 +703,15 @@ TSharedRef<SGraphEditor> FChaosClothAssetEditorToolkit::CreateGraphEditorWidget(
 	InEvents.OnVerifyTextCommit = FOnNodeVerifyTextCommit::CreateSP(this, &FChaosClothAssetEditorToolkit::OnNodeVerifyTitleCommit);
 	InEvents.OnTextCommitted = FOnNodeTextCommitted::CreateSP(this, &FChaosClothAssetEditorToolkit::OnNodeTitleCommitted);
 
-	return SNew(SDataflowGraphEditor, DataflowToEdit)
+	TSharedRef<SDataflowGraphEditor> NewGraphEditor = SNew(SDataflowGraphEditor, DataflowToEdit)
 		.GraphToEdit(DataflowToEdit)
 		.GraphEvents(InEvents)
 		.DetailsView(InNodeDetailsEditor)
 		.EvaluateGraph(Evaluate);
 
+	NewGraphEditor->OnSelectionChangedMulticast.AddSP(this, &FChaosClothAssetEditorToolkit::OnNodeSelectionChanged);
+
+	return NewGraphEditor;
 }
 
 TSharedPtr<IStructureDetailsView> FChaosClothAssetEditorToolkit::CreateNodeDetailsEditorWidget(UObject* ObjectToEdit)
@@ -729,6 +763,28 @@ void FChaosClothAssetEditorToolkit::OnNodeTitleCommitted(const FText& InNewText,
 	FDataflowEditorCommands::OnNodeTitleCommitted(InNewText, InCommitType, GraphNode);
 }
 
+void FChaosClothAssetEditorToolkit::OnNodeSelectionChanged(const TSet<UObject*>& NewSelection) const
+{
+	if (Dataflow)
+	{
+		// remove nodes if they don't have the render box checked
+		int32 NewLen = Algo::RemoveIf(Dataflow->RenderTargets, [](const UDataflowEdNode* Node)
+		{
+			return !Node->bRenderInAssetEditor;
+		});
+		Dataflow->RenderTargets.SetNum(NewLen);
+
+		for (UObject* Selected : NewSelection)
+		{
+			if (UDataflowEdNode* Node = Cast<UDataflowEdNode>(Selected))
+			{
+				Dataflow->RenderTargets.Add(Node);
+			}
+		}
+		Dataflow->LastModifiedRenderTarget = Dataflow::FTimestamp::Current();
+	}
+}
+
 //~ Ends DataflowEditorActions
 
 #undef LOCTEXT_NAMESPACE
diff --git a/Engine/Plugins/Experimental/ChaosClothAssetEditor/Source/ChaosClothAssetEditor/Private/ChaosClothAsset/ClothEditorToolkit.h b/Engine/Plugins/Experimental/ChaosClothAssetEditor/Source/ChaosClothAssetEditor/Private/ChaosClothAsset/ClothEditorToolkit.h
--- a/Engine/Plugins/Experimental/ChaosClothAssetEditor/Source/ChaosClothAssetEditor/Private/ChaosClothAsset/ClothEditorToolkit.h
+++ b/Engine/Plugins/Experimental/ChaosClothAssetEditor/Source/ChaosClothAssetEditor/Private/ChaosClothAsset/ClothEditorToolkit.h
@@ -43,6 +43,9 @@ public:
 	explicit FChaosClothAssetEditorToolkit(UAssetEditor* InOwningAssetEditor);
 	virtual ~FChaosClothAssetEditorToolkit();
 
+	TSharedPtr<Dataflow::FEngineContext> GetDataflowContext() const;
+	const UDataflow* GetDataflow() const;
+
 private:
 
 	static const FName ClothPreviewTabID;
@@ -70,7 +73,8 @@ private:
 	virtual void PostInitAssetEditor() override;
 
 	// IAssetEditorInstance
-	virtual bool IsPrimaryEditor() const override { return true; };
+	// TODO: If this returns true then the editor cannot re-open after it's closed. Figure out why.
+	virtual bool IsPrimaryEditor() const override { return false; };
 
 	// IToolkit
 	virtual FText GetToolkitName() const override;
@@ -125,6 +129,7 @@ private:
 	void OnPropertyValueChanged(const FPropertyChangedEvent& PropertyChangedEvent);
 	bool OnNodeVerifyTitleCommit(const FText& NewText, UEdGraphNode* GraphNode, FText& OutErrorMessage) const;
 	void OnNodeTitleCommitted(const FText& InNewText, ETextCommit::Type InCommitType, UEdGraphNode* GraphNode) const;
+	void OnNodeSelectionChanged(const TSet<UObject*>& NewSelection) const;
 
 	//~ End Dataflow support
 
diff --git a/Engine/Source/Runtime/Experimental/Chaos/Public/GeometryCollection/GeometryCollectionConvexUtility.h b/Engine/Source/Runtime/Experimental/Chaos/Public/GeometryCollection/GeometryCollectionConvexUtility.h
--- a/Engine/Source/Runtime/Experimental/Chaos/Public/GeometryCollection/GeometryCollectionConvexUtility.h
+++ b/Engine/Source/Runtime/Experimental/Chaos/Public/GeometryCollection/GeometryCollectionConvexUtility.h
@@ -31,7 +31,7 @@ namespace UE::GeometryCollectionConvexUtility
 	// Note these hulls are typically computed in a shared coordinate space, in contrast to the final hulls on the geometry collection which are in the local space of each bone
 	struct FConvexHulls
 	{
-		TArray<TUniquePtr<Chaos::FConvex>> Hulls;
+		TArray<TUniquePtr<::Chaos::FConvex>> Hulls;
 
 		// Mapping from geometry collection bones indices to indices in the Hulls array. A Set is used to support multiple hulls per bone.
 		TArray<TSet<int32>> TransformToHullsIndices;
