Cloth Editor: disable/hide the Remesh tool until we can properly apply the results back to the cloth asset

diff --git a/Engine/Plugins/Experimental/ChaosClothAssetEditor/Source/ChaosClothAssetEditor/Private/ChaosClothAsset/ClothEditorMode.cpp b/Engine/Plugins/Experimental/ChaosClothAssetEditor/Source/ChaosClothAssetEditor/Private/ChaosClothAsset/ClothEditorMode.cpp
--- a/Engine/Plugins/Experimental/ChaosClothAssetEditor/Source/ChaosClothAssetEditor/Private/ChaosClothAsset/ClothEditorMode.cpp
+++ b/Engine/Plugins/Experimental/ChaosClothAssetEditor/Source/ChaosClothAssetEditor/Private/ChaosClothAsset/ClothEditorMode.cpp
@@ -172,7 +172,9 @@ void UChaosClothAssetEditorMode::RegisterTools()
 	// be easily retrieved from the active tool name in ChaosClothAssetEditorModeToolkit::OnToolStarted. Otherwise
 	// we would need to keep some other mapping from tool identifier to tool icon.
 
-	RegisterTool(CommandInfos.BeginRemeshTool, FChaosClothAssetEditorCommands::BeginRemeshToolIdentifier, NewObject<URemeshMeshToolBuilder>());
+	// TODO: Re-add the remesh tool when we have a way to remesh both 2d and 3d rest space meshes at the same time
+	//RegisterTool(CommandInfos.BeginRemeshTool, FChaosClothAssetEditorCommands::BeginRemeshToolIdentifier, NewObject<URemeshMeshToolBuilder>());
+
 	RegisterTool(CommandInfos.BeginAttributeEditorTool, FChaosClothAssetEditorCommands::BeginAttributeEditorToolIdentifier, NewObject<UAttributeEditorToolBuilder>());
 	RegisterTool(CommandInfos.BeginWeightMapPaintTool, FChaosClothAssetEditorCommands::BeginWeightMapPaintToolIdentifier, NewObject<UClothEditorWeightMapPaintToolBuilder>());
 	RegisterTool(CommandInfos.BeginClothTrainingTool, FChaosClothAssetEditorCommands::BeginClothTrainingToolIdentifier, NewObject<UClothTrainingToolBuilder>());
diff --git a/Engine/Plugins/Experimental/ChaosClothAssetEditor/Source/ChaosClothAssetEditor/Private/ChaosClothAsset/ClothEditorModeToolkit.cpp b/Engine/Plugins/Experimental/ChaosClothAssetEditor/Source/ChaosClothAssetEditor/Private/ChaosClothAsset/ClothEditorModeToolkit.cpp
--- a/Engine/Plugins/Experimental/ChaosClothAssetEditor/Source/ChaosClothAssetEditor/Private/ChaosClothAsset/ClothEditorModeToolkit.cpp
+++ b/Engine/Plugins/Experimental/ChaosClothAssetEditor/Source/ChaosClothAssetEditor/Private/ChaosClothAsset/ClothEditorModeToolkit.cpp
@@ -24,7 +24,9 @@ void FChaosClothAssetEditorModeToolkit::BuildToolPalette(FName PaletteIndex, cla
 	// TODO: make separate palettes to split up the tools tab
 	if (PaletteIndex == FBaseCharacterFXEditorModeToolkit::ToolsTabName)
 	{
-		ToolbarBuilder.AddToolBarButton(Commands.BeginRemeshTool);
+		// TODO: Re-add the remesh tool when we have a way to remesh both 2d and 3d rest space meshes at the same time
+		//ToolbarBuilder.AddToolBarButton(Commands.BeginRemeshTool);
+
 		ToolbarBuilder.AddToolBarButton(Commands.BeginAttributeEditorTool);
 		ToolbarBuilder.AddToolBarButton(Commands.BeginWeightMapPaintTool);
 
