ChaosCloth: If a Cloth asset doesn't have a Dataflow graph asset attached to it when it's opened, offer to create a new one for the user

diff --git a/Engine/Plugins/Experimental/ChaosClothAssetEditor/Source/ChaosClothAssetEditor/Private/ChaosClothAsset/AssetTypeActions_ClothAsset.cpp b/Engine/Plugins/Experimental/ChaosClothAssetEditor/Source/ChaosClothAssetEditor/Private/ChaosClothAsset/AssetTypeActions_ClothAsset.cpp
--- a/Engine/Plugins/Experimental/ChaosClothAssetEditor/Source/ChaosClothAssetEditor/Private/ChaosClothAsset/AssetTypeActions_ClothAsset.cpp
+++ b/Engine/Plugins/Experimental/ChaosClothAssetEditor/Source/ChaosClothAssetEditor/Private/ChaosClothAsset/AssetTypeActions_ClothAsset.cpp
@@ -5,9 +5,74 @@
 #include "ToolMenus.h"
 #include "ChaosClothAsset/ClothEditor.h"
 #include "ChaosClothAsset/ClothAsset.h"
+#include "Dataflow/DataflowObject.h"
+#include "Misc/FileHelper.h"
+#include "Misc/MessageDialog.h"
+#include "ContentBrowserModule.h"
+#include "IContentBrowserSingleton.h"
+#include "AssetRegistry/AssetRegistryModule.h"
+#include "FileHelpers.h"
+
+#define LOCTEXT_NAMESPACE "FAssetTypeActions_ClothAsset"
 
 namespace UE::Chaos::ClothAsset
 {
+	namespace ClothAssetActionsHelpers
+	{
+		// Create a new UDataflow if one doesn't already exist for the Cloth Asset
+		UObject* CreateNewDataflowAsset(const UChaosClothAsset* ClothAsset)
+		{
+			if (FMessageDialog::Open(EAppMsgType::YesNo, LOCTEXT("MissingDataflow", "This Cloth asset currently has no Dataflow graph. Would you like to create a new one?")) == EAppReturnType::Yes)
+			{
+				const UClass* const DataflowClass = UDataflow::StaticClass();
+
+				FSaveAssetDialogConfig NewDataflowAssetDialogConfig;
+				{
+					const FString PackageName = ClothAsset->GetOutermost()->GetName();
+					NewDataflowAssetDialogConfig.DefaultPath = FPackageName::GetLongPackagePath(PackageName);
+					NewDataflowAssetDialogConfig.AssetClassNames.Add(DataflowClass->GetClassPathName());
+					NewDataflowAssetDialogConfig.ExistingAssetPolicy = ESaveAssetDialogExistingAssetPolicy::AllowButWarn;
+					NewDataflowAssetDialogConfig.DialogTitleOverride = LOCTEXT("NewDataflowAssetDialogTitle", "Save Dataflow Asset As");
+				}
+
+				FContentBrowserModule& ContentBrowserModule = FModuleManager::LoadModuleChecked<FContentBrowserModule>("ContentBrowser");
+
+				FString NewPackageName;
+				FText OutError;
+				for (bool bFilenameValid = false; !bFilenameValid; bFilenameValid = FFileHelper::IsFilenameValidForSaving(NewPackageName, OutError))
+				{
+					const FString AssetSavePath = ContentBrowserModule.Get().CreateModalSaveAssetDialog(NewDataflowAssetDialogConfig);
+					if (AssetSavePath.IsEmpty())
+					{
+						return nullptr;
+					}
+					NewPackageName = FPackageName::ObjectPathToPackageName(AssetSavePath);
+				}
+
+				const FName NewAssetName(FPackageName::GetLongPackageAssetName(NewPackageName));
+				UPackage* const NewPackage = CreatePackage(*NewPackageName);
+				UObject* const NewAsset = NewObject<UObject>(NewPackage, DataflowClass, NewAssetName, RF_Public | RF_Standalone | RF_Transactional);
+
+				NewAsset->MarkPackageDirty();
+
+				// Notify the asset registry
+				FAssetRegistryModule::AssetCreated(NewAsset);
+
+				// Save the package
+				TArray<UPackage*> PackagesToSave;
+				PackagesToSave.Add(NewAsset->GetOutermost());
+				constexpr bool bCheckDirty = false;
+				constexpr bool bPromptToSave = false;
+				FEditorFileUtils::PromptForCheckoutAndSave(PackagesToSave, bCheckDirty, bPromptToSave);
+
+				return NewAsset;
+			}
+
+			return nullptr;
+		}
+	}
+
+
 	FText FAssetTypeActions_ClothAsset::GetName() const
 	{
 		return NSLOCTEXT("AssetTypeActions", "AssetTypeActions_ClothAsset", "Cloth Asset");
@@ -37,14 +102,29 @@ namespace UE::Chaos::ClothAsset
 			if (Cast<UChaosClothAsset>(Object))
 			{
 				ClothObjects.Add(Object);
+				break;
 			}
 		}
 
-		// TODO:
+		// For now the cloth editor only works on one asset at a time
+		ensure(ClothObjects.Num() == 0 || ClothObjects.Num() == 1);
+
 		if (ClothObjects.Num() > 0)
 		{
 			UAssetEditorSubsystem* const AssetEditorSubsystem = GEditor->GetEditorSubsystem<UAssetEditorSubsystem>();
 			UChaosClothAssetEditor* const AssetEditor = NewObject<UChaosClothAssetEditor>(AssetEditorSubsystem, NAME_None, RF_Transient);
+
+			// Validate the asset
+			UChaosClothAsset* const ClothAsset = CastChecked<UChaosClothAsset>(ClothObjects[0]);
+
+			if (!ClothAsset->DataflowAsset)
+			{
+				if (UObject* const NewAsset = ClothAssetActionsHelpers::CreateNewDataflowAsset(ClothAsset))
+				{
+					ClothAsset->DataflowAsset = CastChecked<UDataflow>(NewAsset);
+				}
+			}
+
 			AssetEditor->Initialize(ClothObjects);
 		}
 	}
@@ -60,3 +140,6 @@ namespace UE::Chaos::ClothAsset
 		return NewObject<USceneThumbnailInfo>(Asset, NAME_None, RF_Transactional);
 	}
 } // End namespace UE::Chaos::ClothAsset
+
+
+#undef LOCTEXT_NAMESPACE
