Cloth Editor: hook up simulation suspend control to the actual ClothComponent suspend functions

diff --git a/Engine/Plugins/Experimental/ChaosClothAssetEditor/Source/ChaosClothAssetEditor/Private/ChaosClothAsset/ClothEditorMode.cpp b/Engine/Plugins/Experimental/ChaosClothAssetEditor/Source/ChaosClothAssetEditor/Private/ChaosClothAsset/ClothEditorMode.cpp
--- a/Engine/Plugins/Experimental/ChaosClothAssetEditor/Source/ChaosClothAssetEditor/Private/ChaosClothAsset/ClothEditorMode.cpp
+++ b/Engine/Plugins/Experimental/ChaosClothAssetEditor/Source/ChaosClothAssetEditor/Private/ChaosClothAsset/ClothEditorMode.cpp
@@ -623,17 +623,28 @@ void UChaosClothAssetEditorMode::HardResetSimulation()
 
 void UChaosClothAssetEditorMode::SuspendSimulation()
 {
-	bIsSimulationSuspended = true;
+	if (ClothComponent)
+	{
+		ClothComponent->SuspendSimulation();
+	}
 }
 
 void UChaosClothAssetEditorMode::ResumeSimulation()
 {
-	bIsSimulationSuspended = false;
+	if (ClothComponent)
+	{
+		ClothComponent->ResumeSimulation();
+	}
 }
 
 bool UChaosClothAssetEditorMode::IsSimulationSuspended() const
 {
-	return bIsSimulationSuspended;
+	if (ClothComponent)
+	{
+		return ClothComponent->IsSimulationSuspended();
+	}
+
+	return false;
 }
 
 UDataflowComponent* UChaosClothAssetEditorMode::GetDataflowComponent() const
@@ -669,8 +680,6 @@ void UChaosClothAssetEditorMode::ModeTick(float DeltaTime)
 	}
 
 
-	bool bShouldTickPreviewWorld = !bIsSimulationSuspended;
-
 	if (bShouldClearTeleportFlag)
 	{
 		ClothComponent->ResetTeleportMode();
@@ -690,10 +699,9 @@ void UChaosClothAssetEditorMode::ModeTick(float DeltaTime)
 
 		bShouldResetSimulation = false;
 		bShouldClearTeleportFlag = true;		// clear the flag next tick
-		bShouldTickPreviewWorld = true;
 	}
 
-	if (PreviewWorld && bShouldTickPreviewWorld)
+	if (PreviewWorld)
 	{
 		PreviewWorld->Tick(ELevelTick::LEVELTICK_All, DeltaTime);
 	}
diff --git a/Engine/Plugins/Experimental/ChaosClothAssetEditor/Source/ChaosClothAssetEditor/Private/ChaosClothAsset/ClothEditorMode.h b/Engine/Plugins/Experimental/ChaosClothAssetEditor/Source/ChaosClothAssetEditor/Private/ChaosClothAsset/ClothEditorMode.h
--- a/Engine/Plugins/Experimental/ChaosClothAssetEditor/Source/ChaosClothAssetEditor/Private/ChaosClothAsset/ClothEditorMode.h
+++ b/Engine/Plugins/Experimental/ChaosClothAssetEditor/Source/ChaosClothAssetEditor/Private/ChaosClothAsset/ClothEditorMode.h
@@ -159,7 +159,6 @@ private:
 	void UpdateSimulationMeshes();
 
 	// Simulation controls
-	bool bIsSimulationSuspended = false;
 	bool bShouldResetSimulation = false;
 	bool bHardReset = false;
 	bool bShouldClearTeleportFlag = false;
