Cloth Editor spreadsheet view: add support for bone indices and weights

diff --git a/Engine/Plugins/Experimental/ChaosClothAssetEditor/Source/ChaosClothAssetEditor/Private/ChaosClothAsset/SClothCollectionOutliner.cpp b/Engine/Plugins/Experimental/ChaosClothAssetEditor/Source/ChaosClothAssetEditor/Private/ChaosClothAsset/SClothCollectionOutliner.cpp
--- a/Engine/Plugins/Experimental/ChaosClothAssetEditor/Source/ChaosClothAssetEditor/Private/ChaosClothAsset/SClothCollectionOutliner.cpp
+++ b/Engine/Plugins/Experimental/ChaosClothAssetEditor/Source/ChaosClothAssetEditor/Private/ChaosClothAsset/SClothCollectionOutliner.cpp
@@ -198,6 +198,12 @@ namespace ClothCollectionOutlinerHelpers
 		case FManagedArrayCollection::EArrayType::FInt32Type:
 			ValueAsString = AttributeValueToString<int32>(ClothCollection, AttributeName, GroupName, AttributeArrayIndex);
 			break;
+		case FManagedArrayCollection::EArrayType::FInt32ArrayType:
+			ValueAsString = AttributeValueToString<TArray<int32>>(ClothCollection, AttributeName, GroupName, AttributeArrayIndex);
+			break;
+		case FManagedArrayCollection::EArrayType::FFloatArrayType:
+			ValueAsString = AttributeValueToString<TArray<float>>(ClothCollection, AttributeName, GroupName, AttributeArrayIndex);
+			break;
 		default:
 			ensure(false);
 			ValueAsString = "(Unknown Data Type)";
