Cloth Editor:
- Create a new ClothEditorPreviewScene to track the contents of the scene visualized in the Preview Viewport
- Move the ClothComponent from the Edtior Mode to the PreviewScene
- Add an optional SkeletalMeshComponent to the PreviewScene, to eventally be used for skinning/animation/collision work

diff --git a/Engine/Plugins/Experimental/ChaosClothAssetEditor/Source/ChaosClothAssetEditor/Private/ChaosClothAsset/ClothEditorMode.cpp b/Engine/Plugins/Experimental/ChaosClothAssetEditor/Source/ChaosClothAssetEditor/Private/ChaosClothAsset/ClothEditorMode.cpp
--- a/Engine/Plugins/Experimental/ChaosClothAssetEditor/Source/ChaosClothAssetEditor/Private/ChaosClothAsset/ClothEditorMode.cpp
+++ b/Engine/Plugins/Experimental/ChaosClothAssetEditor/Source/ChaosClothAssetEditor/Private/ChaosClothAsset/ClothEditorMode.cpp
@@ -9,6 +9,7 @@
 #include "ChaosClothAsset/ClothAdapter.h"
 #include "ChaosClothAsset/ClothEditorRestSpaceViewportClient.h"
 #include "ChaosClothAsset/ClothPatternToDynamicMesh.h"
+#include "ChaosClothAsset/ClothEditorPreviewScene.h"
 #include "AssetEditorModeManager.h"
 #include "Drawing/MeshElementsVisualizer.h"
 #include "EditorViewportClient.h"
@@ -258,21 +259,21 @@ void UChaosClothAssetEditorMode::Exit()
 	WireframesToTick.Reset();
 	DynamicMeshComponents.Reset();
 	PropertyObjectsToTick.Empty();
-	PreviewWorld = nullptr;
+	PreviewScene = nullptr;
 
 	Super::Exit();
 }
 
-void UChaosClothAssetEditorMode::SetPreviewWorld(TObjectPtr<UWorld> InWorld)
+void UChaosClothAssetEditorMode::SetPreviewScene(FChaosClothPreviewScene* InPreviewScene)
 {
-	PreviewWorld = InWorld;
+	PreviewScene = InPreviewScene;
 }
 
 void UChaosClothAssetEditorMode::UpdateSimulationMeshes()
 {
 	GetToolManager()->BeginUndoTransaction(LOCTEXT("ChaosClothAssetEditorApplyChangesTransaction", "Cloth Editor Apply Changes"));
 
-	UChaosClothAsset* ChaosClothAsset = ClothComponent->GetClothAsset();
+	UChaosClothAsset* ChaosClothAsset = PreviewScene->ClothComponent->GetClothAsset();
 	ChaosClothAsset->Modify();
 
 	UE::Chaos::ClothAsset::FClothAdapter ClothAdapter(ChaosClothAsset->GetClothCollection());
@@ -347,7 +348,7 @@ void UChaosClothAssetEditorMode::UpdateSimulationMeshes()
 
 	// Reset cloth component
 	{
-		const FComponentReregisterContext Context(ClothComponent);
+		const FComponentReregisterContext Context(PreviewScene->ClothComponent);
 	}
 
 	GetToolManager()->EndUndoTransaction();
@@ -393,12 +394,12 @@ void UChaosClothAssetEditorMode::ReinitializeDynamicMeshComponents()
 	WireframesToTick.Empty();
 	DynamicMeshSourceInfos.Empty();
 
-	if (!ClothComponent)
+	if (!PreviewScene->ClothComponent)
 	{
 		return;
 	}
 
-	const UChaosClothAsset* ChaosClothAsset = ClothComponent->GetClothAsset();
+	const UChaosClothAsset* ChaosClothAsset = PreviewScene->ClothComponent->GetClothAsset();
 
 	if (!ChaosClothAsset)
 	{
@@ -571,7 +572,7 @@ void UChaosClothAssetEditorMode::RefocusRestSpaceViewportClient()
 void UChaosClothAssetEditorMode::InitializeTargets(const TArray<TObjectPtr<UObject>>& AssetsIn)
 {
 	// InitializeContexts needs to have been called first so that we have the 3d preview world ready.
-	check(PreviewWorld);
+	check(PreviewScene);
 
 	OriginalObjectsToEdit = AssetsIn;
 
@@ -587,13 +588,11 @@ void UChaosClothAssetEditorMode::InitializeTargets(const TArray<TObjectPtr<UObje
 	{
 		if (UChaosClothAsset* ChaosClothAsset = Cast<UChaosClothAsset>(Asset))
 		{
-			ClothComponent = NewObject<UChaosClothComponent>();
-			ClothComponent->SetClothAsset(ChaosClothAsset);
-			ClothComponent->RegisterComponentWithWorld(PreviewWorld);
+			PreviewScene->CreateClothComponent(ChaosClothAsset);
 
 			USelection* SelectedComponents = GetModeManager()->GetSelectedComponents();
 			SelectedComponents->Modify();
-			SelectedComponents->Select(ClothComponent);
+			SelectedComponents->Select(PreviewScene->ClothComponent);
 
 			break;
 		}
@@ -602,7 +601,7 @@ void UChaosClothAssetEditorMode::InitializeTargets(const TArray<TObjectPtr<UObje
 	ReinitializeDynamicMeshComponents();
 
 	DataflowComponent = NewObject<UDataflowComponent>();
-	DataflowComponent->RegisterComponentWithWorld(PreviewWorld);
+	DataflowComponent->RegisterComponentWithWorld(PreviewScene->GetWorld());
 
 	bShouldFocusRestSpaceView = true;
 }
@@ -623,25 +622,25 @@ void UChaosClothAssetEditorMode::HardResetSimulation()
 
 void UChaosClothAssetEditorMode::SuspendSimulation()
 {
-	if (ClothComponent)
+	if (PreviewScene && PreviewScene->ClothComponent)
 	{
-		ClothComponent->SuspendSimulation();
+		PreviewScene->ClothComponent->SuspendSimulation();
 	}
 }
 
 void UChaosClothAssetEditorMode::ResumeSimulation()
 {
-	if (ClothComponent)
+	if (PreviewScene && PreviewScene->ClothComponent)
 	{
-		ClothComponent->ResumeSimulation();
+		PreviewScene->ClothComponent->ResumeSimulation();
 	}
 }
 
 bool UChaosClothAssetEditorMode::IsSimulationSuspended() const
 {
-	if (ClothComponent)
+	if (PreviewScene && PreviewScene->ClothComponent)
 	{
-		return ClothComponent->IsSimulationSuspended();
+		return PreviewScene->ClothComponent->IsSimulationSuspended();
 	}
 
 	return false;
@@ -682,7 +681,7 @@ void UChaosClothAssetEditorMode::ModeTick(float DeltaTime)
 
 	if (bShouldClearTeleportFlag)
 	{
-		ClothComponent->ResetTeleportMode();
+		PreviewScene->ClothComponent->ResetTeleportMode();
 		bShouldClearTeleportFlag = false;
 	}
 
@@ -690,20 +689,20 @@ void UChaosClothAssetEditorMode::ModeTick(float DeltaTime)
 	{
 		if (bHardReset)
 		{
-			const FComponentReregisterContext Context(ClothComponent);
+			const FComponentReregisterContext Context(PreviewScene->ClothComponent);
 		}
 		else
 		{
-			ClothComponent->ForceNextUpdateTeleportAndReset();
+			PreviewScene->ClothComponent->ForceNextUpdateTeleportAndReset();
 		}
 
 		bShouldResetSimulation = false;
 		bShouldClearTeleportFlag = true;		// clear the flag next tick
 	}
 
-	if (PreviewWorld)
+	if (PreviewScene->GetWorld())
 	{
-		PreviewWorld->Tick(ELevelTick::LEVELTICK_All, DeltaTime);
+		PreviewScene->GetWorld()->Tick(ELevelTick::LEVELTICK_All, DeltaTime);
 	}
 }
 
@@ -782,9 +781,9 @@ FBox UChaosClothAssetEditorMode::SelectionBoundingBox() const
 
 FBox UChaosClothAssetEditorMode::PreviewBoundingBox() const
 {
-	if (ClothComponent)
+	if (PreviewScene->ClothComponent)
 	{
-		FBoxSphereBounds Bounds = ClothComponent->Bounds;
+		FBoxSphereBounds Bounds = PreviewScene->ClothComponent->Bounds;
 		return Bounds.GetBox();
 	}
 
diff --git a/Engine/Plugins/Experimental/ChaosClothAssetEditor/Source/ChaosClothAssetEditor/Private/ChaosClothAsset/ClothEditorMode.h b/Engine/Plugins/Experimental/ChaosClothAssetEditor/Source/ChaosClothAssetEditor/Private/ChaosClothAsset/ClothEditorMode.h
--- a/Engine/Plugins/Experimental/ChaosClothAssetEditor/Source/ChaosClothAssetEditor/Private/ChaosClothAsset/ClothEditorMode.h
+++ b/Engine/Plugins/Experimental/ChaosClothAssetEditor/Source/ChaosClothAssetEditor/Private/ChaosClothAsset/ClothEditorMode.h
@@ -27,6 +27,7 @@ class FEditorViewportClient;
 class FChaosClothEditorRestSpaceViewportClient;
 class FViewport;
 class UDataflowComponent;
+class FChaosClothPreviewScene;
 
 /**
  * The cloth editor mode is the mode used in the cloth asset editor. It holds most of the inter-tool state.
@@ -91,9 +92,7 @@ private:
 	// to turn them into the input objects that tools get (since these need preview meshes, etc).
 	static const FToolTargetTypeRequirements& GetToolTargetRequirements();
 
-	// Both SetPreviewWorld and InitializeTargets functions must be called for things to function properly. SetPreviewWorld should
-	// be done first so that the 3d preview world is ready for creating meshes in InitializeTargets.
-	void SetPreviewWorld(TObjectPtr<UWorld> InWorld);
+	void SetPreviewScene(FChaosClothPreviewScene* PreviewScene);
 
 	// Bounding box for rest space meshes
 	virtual FBox SceneBoundingBox() const override;
@@ -101,18 +100,14 @@ private:
 	void SetRestSpaceViewportClient(TWeakPtr<FChaosClothEditorRestSpaceViewportClient, ESPMode::ThreadSafe> ViewportClient);
 	void RefocusRestSpaceViewportClient();
 
-	// Preview simulation mesh
-	UPROPERTY()
-	TObjectPtr<UChaosClothComponent> ClothComponent;
 
 	// Rest space wireframes. They have to get ticked to be able to respond to setting changes. 
 	UPROPERTY()
 	TArray<TObjectPtr<UMeshElementsVisualizer>> WireframesToTick;
 
-	// Here largely for convenience to avoid having to pass it around functions.
-	UPROPERTY()
-	TObjectPtr<UWorld> PreviewWorld = nullptr;
-	
+	// Preview Scene, here largely for convenience to avoid having to pass it around functions. Owned by the ClothEditorToolkit.
+	FChaosClothPreviewScene* PreviewScene = nullptr;
+
 	// Mode-level property objects (visible or not) that get ticked.
 	UPROPERTY()
 	TArray<TObjectPtr<UInteractiveToolPropertySet>> PropertyObjectsToTick;
diff --git a/Engine/Plugins/Experimental/ChaosClothAssetEditor/Source/ChaosClothAssetEditor/Private/ChaosClothAsset/ClothEditorPreviewScene.cpp b/Engine/Plugins/Experimental/ChaosClothAssetEditor/Source/ChaosClothAssetEditor/Private/ChaosClothAsset/ClothEditorPreviewScene.cpp
new file mode 100644
--- /dev/null
+++ b/Engine/Plugins/Experimental/ChaosClothAssetEditor/Source/ChaosClothAssetEditor/Private/ChaosClothAsset/ClothEditorPreviewScene.cpp
@@ -0,0 +1,107 @@
+// Copyright Epic Games, Inc. All Rights Reserved.
+
+#include "ChaosClothAsset/ClothEditorPreviewScene.h"
+#include "ChaosClothAsset/ClothComponent.h"
+#include "Components/SkeletalMeshComponent.h"
+#include "Animation/SkeletalMeshActor.h"
+
+#define LOCTEXT_NAMESPACE "UChaosClothEditorPreviewScene"
+
+void UChaosClothPreviewSceneDescription::SetPreviewScene(FChaosClothPreviewScene* InPreviewScene)
+{
+	PreviewScene = InPreviewScene;
+}
+
+void UChaosClothPreviewSceneDescription::PostEditChangeProperty(FPropertyChangedEvent& PropertyChangedEvent)
+{
+	if (PreviewScene)
+	{
+		PreviewScene->SceneDescriptionPropertyChanged(PropertyChangedEvent);
+	}
+}
+
+
+
+FChaosClothPreviewScene::FChaosClothPreviewScene(FPreviewScene::ConstructionValues ConstructionValues) :
+	FAdvancedPreviewScene(ConstructionValues)
+{
+	PreviewSceneDescription = NewObject<UChaosClothPreviewSceneDescription>();
+	PreviewSceneDescription->SetPreviewScene(this);
+
+	if (PreviewSceneDescription->SkeletalMeshAsset)
+	{
+		InitializeSkeletalMeshActor();
+	}
+}
+
+FChaosClothPreviewScene::~FChaosClothPreviewScene()
+{
+	if (SkeletalMeshActor)
+	{
+		SkeletalMeshActor->UnregisterAllComponents();
+	}
+
+	if (ClothComponent)
+	{
+		ClothComponent->UnregisterComponent();
+	}
+}
+
+void FChaosClothPreviewScene::AddReferencedObjects(FReferenceCollector& Collector)
+{
+	FAdvancedPreviewScene::AddReferencedObjects(Collector);
+
+	Collector.AddReferencedObject(PreviewSceneDescription);
+	Collector.AddReferencedObject(ClothComponent);
+	Collector.AddReferencedObject(SkeletalMeshActor);
+}
+
+void FChaosClothPreviewScene::SceneDescriptionPropertyChanged(struct FPropertyChangedEvent& PropertyChangedEvent)
+{
+	if (PropertyChangedEvent.Property->GetFName() == GET_MEMBER_NAME_CHECKED(UChaosClothPreviewSceneDescription, SkeletalMeshAsset))
+	{
+		if (PreviewSceneDescription->SkeletalMeshAsset)
+		{
+			InitializeSkeletalMeshActor();
+		}
+	}
+
+	if (PropertyChangedEvent.Property->GetFName() == GET_MEMBER_NAME_CHECKED(UChaosClothPreviewSceneDescription, SkeletalMeshTransform))
+	{
+		if (SkeletalMeshActor)
+		{
+			SkeletalMeshActor->SetActorTransform(PreviewSceneDescription->SkeletalMeshTransform);
+		}
+	}
+}
+
+
+void FChaosClothPreviewScene::InitializeSkeletalMeshActor()
+{
+	if (SkeletalMeshActor)
+	{
+		SkeletalMeshActor->UnregisterAllComponents();
+	}
+
+	ensure(PreviewSceneDescription->SkeletalMeshAsset);
+
+	SkeletalMeshActor = GetWorld()->SpawnActor<ASkeletalMeshActor>(ASkeletalMeshActor::StaticClass(), FTransform::Identity);
+	if (SkeletalMeshActor)
+	{
+		SkeletalMeshActor->GetSkeletalMeshComponent()->SetSkeletalMeshAsset(PreviewSceneDescription->SkeletalMeshAsset);
+	}
+
+	SkeletalMeshActor->RegisterAllComponents();
+}
+
+
+void FChaosClothPreviewScene::CreateClothComponent(UChaosClothAsset* Asset)
+{
+	ClothComponent = NewObject<UChaosClothComponent>();
+	ClothComponent->SetClothAsset(Asset);
+	ClothComponent->RegisterComponentWithWorld(PreviewWorld);
+}
+
+
+#undef LOCTEXT_NAMESPACE
+
diff --git a/Engine/Plugins/Experimental/ChaosClothAssetEditor/Source/ChaosClothAssetEditor/Private/ChaosClothAsset/ClothEditorPreviewScene.h b/Engine/Plugins/Experimental/ChaosClothAssetEditor/Source/ChaosClothAssetEditor/Private/ChaosClothAsset/ClothEditorPreviewScene.h
new file mode 100644
--- /dev/null
+++ b/Engine/Plugins/Experimental/ChaosClothAssetEditor/Source/ChaosClothAssetEditor/Private/ChaosClothAsset/ClothEditorPreviewScene.h
@@ -0,0 +1,82 @@
+// Copyright Epic Games, Inc. All Rights Reserved.
+
+#pragma once
+
+#include "AdvancedPreviewScene.h"
+#include "Math/Transform.h"
+#include "ClothEditorPreviewScene.generated.h"
+
+class UChaosClothAsset;
+class USkeletalMesh;
+class UChaosClothComponent;
+class FChaosClothPreviewScene;
+class UPhysicsAsset;
+class UAnimSequence;
+class ASkeletalMeshActor;
+
+///
+/// The UChaosClothPreviewSceneDescription is a description of the Preview scene contents, intended to be editable in an FAdvancedPreviewSettingsWidget
+/// 
+UCLASS()
+class CHAOSCLOTHASSETEDITOR_API UChaosClothPreviewSceneDescription : public UObject
+{
+public:
+	GENERATED_BODY()
+
+	void SetPreviewScene(FChaosClothPreviewScene* PreviewScene);
+
+	// Skeletal Mesh source asset
+	UPROPERTY(EditAnywhere, Category="SkeletalMesh")
+	TObjectPtr<USkeletalMesh> SkeletalMeshAsset;
+
+	UPROPERTY(EditAnywhere, Category = "SkeletalMesh")
+	FTransform SkeletalMeshTransform;
+
+	// TODO: Add anything else to the scene, e.g.:
+	//UPROPERTY(EditAnywhere, Category = "Animation")
+	//TObjectPtr<UAnimSequence> AnimationSequence;
+
+private:
+
+	// Listen for changes to the scene description members and notify the PreviewScene
+	virtual void PostEditChangeProperty(struct FPropertyChangedEvent& PropertyChangedEvent) override;
+
+	FChaosClothPreviewScene* PreviewScene;
+};
+
+
+///
+/// FChaosClothPreviewScene is the actual Preview scene, with contents specified by the SceneDescription
+/// 
+class CHAOSCLOTHASSETEDITOR_API FChaosClothPreviewScene : public FAdvancedPreviewScene
+{
+public:
+
+	FChaosClothPreviewScene(FPreviewScene::ConstructionValues ConstructionValues);
+	virtual ~FChaosClothPreviewScene();
+
+	virtual void AddReferencedObjects(FReferenceCollector& Collector) override;
+
+	const UChaosClothPreviewSceneDescription* GetPreviewSceneDescription() const { return PreviewSceneDescription; }
+	UChaosClothPreviewSceneDescription* GetPreviewSceneDescription() { return PreviewSceneDescription; }
+
+	// Update Scene in response to the SceneDescription changing
+	void SceneDescriptionPropertyChanged(struct FPropertyChangedEvent& PropertyChangedEvent);
+
+	void CreateClothComponent(UChaosClothAsset* Asset);
+
+
+	// Preview simulation mesh
+	TObjectPtr<UChaosClothComponent> ClothComponent;
+
+	// Skeletal Mesh
+	TObjectPtr<ASkeletalMeshActor> SkeletalMeshActor;
+
+private:
+
+	TObjectPtr<UChaosClothPreviewSceneDescription> PreviewSceneDescription;
+
+	void InitializeSkeletalMeshActor();
+};
+
+
diff --git a/Engine/Plugins/Experimental/ChaosClothAssetEditor/Source/ChaosClothAssetEditor/Private/ChaosClothAsset/ClothEditorToolkit.cpp b/Engine/Plugins/Experimental/ChaosClothAssetEditor/Source/ChaosClothAssetEditor/Private/ChaosClothAsset/ClothEditorToolkit.cpp
--- a/Engine/Plugins/Experimental/ChaosClothAssetEditor/Source/ChaosClothAssetEditor/Private/ChaosClothAsset/ClothEditorToolkit.cpp
+++ b/Engine/Plugins/Experimental/ChaosClothAssetEditor/Source/ChaosClothAssetEditor/Private/ChaosClothAsset/ClothEditorToolkit.cpp
@@ -31,11 +31,13 @@
 #include "Engine/Canvas.h"
 #include "PropertyEditorModule.h"
 #include "Algo/RemoveIf.h"
+#include "AdvancedPreviewSceneModule.h"
 
 #define LOCTEXT_NAMESPACE "ChaosClothAssetEditorToolkit"
 
 const FName FChaosClothAssetEditorToolkit::ClothPreviewTabID(TEXT("ChaosClothAssetEditor_ClothPreviewTab"));
 const FName FChaosClothAssetEditorToolkit::OutlinerTabID(TEXT("ChaosClothAssetEditor_OutlinerTab"));
+const FName FChaosClothAssetEditorToolkit::PreviewSceneDetailsTabID(TEXT("ChaosClothAssetEditor_PreviewSceneDetailsTab"));
 const FName FChaosClothAssetEditorToolkit::GraphCanvasTabId(TEXT("ChaosClothAssetEditor_GraphCanvas"));
 const FName FChaosClothAssetEditorToolkit::NodeDetailsTabId(TEXT("ChaosClothAssetEditor_NodeDetails"));
 
@@ -146,6 +148,7 @@ FChaosClothAssetEditorToolkit::FChaosClothAssetEditorToolkit(UAssetEditor* InOwn
 					->SetSizeCoefficient(0.3f)
 					->AddTab(DetailsTabID, ETabState::OpenedTab)
 					->AddTab(OutlinerTabID, ETabState::OpenedTab)
+					->AddTab(PreviewSceneDetailsTabID, ETabState::OpenedTab)
 					->SetExtensionId("DetailsArea")
 					->SetHideTabWell(true)
 				)
@@ -164,7 +167,7 @@ FChaosClothAssetEditorToolkit::FChaosClothAssetEditorToolkit(UAssetEditor* InOwn
 	PreviewSceneArgs.bShouldSimulatePhysics = 1;
 	PreviewSceneArgs.bCreatePhysicsScene = 1;
 	
-	ClothPreviewScene = MakeUnique<FAdvancedPreviewScene>(PreviewSceneArgs);
+	ClothPreviewScene = MakeShared<FChaosClothPreviewScene>(PreviewSceneArgs);
 	ClothPreviewScene->SetFloorVisibility(false, true);
 	ClothPreviewEditorModeManager = MakeShared<FAssetEditorModeManager>();
 	ClothPreviewEditorModeManager->SetPreviewScene(ClothPreviewScene.Get());
@@ -242,8 +245,8 @@ void FChaosClothAssetEditorToolkit::InitializeEdMode(UBaseCharacterFXEditorMode*
 	UChaosClothAssetEditorMode* ClothMode = Cast<UChaosClothAssetEditorMode>(EdMode);
 	check(ClothMode);
 
-	// The mode will need to be able to get to the live preview world
-	ClothMode->SetPreviewWorld(ClothPreviewEditorModeManager->GetWorld());
+	check(ClothPreviewScene.IsValid());
+	ClothMode->SetPreviewScene(ClothPreviewScene.Get());
 
 	TArray<TObjectPtr<UObject>> ObjectsToEdit;
 	OwningAssetEditor->GetObjectsToEdit(ObjectsToEdit);
@@ -409,7 +412,7 @@ void FChaosClothAssetEditorToolkit::PostInitAssetEditor()
 	ViewportClient->ReceivedFocus(ViewportClient->Viewport);
 
 	// Set up 3D viewport
-	ClothPreviewViewportClient->SetClothComponent(ClothMode->ClothComponent);
+	ClothPreviewViewportClient->SetClothComponent(ClothPreviewScene->ClothComponent);
 	ClothPreviewViewportClient->SetClothEdMode(ClothMode);
 	ClothPreviewViewportClient->SetClothEditorToolkit(StaticCastSharedRef<FChaosClothAssetEditorToolkit>(this->AsShared()));
 
@@ -502,6 +505,11 @@ void FChaosClothAssetEditorToolkit::RegisterTabSpawners(const TSharedRef<FTabMan
 		.SetGroup(AssetEditorTabsCategory.ToSharedRef())
 		.SetIcon(FSlateIcon(FAppStyle::GetAppStyleSetName(), "LevelEditor.Tabs.Outliner"));
 
+	InTabManager->RegisterTabSpawner(PreviewSceneDetailsTabID, FOnSpawnTab::CreateSP(this, &FChaosClothAssetEditorToolkit::SpawnTab_PreviewSceneDetails))
+		.SetDisplayName(LOCTEXT("PreviewSceneDetails", "Preview Scene Details"))
+		.SetGroup(AssetEditorTabsCategory.ToSharedRef())
+		.SetIcon(FSlateIcon(FAppStyle::GetAppStyleSetName(), "LevelEditor.Tabs.Details"));
+
 	InTabManager->RegisterTabSpawner(GraphCanvasTabId, FOnSpawnTab::CreateSP(this, &FChaosClothAssetEditorToolkit::SpawnTab_GraphCanvas))
 		.SetDisplayName(LOCTEXT("DataflowTab", "Graph"))
 		.SetGroup(AssetEditorTabsCategory.ToSharedRef())
@@ -600,6 +608,15 @@ TSharedRef<SDockTab> FChaosClothAssetEditorToolkit::SpawnTab_Outliner(const FSpa
 	return DockableTab;
 }
 
+
+TSharedRef<SDockTab> FChaosClothAssetEditorToolkit::SpawnTab_PreviewSceneDetails(const FSpawnTabArgs& Args)
+{
+	SAssignNew(PreviewSceneDockTab, SDockTab)
+		.Label(LOCTEXT("PreviewSceneDetailsTitle", "Preview Scene Details"));
+
+	return PreviewSceneDockTab.ToSharedRef();
+}
+
 TSharedRef<SDockTab> FChaosClothAssetEditorToolkit::SpawnTab_GraphCanvas(const FSpawnTabArgs& Args)
 {
 	check(Args.GetTabId() == GraphCanvasTabId);
@@ -654,11 +671,23 @@ void FChaosClothAssetEditorToolkit::InitDetailsViewPanel()
 		SetEditingObject(ObjectToEditInDetailsView);
 	}
 
-	const UChaosClothAssetEditorMode* const ClothMode = CastChecked<UChaosClothAssetEditorMode>(EditorModeManager->GetActiveScriptableMode(UChaosClothAssetEditorMode::EM_ChaosClothAssetEditorModeId));
+	TArray<FAdvancedPreviewSceneModule::FDetailDelegates> Delegates;	
 
-	TSharedPtr<UE::Chaos::ClothAsset::FClothCollection> ClothCollection;
+	ensure(ClothPreviewScene.IsValid());
+	FAdvancedPreviewSceneModule& AdvancedPreviewSceneModule = FModuleManager::LoadModuleChecked<FAdvancedPreviewSceneModule>("AdvancedPreviewScene");
+	AdvancedPreviewSettingsWidget = AdvancedPreviewSceneModule.CreateAdvancedPreviewSceneSettingsWidget(ClothPreviewScene.ToSharedRef(), 
+		ClothPreviewScene->GetPreviewSceneDescription(),
+		TArray<FAdvancedPreviewSceneModule::FDetailCustomizationInfo>(), 
+		TArray<FAdvancedPreviewSceneModule::FPropertyTypeCustomizationInfo>(),
+		Delegates);
+
+	if (PreviewSceneDockTab.IsValid())
+	{
+		PreviewSceneDockTab->SetContent(AdvancedPreviewSettingsWidget.ToSharedRef());
+	}
 
-	if (const TObjectPtr<UChaosClothComponent> ClothComponent = ClothMode->ClothComponent)
+	TSharedPtr<UE::Chaos::ClothAsset::FClothCollection> ClothCollection;
+	if (const TObjectPtr<UChaosClothComponent> ClothComponent = ClothPreviewScene->ClothComponent)
 	{
 		if (UChaosClothAsset* const ClothAsset = ClothComponent->GetClothAsset())
 		{
diff --git a/Engine/Plugins/Experimental/ChaosClothAssetEditor/Source/ChaosClothAssetEditor/Private/ChaosClothAsset/ClothEditorToolkit.h b/Engine/Plugins/Experimental/ChaosClothAssetEditor/Source/ChaosClothAssetEditor/Private/ChaosClothAsset/ClothEditorToolkit.h
--- a/Engine/Plugins/Experimental/ChaosClothAssetEditor/Source/ChaosClothAssetEditor/Private/ChaosClothAsset/ClothEditorToolkit.h
+++ b/Engine/Plugins/Experimental/ChaosClothAssetEditor/Source/ChaosClothAssetEditor/Private/ChaosClothAsset/ClothEditorToolkit.h
@@ -4,8 +4,8 @@
 
 #include "BaseCharacterFXEditorToolkit.h"
 #include "Dataflow/DataflowObjectInterface.h"
+#include "ChaosClothAsset/ClothEditorPreviewScene.h"
 
-class FAdvancedPreviewScene;
 class FChaosClothAssetEditor3DViewportClient;
 template<typename T> class SComboBox;
 class SClothCollectionOutliner;
@@ -14,6 +14,7 @@ class UChaosClothAsset;
 class SGraphEditor;
 class IStructureDetailsView;
 class UEdGraphNode;
+class UChaosClothComponent;
 
 namespace Dataflow
 {
@@ -50,6 +51,7 @@ private:
 
 	static const FName ClothPreviewTabID;
 	static const FName OutlinerTabID;
+	static const FName PreviewSceneDetailsTabID;
 
 	// FTickableEditorObject
 	virtual void Tick(float DeltaTime) override;
@@ -89,13 +91,14 @@ private:
 
 	TSharedRef<SDockTab> SpawnTab_ClothPreview(const FSpawnTabArgs& Args);
 	TSharedRef<SDockTab> SpawnTab_Outliner(const FSpawnTabArgs& Args);
+	TSharedRef<SDockTab> SpawnTab_PreviewSceneDetails(const FSpawnTabArgs& Args);
 	TSharedRef<SDockTab> SpawnTab_GraphCanvas(const FSpawnTabArgs& Args);
 	TSharedRef<SDockTab> SpawnTab_NodeDetails(const FSpawnTabArgs& Args);
 
 	void InitDetailsViewPanel();
 
-	/** Scene in which the 3D sim space preview meshes live. */
-	TUniquePtr<FAdvancedPreviewScene> ClothPreviewScene;
+	/** Scene in which the 3D sim space preview meshes live. Ownership shared with AdvancedPreviewSettingsWidget*/
+	TSharedPtr<FChaosClothPreviewScene> ClothPreviewScene;
 
 	TSharedPtr<class FEditorViewportTabContent> ClothPreviewTabContent;
 	AssetEditorViewportFactoryFunction ClothPreviewViewportDelegate;
@@ -104,6 +107,9 @@ private:
 
 	TWeakPtr<SEditorViewport> RestSpaceViewport;
 
+	TSharedPtr<SDockTab> PreviewSceneDockTab;
+	TSharedPtr<SWidget> AdvancedPreviewSettingsWidget;
+
 	TSharedPtr<SClothCollectionOutliner> Outliner;
 
 	TSharedPtr<SComboBox<FName>> SelectedGroupNameComboBox;
