Added an experimental tool to the Cloth Editor for transfering bone weights from a skeletal asset to a cloth asset.

diff --git a/Engine/Plugins/Experimental/ChaosClothAssetEditor/Source/ChaosClothAssetEditor/Private/ChaosClothAsset/ClothEditorCommands.cpp b/Engine/Plugins/Experimental/ChaosClothAssetEditor/Source/ChaosClothAssetEditor/Private/ChaosClothAsset/ClothEditorCommands.cpp
--- a/Engine/Plugins/Experimental/ChaosClothAssetEditor/Source/ChaosClothAssetEditor/Private/ChaosClothAsset/ClothEditorCommands.cpp
+++ b/Engine/Plugins/Experimental/ChaosClothAssetEditor/Source/ChaosClothAssetEditor/Private/ChaosClothAsset/ClothEditorCommands.cpp
@@ -10,6 +10,7 @@ const FString FChaosClothAssetEditorCommands::BeginRemeshToolIdentifier = TEXT("
 const FString FChaosClothAssetEditorCommands::BeginAttributeEditorToolIdentifier = TEXT("BeginAttributeEditorTool");
 const FString FChaosClothAssetEditorCommands::BeginWeightMapPaintToolIdentifier = TEXT("BeginWeightMapPaintTool");
 const FString FChaosClothAssetEditorCommands::BeginClothTrainingToolIdentifier = TEXT("BeginClothTrainingTool");
+const FString FChaosClothAssetEditorCommands::BeginTransferSkinWeightsToolIdentifier = TEXT("BeginTransferSkinWeightsTool");
 const FString FChaosClothAssetEditorCommands::ToggleSimMeshWireframeIdentifier = TEXT("ToggleSimMeshWireframe");
 const FString FChaosClothAssetEditorCommands::ToggleRenderMeshWireframeIdentifier = TEXT("ToggleRenderMeshWireframe");
 const FString FChaosClothAssetEditorCommands::ToggleSimulationSuspendedIdentifier = TEXT("ToggleSimulationSuspended");
@@ -39,6 +40,8 @@ void FChaosClothAssetEditorCommands::RegisterCommands()
 
 	UI_COMMAND(BeginClothTrainingTool, "Train", "Launch Cloth Training tool", EUserInterfaceActionType::Button, FInputChord());
 
+	UI_COMMAND(BeginTransferSkinWeightsTool, "Transfer Skin Weights", "Launch Transfer Skin Weights tool", EUserInterfaceActionType::Button, FInputChord());
+
 	UI_COMMAND(TogglePatternMode, "TogglePatternMode", "Toggle pattern mode", EUserInterfaceActionType::Button, FInputChord());
 
 	UI_COMMAND(ToggleSimMeshWireframe, "ToggleSimMeshWireframe", "Toggle simulation mesh wireframe", EUserInterfaceActionType::ToggleButton, FInputChord());
diff --git a/Engine/Plugins/Experimental/ChaosClothAssetEditor/Source/ChaosClothAssetEditor/Private/ChaosClothAsset/ClothEditorCommands.h b/Engine/Plugins/Experimental/ChaosClothAssetEditor/Source/ChaosClothAssetEditor/Private/ChaosClothAsset/ClothEditorCommands.h
--- a/Engine/Plugins/Experimental/ChaosClothAssetEditor/Source/ChaosClothAssetEditor/Private/ChaosClothAsset/ClothEditorCommands.h
+++ b/Engine/Plugins/Experimental/ChaosClothAssetEditor/Source/ChaosClothAssetEditor/Private/ChaosClothAsset/ClothEditorCommands.h
@@ -39,6 +39,9 @@ public:
 	const static FString BeginClothTrainingToolIdentifier;
 	TSharedPtr<FUICommandInfo> BeginClothTrainingTool;
 
+	const static FString BeginTransferSkinWeightsToolIdentifier;
+	TSharedPtr<FUICommandInfo> BeginTransferSkinWeightsTool;
+
 	// Rest space viewport commands
 	const static FString TogglePatternModeIdentifier;
 	TSharedPtr<FUICommandInfo> TogglePatternMode;
diff --git a/Engine/Plugins/Experimental/ChaosClothAssetEditor/Source/ChaosClothAssetEditor/Private/ChaosClothAsset/ClothEditorMode.cpp b/Engine/Plugins/Experimental/ChaosClothAssetEditor/Source/ChaosClothAssetEditor/Private/ChaosClothAsset/ClothEditorMode.cpp
--- a/Engine/Plugins/Experimental/ChaosClothAssetEditor/Source/ChaosClothAssetEditor/Private/ChaosClothAsset/ClothEditorMode.cpp
+++ b/Engine/Plugins/Experimental/ChaosClothAssetEditor/Source/ChaosClothAssetEditor/Private/ChaosClothAsset/ClothEditorMode.cpp
@@ -35,6 +35,7 @@
 #include "MeshAttributePaintTool.h"
 #include "ChaosClothAsset/ClothWeightMapPaintTool.h"
 #include "ChaosClothAsset/ClothTrainingTool.h"
+#include "ChaosClothAsset/ClothTransferSkinWeightsTool.h"
 #include "DynamicMesh/DynamicMeshAttributeSet.h"
 #include "DynamicMesh/DynamicVertexSkinWeightsAttribute.h"
 #include "DynamicMeshEditor.h"
@@ -179,6 +180,7 @@ void UChaosClothAssetEditorMode::RegisterTools()
 	RegisterTool(CommandInfos.BeginAttributeEditorTool, FChaosClothAssetEditorCommands::BeginAttributeEditorToolIdentifier, NewObject<UAttributeEditorToolBuilder>());
 	RegisterTool(CommandInfos.BeginWeightMapPaintTool, FChaosClothAssetEditorCommands::BeginWeightMapPaintToolIdentifier, NewObject<UClothEditorWeightMapPaintToolBuilder>());
 	RegisterTool(CommandInfos.BeginClothTrainingTool, FChaosClothAssetEditorCommands::BeginClothTrainingToolIdentifier, NewObject<UClothTrainingToolBuilder>());
+	RegisterTool(CommandInfos.BeginTransferSkinWeightsTool, FChaosClothAssetEditorCommands::BeginTransferSkinWeightsToolIdentifier, NewObject<UClothTransferSkinWeightsToolBuilder>());
 }
 
 bool UChaosClothAssetEditorMode::ShouldToolStartBeAllowed(const FString& ToolIdentifier) const
diff --git a/Engine/Plugins/Experimental/ChaosClothAssetEditor/Source/ChaosClothAssetEditor/Private/ChaosClothAsset/ClothEditorModeToolkit.cpp b/Engine/Plugins/Experimental/ChaosClothAssetEditor/Source/ChaosClothAssetEditor/Private/ChaosClothAsset/ClothEditorModeToolkit.cpp
--- a/Engine/Plugins/Experimental/ChaosClothAssetEditor/Source/ChaosClothAssetEditor/Private/ChaosClothAsset/ClothEditorModeToolkit.cpp
+++ b/Engine/Plugins/Experimental/ChaosClothAssetEditor/Source/ChaosClothAssetEditor/Private/ChaosClothAsset/ClothEditorModeToolkit.cpp
@@ -29,8 +29,8 @@ void FChaosClothAssetEditorModeToolkit::BuildToolPalette(FName PaletteIndex, cla
 
 		ToolbarBuilder.AddToolBarButton(Commands.BeginAttributeEditorTool);
 		ToolbarBuilder.AddToolBarButton(Commands.BeginWeightMapPaintTool);
-
 		ToolbarBuilder.AddToolBarButton(Commands.BeginClothTrainingTool);
+		ToolbarBuilder.AddToolBarButton(Commands.BeginTransferSkinWeightsTool);
 	}
 }
 
diff --git a/Engine/Plugins/Experimental/ChaosClothAssetEditor/Source/ChaosClothAssetEditor/Private/ChaosClothAsset/ClothEditorStyle.cpp b/Engine/Plugins/Experimental/ChaosClothAssetEditor/Source/ChaosClothAssetEditor/Private/ChaosClothAsset/ClothEditorStyle.cpp
--- a/Engine/Plugins/Experimental/ChaosClothAssetEditor/Source/ChaosClothAssetEditor/Private/ChaosClothAsset/ClothEditorStyle.cpp
+++ b/Engine/Plugins/Experimental/ChaosClothAssetEditor/Source/ChaosClothAssetEditor/Private/ChaosClothAsset/ClothEditorStyle.cpp
@@ -53,6 +53,10 @@ FChaosClothAssetEditorStyle::FChaosClothAssetEditorStyle()
 	PropertyNameString = "ChaosClothAssetEditor." + FChaosClothAssetEditorCommands::BeginClothTrainingToolIdentifier;
 	Set(*PropertyNameString, new FSlateImageBrush(FChaosClothAssetEditorStyle::InContent("Icons/ModelingAttributePaint_x40", ".png"), ToolbarIconSize));
 
+	// TODO: Get a permanent icon to replace this placeholder
+	PropertyNameString = "ChaosClothAssetEditor." + FChaosClothAssetEditorCommands::BeginTransferSkinWeightsToolIdentifier;
+	Set(*PropertyNameString, new FSlateImageBrush(FChaosClothAssetEditorStyle::InContent("Icons/ModelingAttributePaint_x40", ".png"), ToolbarIconSize));
+
 	// TODO: Get a permanent icon to replace this placeholder
 	PropertyNameString = "ChaosClothAssetEditor." + FChaosClothAssetEditorCommands::TogglePatternModeIdentifier;
 	Set(*PropertyNameString, new FSlateImageBrush(FChaosClothAssetEditorStyle::InContent("Icons/TogglePatternMode_40x", ".png"), ToolbarIconSize));
diff --git a/Engine/Plugins/Experimental/ChaosClothAssetEditor/Source/ChaosClothAssetEditorTools/ChaosClothAssetEditorTools.Build.cs b/Engine/Plugins/Experimental/ChaosClothAssetEditor/Source/ChaosClothAssetEditorTools/ChaosClothAssetEditorTools.Build.cs
--- a/Engine/Plugins/Experimental/ChaosClothAssetEditor/Source/ChaosClothAssetEditorTools/ChaosClothAssetEditorTools.Build.cs
+++ b/Engine/Plugins/Experimental/ChaosClothAssetEditor/Source/ChaosClothAssetEditorTools/ChaosClothAssetEditorTools.Build.cs
@@ -50,7 +50,10 @@ public class ChaosClothAssetEditorTools : ModuleRules
 				"ChaosClothAssetTools",
 				"ChaosClothAssetEngine",
 				"ClothingSystemEditor",
-				"MeshConversion"
+				"MeshConversion",
+				"MeshDescription",
+				"SkeletalMeshDescription",
+				"ChaosClothAsset"
 				// ... add private dependencies that you statically link with here ...	
 			}
 			);
diff --git a/Engine/Plugins/Experimental/ChaosClothAssetEditor/Source/ChaosClothAssetEditorTools/Private/ChaosClothAsset/ClothTransferSkinWeightsTool.cpp b/Engine/Plugins/Experimental/ChaosClothAssetEditor/Source/ChaosClothAssetEditorTools/Private/ChaosClothAsset/ClothTransferSkinWeightsTool.cpp
new file mode 100644
--- /dev/null
+++ b/Engine/Plugins/Experimental/ChaosClothAssetEditor/Source/ChaosClothAssetEditorTools/Private/ChaosClothAsset/ClothTransferSkinWeightsTool.cpp
@@ -0,0 +1,252 @@
+// Copyright Epic Games, Inc. All Rights Reserved.
+
+#include "ChaosClothAsset/ClothTransferSkinWeightsTool.h"
+
+#include "ChaosClothAsset/ClothComponent.h"
+#include "ChaosClothAsset/ClothComponentToolTarget.h"
+#include "ChaosClothAsset/ClothPatternToDynamicMesh.h"
+#include "ChaosClothAsset/ClothAdapter.h"
+#include "ChaosClothAsset/ClothCollection.h"
+
+#include "BoneWeights.h"
+#include "SkeletalMeshAttributes.h"
+
+#include "ToolTargetManager.h"
+#include "MeshDescriptionToDynamicMesh.h"
+
+#include "Rendering/SkeletalMeshLODImporterData.h"
+#include "Rendering/SkeletalMeshModel.h"
+
+#include "DynamicMesh/DynamicMeshAttributeSet.h"
+#include "DynamicMesh/DynamicVertexSkinWeightsAttribute.h"
+#include "DynamicMesh/DynamicMeshAABBTree3.h"
+#include "DynamicMesh/MeshTransforms.h"
+
+#include "Operations/TransferBoneWeights.h"
+
+#include "TransformTypes.h"
+
+#include "InteractiveTool.h"
+#include "InteractiveToolManager.h"
+
+#include "Engine/SkeletalMesh.h"
+
+#define LOCTEXT_NAMESPACE "ClothSkinWeightRetargetingTool"
+
+// ------------------- Properties -------------------
+
+void UClothTransferSkinWeightsToolActionProperties::PostAction(EClothTransferSkinWeightsToolActions Action)
+{
+	if (ParentTool.IsValid())
+	{
+		ParentTool->RequestAction(Action);
+	}
+}
+
+
+// ------------------- Builder -------------------
+
+const FToolTargetTypeRequirements& UClothTransferSkinWeightsToolBuilder::GetTargetRequirements() const
+{
+	static FToolTargetTypeRequirements TypeRequirements({ 
+		UPrimitiveComponentBackedTarget::StaticClass(),
+		UClothAssetBackedTarget::StaticClass()
+		});
+	return TypeRequirements;
+}
+
+bool UClothTransferSkinWeightsToolBuilder::CanBuildTool(const FToolBuilderState& SceneState) const
+{
+	return (SceneState.TargetManager->CountSelectedAndTargetable(SceneState, GetTargetRequirements()) == 1);
+}
+
+UInteractiveTool* UClothTransferSkinWeightsToolBuilder::BuildTool(const FToolBuilderState& SceneState) const
+{
+	UClothTransferSkinWeightsTool* NewTool = NewObject<UClothTransferSkinWeightsTool>();
+	
+	UToolTarget* Target = SceneState.TargetManager->BuildFirstSelectedTargetable(SceneState, GetTargetRequirements());
+	NewTool->SetTarget(Target);
+
+	return NewTool;
+}
+
+// ------------------- Tool -------------------
+
+void UClothTransferSkinWeightsTool::Setup()
+{
+	UInteractiveTool::Setup();
+
+	UClothComponentToolTarget* ClothComponentToolTarget = Cast<UClothComponentToolTarget>(Target);
+	ClothComponent = ClothComponentToolTarget->GetClothComponent();
+	
+	ToolProperties = NewObject<UClothTransferSkinWeightsToolProperties>(this);
+	AddToolPropertySource(ToolProperties);
+
+	ActionProperties = NewObject<UClothTransferSkinWeightsToolActionProperties>(this);
+	ActionProperties->ParentTool = this;
+	AddToolPropertySource(ActionProperties);
+}
+
+void UClothTransferSkinWeightsTool::TransferWeights()
+{
+	using namespace UE::AnimationCore;
+	using namespace UE::Geometry;
+	
+	//TODO: for now, assume we are always transfering from LOD 0, but make this a parameter in the future...
+	constexpr int32 SourceLODIdx = 0; 
+
+	/** 
+	 * Compute mappings between indices and bone names.
+	 * 
+     * @note We assume that each mesh inherits its reference skeleton from the same USkeleton asset. However, their  
+     * internal indexing can be different and hence when transfering weights we need to make sure we reference the correct  
+     * bones via their names instead of indices. 
+     */
+	auto GetBoneMaps = [](const USkinnedAsset* SourceSkinnedAsset, 
+						  const USkinnedAsset* TargetSkinnedAsset, 
+						  TMap<FBoneIndexType, FName>& IndexToBone,
+						  TMap<FName, FBoneIndexType>& BoneToIndex)
+	{
+		BoneToIndex.Reset();
+		IndexToBone.Reset();
+		const FReferenceSkeleton& SourceRefSkeleton = SourceSkinnedAsset->GetRefSkeleton();
+		for (int32 Index = 0; Index < SourceRefSkeleton.GetRawBoneNum(); ++Index)
+		{
+			IndexToBone.Add(Index, SourceRefSkeleton.GetRawRefBoneInfo()[Index].Name);
+		}
+
+		const FReferenceSkeleton& TargetRefSkeleton = TargetSkinnedAsset->GetRefSkeleton();
+		for (int32 Index = 0; Index < TargetRefSkeleton.GetRawBoneNum(); ++Index)
+		{
+			BoneToIndex.Add(TargetRefSkeleton.GetRawRefBoneInfo()[Index].Name, Index);
+		}
+	};
+
+	auto SkeletalMeshToDynamicMesh = [SourceLODIdx](USkeletalMesh* FromSkeletalMeshAsset,
+									    		    FDynamicMesh3& ToDynamicMesh)
+	{
+		FMeshDescription SourceMesh;
+
+		// Check first if we have bulk data available and non-empty.
+		if (FromSkeletalMeshAsset->IsLODImportedDataBuildAvailable(SourceLODIdx) && !FromSkeletalMeshAsset->IsLODImportedDataEmpty(SourceLODIdx))
+		{
+			FSkeletalMeshImportData SkeletalMeshImportData;
+			FromSkeletalMeshAsset->LoadLODImportedData(SourceLODIdx, SkeletalMeshImportData);
+			SkeletalMeshImportData.GetMeshDescription(SourceMesh);
+		}
+		else
+		{
+			// Fall back on the LOD model directly if no bulk data exists. When we commit
+			// the mesh description, we override using the bulk data. This can happen for older
+			// skeletal meshes, from UE 4.24 and earlier.
+			const FSkeletalMeshModel* SkeletalMeshModel = FromSkeletalMeshAsset->GetImportedModel();
+			if (SkeletalMeshModel && SkeletalMeshModel->LODModels.IsValidIndex(SourceLODIdx))
+			{
+				SkeletalMeshModel->LODModels[SourceLODIdx].GetMeshDescription(SourceMesh, FromSkeletalMeshAsset);
+			}
+		}
+
+		FMeshDescriptionToDynamicMesh Converter;
+		Converter.Convert(&SourceMesh, ToDynamicMesh);
+	};
+
+	// User hasn't specified the source mesh in the UI
+	if (ToolProperties->SourceMesh == nullptr) 
+	{
+		//TODO: Display error message
+		return;
+	}
+
+	// Convert source Skeletal Mesh to Dynamic Mesh
+	FDynamicMesh3 SourceDynamicMesh;
+	USkeletalMesh* FromSkeletalMeshAsset = ToolProperties->SourceMesh.Get();
+	SkeletalMeshToDynamicMesh(FromSkeletalMeshAsset, SourceDynamicMesh);
+	FTransformSRT3d SourceToWorld; //TODO: Allows the user to set this value or infer it from an editor
+	MeshTransforms::ApplyTransform(SourceDynamicMesh, SourceToWorld, true);
+
+	UChaosClothAsset* TargetClothAsset = ClothComponent->GetClothAsset(); 
+
+	// Compute bone index mappings
+	TMap<FBoneIndexType, FName> SourceIndexToBone;
+	TMap<FName, FBoneIndexType> TargetBoneToIndex;
+	GetBoneMaps(static_cast<USkinnedAsset*>(FromSkeletalMeshAsset), static_cast<USkinnedAsset*>(TargetClothAsset), SourceIndexToBone, TargetBoneToIndex);
+	
+	// Setup bone weight transfer operator
+	FTransferBoneWeights TransferBoneWeights(&SourceDynamicMesh, FSkeletalMeshAttributes::DefaultSkinWeightProfileName); 
+	TransferBoneWeights.SourceIndexToBone = &SourceIndexToBone;
+	TransferBoneWeights.TargetBoneToIndex = &TargetBoneToIndex;
+	if (TransferBoneWeights.Validate() != EOperationValidationResult::Ok)
+	{
+		//TODO: Display error message
+		return;
+	}
+
+	UE::Chaos::ClothAsset::FClothAdapter ClothAdapter(TargetClothAsset->GetClothCollection());
+    FTransformSRT3d TargetToWorld; //TODO: Allows the user to set this value or infer it from an editor
+
+	// Iterate over the LODs and transfer the bone weights from the source Skeletal mesh to the Cloth asset
+	for (int TargetLODIdx = 0; TargetLODIdx < ClothAdapter.GetNumLods(); ++TargetLODIdx) 
+	{
+		UE::Chaos::ClothAsset::FClothLodAdapter ClothLodAdapter = ClothAdapter.GetLod(TargetLODIdx);
+
+		// Cloth collection data arrays we are writing to
+		TArrayView<int32> NumBoneInfluences = ClothLodAdapter.GetPatternsSimNumBoneInfluences();
+		TArrayView<TArray<int32>> SimBoneIndices = ClothLodAdapter.GetPatternsSimBoneIndices();
+		TArrayView<TArray<float>> SimBoneWeights = ClothLodAdapter.GetPatternsSimBoneWeights();
+
+		const TArrayView<FVector3f> SimPositions =  ClothLodAdapter.GetPatternsSimRestPosition();
+		
+		checkSlow(SimPositions.Num() == SimBoneIndices.Num());
+		
+		const int32 NumVert = ClothLodAdapter.GetPatternsNumSimVertices();
+		constexpr bool bUseParallel = true; 
+
+		// Iterate over each vertex and write the data from FBoneWeights into cloth collection managed arrays
+		ParallelFor(NumVert, [&](int32 VertexID)
+		{
+			const FVector3f Pos = SimPositions[VertexID];
+			const FVector3d PosD = FVector3d((double)Pos[0], (double)Pos[1], (double)Pos[2]);
+			
+			UE::AnimationCore::FBoneWeights BoneWeights;
+			TransferBoneWeights.Compute(PosD, TargetToWorld, BoneWeights);
+			
+			const int32 NumBones = BoneWeights.Num();
+			
+			NumBoneInfluences[VertexID] = NumBones;
+			SimBoneIndices[VertexID].SetNum(NumBones);
+			SimBoneWeights[VertexID].SetNum(NumBones);
+
+			for (int BoneIdx = 0; BoneIdx < NumBones; ++BoneIdx) 
+			{
+				SimBoneIndices[VertexID][BoneIdx] = BoneWeights[BoneIdx].GetBoneIndex();
+				SimBoneWeights[VertexID][BoneIdx] = BoneWeights[BoneIdx].GetWeight();
+			}
+
+		}, bUseParallel ? EParallelForFlags::None : EParallelForFlags::ForceSingleThread);
+	}
+}
+
+
+void UClothTransferSkinWeightsTool::OnTick(float DeltaTime)
+{
+	if (PendingAction != EClothTransferSkinWeightsToolActions::NoAction)
+	{
+		if (PendingAction == EClothTransferSkinWeightsToolActions::Transfer)
+		{
+			TransferWeights();
+		}
+		PendingAction = EClothTransferSkinWeightsToolActions::NoAction;
+	}
+}
+
+
+void UClothTransferSkinWeightsTool::RequestAction(EClothTransferSkinWeightsToolActions ActionType)
+{
+	if (PendingAction != EClothTransferSkinWeightsToolActions::NoAction)
+	{
+		return;
+	}
+	PendingAction = ActionType;
+}
+
+#undef LOCTEXT_NAMESPACE
diff --git a/Engine/Plugins/Experimental/ChaosClothAssetEditor/Source/ChaosClothAssetEditorTools/Public/ChaosClothAsset/ClothTransferSkinWeightsTool.h b/Engine/Plugins/Experimental/ChaosClothAssetEditor/Source/ChaosClothAssetEditorTools/Public/ChaosClothAsset/ClothTransferSkinWeightsTool.h
new file mode 100644
--- /dev/null
+++ b/Engine/Plugins/Experimental/ChaosClothAssetEditor/Source/ChaosClothAssetEditorTools/Public/ChaosClothAsset/ClothTransferSkinWeightsTool.h
@@ -0,0 +1,101 @@
+// Copyright Epic Games, Inc. All Rights Reserved.
+
+#pragma once
+
+#include "UObject/Object.h"
+#include "SingleSelectionTool.h"
+#include "InteractiveToolBuilder.h"
+
+#include "ClothTransferSkinWeightsTool.generated.h"
+
+class UChaosClothComponent;
+class UClothTransferSkinWeightsTool;
+class USkeletalMesh;
+
+UCLASS()
+class CHAOSCLOTHASSETEDITORTOOLS_API UClothTransferSkinWeightsToolProperties : public UObject
+{
+	GENERATED_BODY()
+public:
+
+	UPROPERTY(EditAnywhere, Category = Files)
+	TObjectPtr<USkeletalMesh> SourceMesh;
+};
+
+UENUM()
+enum class EClothTransferSkinWeightsToolActions
+{
+	NoAction,
+	Transfer
+};
+
+UCLASS()
+class CHAOSCLOTHASSETEDITORTOOLS_API UClothTransferSkinWeightsToolActionProperties : public UObject
+{
+	GENERATED_BODY()
+public:
+
+	TWeakObjectPtr<UClothTransferSkinWeightsTool> ParentTool;
+
+	void Initialize(UClothTransferSkinWeightsTool* ParentToolIn) { ParentTool = ParentToolIn; }
+
+	void PostAction(EClothTransferSkinWeightsToolActions Action);
+
+	UFUNCTION(CallInEditor, Category = Actions, meta = (DisplayName = "Transfer weights", DisplayPriority = 1))
+	void TransferWeights()
+	{
+		PostAction(EClothTransferSkinWeightsToolActions::Transfer);
+	}
+
+};
+
+
+UCLASS()
+class CHAOSCLOTHASSETEDITORTOOLS_API UClothTransferSkinWeightsToolBuilder : public UInteractiveToolWithToolTargetsBuilder
+{
+	GENERATED_BODY()
+
+protected:
+
+	virtual const FToolTargetTypeRequirements& GetTargetRequirements() const override;
+
+public:
+
+	virtual bool CanBuildTool(const FToolBuilderState& SceneState) const override;
+
+	virtual UInteractiveTool* BuildTool(const FToolBuilderState& SceneState) const override;
+};
+
+
+UCLASS()
+class CHAOSCLOTHASSETEDITORTOOLS_API UClothTransferSkinWeightsTool : public USingleSelectionTool
+{
+	GENERATED_BODY()
+public:
+
+	// UInteractiveTool
+	virtual void Setup() override;
+	virtual void OnTick(float DeltaTime) override;
+
+private:
+
+	friend class UClothTransferSkinWeightsToolActionProperties;
+
+	UPROPERTY()
+	TObjectPtr<UClothTransferSkinWeightsToolProperties> ToolProperties;
+
+	UPROPERTY()
+	TObjectPtr<UClothTransferSkinWeightsToolActionProperties> ActionProperties;
+
+	EClothTransferSkinWeightsToolActions PendingAction = EClothTransferSkinWeightsToolActions::NoAction;
+
+	UPROPERTY()
+	TObjectPtr<const UChaosClothComponent> ClothComponent;
+
+	void RequestAction(EClothTransferSkinWeightsToolActions ActionType);
+
+	void TransferWeights();
+
+};
+
+
