Cloth Editor: reinitialize the graph editor widget if the Dataflow asset attached to the cloth asset is changed

diff --git a/Engine/Plugins/Experimental/ChaosClothAssetEditor/Source/ChaosClothAssetEditor/Private/ChaosClothAsset/ClothEditorToolkit.cpp b/Engine/Plugins/Experimental/ChaosClothAssetEditor/Source/ChaosClothAssetEditor/Private/ChaosClothAsset/ClothEditorToolkit.cpp
--- a/Engine/Plugins/Experimental/ChaosClothAssetEditor/Source/ChaosClothAssetEditor/Private/ChaosClothAsset/ClothEditorToolkit.cpp
+++ b/Engine/Plugins/Experimental/ChaosClothAssetEditor/Source/ChaosClothAssetEditor/Private/ChaosClothAsset/ClothEditorToolkit.cpp
@@ -284,7 +284,7 @@ void FChaosClothAssetEditorToolkit::CreateWidgets()
 			Dataflow->Schema = UDataflowSchema::StaticClass();
 
 			NodeDetailsEditor = CreateNodeDetailsEditorWidget(ClothAsset);
-			GraphEditor = CreateGraphEditorWidget(Dataflow, NodeDetailsEditor);
+			GraphEditor = CreateGraphEditorWidget();
 		}
 	}
 }
@@ -670,6 +670,7 @@ void FChaosClothAssetEditorToolkit::InitDetailsViewPanel()
 		ensure(ObjectToEditInDetailsView->HasAnyFlags(RF_Transactional));		// Ensure all objects are transactable for undo/redo in the details panel
 		SetEditingObject(ObjectToEditInDetailsView);
 	}
+	DetailsView->OnFinishedChangingProperties().AddSP(this, &FChaosClothAssetEditorToolkit::OnFinishedChangingAssetProperties);
 
 	TArray<FAdvancedPreviewSceneModule::FDetailDelegates> Delegates;	
 
@@ -704,45 +705,96 @@ void FChaosClothAssetEditorToolkit::InitDetailsViewPanel()
 
 		ClothCollectionGroupNames = ClothCollection->GroupNames();
 	}
-
 }
 
-TSharedRef<SGraphEditor> FChaosClothAssetEditorToolkit::CreateGraphEditorWidget(UDataflow* DataflowToEdit, TSharedPtr<IStructureDetailsView> InNodeDetailsEditor)
+void FChaosClothAssetEditorToolkit::OnFinishedChangingAssetProperties(const FPropertyChangedEvent& Event)
 {
-	ensure(DataflowToEdit);
-	using namespace Dataflow;
+	const FProperty* const ChangedProperty = Event.Property;
 
-	FDataflowEditorCommands::FGraphEvaluationCallback Evaluate = [this, DataflowToEdit](FDataflowNode* Node, FDataflowOutput* Out)
+	if (ChangedProperty && ChangedProperty->GetFName() == TEXT("DataflowAsset"))
 	{
-		if (DataflowToEdit)
+		const UChaosClothAsset* const ClothAsset = GetAsset();
+		if (ClothAsset)
 		{
-			UObject* const Asset = GetAsset();
+			Dataflow = ClothEditorToolkitHelpers::GetDataflowFrom(ClothAsset);
 
-			if (!DataflowContext)
+			if (Dataflow)
 			{
-				DataflowContext = TSharedPtr<Dataflow::FEngineContext>(new Dataflow::FClothAssetDataflowContext(Asset, Dataflow, Dataflow::FTimestamp::Invalid));
+				Dataflow->Schema = UDataflowSchema::StaticClass();
+				ReinitializeGraphEditorWidget();
 			}
-			LastDataflowNodeTimestamp = Dataflow::FTimestamp::Invalid;
-			
-			FDataflowEditorCommands::EvaluateTerminalNode(*DataflowContext.Get(), LastDataflowNodeTimestamp, Dataflow, Node, Out, Asset, DataflowTerminalPath);
 		}
+	}
+}
+
+void FChaosClothAssetEditorToolkit::EvaluateNode(FDataflowNode* Node, FDataflowOutput* Out)
+{
+	if (Dataflow)
+	{
+		UObject* const Asset = GetAsset();
+
+		if (!DataflowContext)
+		{
+			DataflowContext = TSharedPtr<Dataflow::FEngineContext>(new Dataflow::FClothAssetDataflowContext(Asset, Dataflow, Dataflow::FTimestamp::Invalid));
+		}
+		LastDataflowNodeTimestamp = Dataflow::FTimestamp::Invalid;
+
+		FDataflowEditorCommands::EvaluateTerminalNode(*DataflowContext.Get(), LastDataflowNodeTimestamp, Dataflow, Node, Out, Asset, DataflowTerminalPath);
+	}
+};
+
+TSharedRef<SDataflowGraphEditor> FChaosClothAssetEditorToolkit::CreateGraphEditorWidget()
+{
+	ensure(Dataflow);
+	using namespace Dataflow;
+
+	const auto EvalLambda = [this](FDataflowNode* Node, FDataflowOutput* Out)
+	{
+		this->EvaluateNode(Node, Out);
 	};
 
 	SGraphEditor::FGraphEditorEvents InEvents;
 	InEvents.OnVerifyTextCommit = FOnNodeVerifyTextCommit::CreateSP(this, &FChaosClothAssetEditorToolkit::OnNodeVerifyTitleCommit);
 	InEvents.OnTextCommitted = FOnNodeTextCommitted::CreateSP(this, &FChaosClothAssetEditorToolkit::OnNodeTitleCommitted);
 
-	TSharedRef<SDataflowGraphEditor> NewGraphEditor = SNew(SDataflowGraphEditor, DataflowToEdit)
-		.GraphToEdit(DataflowToEdit)
+	TSharedRef<SDataflowGraphEditor> NewGraphEditor = SNew(SDataflowGraphEditor, Dataflow)
+		.GraphToEdit(Dataflow)
 		.GraphEvents(InEvents)
-		.DetailsView(InNodeDetailsEditor)
-		.EvaluateGraph(Evaluate);
+		.DetailsView(NodeDetailsEditor)
+		.EvaluateGraph(EvalLambda);
 
 	NewGraphEditor->OnSelectionChangedMulticast.AddSP(this, &FChaosClothAssetEditorToolkit::OnNodeSelectionChanged);
 
 	return NewGraphEditor;
 }
 
+
+void FChaosClothAssetEditorToolkit::ReinitializeGraphEditorWidget()
+{
+	ensure(Dataflow);
+
+	const auto EvalLambda = [this](FDataflowNode* Node, FDataflowOutput* Out)
+	{
+		EvaluateNode(Node, Out);
+	};
+
+	SGraphEditor::FGraphEditorEvents InEvents;
+	InEvents.OnVerifyTextCommit = FOnNodeVerifyTextCommit::CreateSP(this, &FChaosClothAssetEditorToolkit::OnNodeVerifyTitleCommit);
+	InEvents.OnTextCommitted = FOnNodeTextCommitted::CreateSP(this, &FChaosClothAssetEditorToolkit::OnNodeTitleCommitted);
+
+	SDataflowGraphEditor::FArguments Args;
+	Args._GraphToEdit = Dataflow;
+	Args._GraphEvents = InEvents;
+	Args._DetailsView = NodeDetailsEditor;
+	Args._EvaluateGraph = EvalLambda;
+
+	UChaosClothAsset* const ClothAsset = GetAsset();
+	GraphEditor->Construct(Args, ClothAsset);
+
+	GraphEditor->OnSelectionChangedMulticast.Clear();
+	GraphEditor->OnSelectionChangedMulticast.AddSP(this, &FChaosClothAssetEditorToolkit::OnNodeSelectionChanged);
+}
+
 TSharedPtr<IStructureDetailsView> FChaosClothAssetEditorToolkit::CreateNodeDetailsEditorWidget(UObject* ObjectToEdit)
 {
 	ensure(ObjectToEdit);
diff --git a/Engine/Plugins/Experimental/ChaosClothAssetEditor/Source/ChaosClothAssetEditor/Private/ChaosClothAsset/ClothEditorToolkit.h b/Engine/Plugins/Experimental/ChaosClothAssetEditor/Source/ChaosClothAssetEditor/Private/ChaosClothAsset/ClothEditorToolkit.h
--- a/Engine/Plugins/Experimental/ChaosClothAssetEditor/Source/ChaosClothAssetEditor/Private/ChaosClothAsset/ClothEditorToolkit.h
+++ b/Engine/Plugins/Experimental/ChaosClothAssetEditor/Source/ChaosClothAssetEditor/Private/ChaosClothAsset/ClothEditorToolkit.h
@@ -12,6 +12,7 @@ class SClothCollectionOutliner;
 class UDataflow;
 class UChaosClothAsset;
 class SGraphEditor;
+class SDataflowGraphEditor;
 class IStructureDetailsView;
 class UEdGraphNode;
 class UChaosClothComponent;
@@ -96,6 +97,7 @@ private:
 	TSharedRef<SDockTab> SpawnTab_NodeDetails(const FSpawnTabArgs& Args);
 
 	void InitDetailsViewPanel();
+	void OnFinishedChangingAssetProperties(const FPropertyChangedEvent&);
 
 	/** Scene in which the 3D sim space preview meshes live. Ownership shared with AdvancedPreviewSettingsWidget*/
 	TSharedPtr<FChaosClothPreviewScene> ClothPreviewScene;
@@ -119,9 +121,12 @@ private:
 	 
 	UDataflow* Dataflow = nullptr;
 
+	void EvaluateNode(FDataflowNode* Node, FDataflowOutput* Out);
+
 	static const FName GraphCanvasTabId;
-	TSharedPtr<SGraphEditor> GraphEditor;
-	TSharedRef<SGraphEditor> CreateGraphEditorWidget(UDataflow* ObjectToEdit, TSharedPtr<IStructureDetailsView> PropertiesEditor);
+	TSharedPtr<SDataflowGraphEditor> GraphEditor;
+	TSharedRef<SDataflowGraphEditor> CreateGraphEditorWidget();
+	void ReinitializeGraphEditorWidget();
 
 	static const FName NodeDetailsTabId;
 	TSharedPtr<IStructureDetailsView> NodeDetailsEditor;
