Cloth Asset Builder: fix for computing wrap deformer weights

diff --git a/Engine/Plugins/Experimental/ChaosClothAssetEditor/Source/ChaosClothAssetTools/Private/ChaosClothAsset/ClothAssetBuilderEditor.cpp b/Engine/Plugins/Experimental/ChaosClothAssetEditor/Source/ChaosClothAssetTools/Private/ChaosClothAsset/ClothAssetBuilderEditor.cpp
--- a/Engine/Plugins/Experimental/ChaosClothAssetEditor/Source/ChaosClothAssetTools/Private/ChaosClothAsset/ClothAssetBuilderEditor.cpp
+++ b/Engine/Plugins/Experimental/ChaosClothAssetEditor/Source/ChaosClothAssetTools/Private/ChaosClothAsset/ClothAssetBuilderEditor.cpp
@@ -173,6 +173,12 @@ void UClothAssetBuilderEditor::BuildLod(FSkeletalMeshLODModel& LODModel, const U
 			}
 		}
 
+		// Remap the LOD indices with the new vertex indices
+		for (uint32& RenderIndex : LODModel.IndexBuffer)
+		{
+			RenderIndex = LodRenderIndexRemap[RenderIndex];
+		}
+
 		// Update the section bone map
 		Section.BoneMap = SectionBoneMap.Array();
 		ActiveBoneIndices.Append(MoveTemp(SectionBoneMap));
@@ -188,11 +194,20 @@ void UClothAssetBuilderEditor::BuildLod(FSkeletalMeshLODModel& LODModel, const U
 		Section.ClothingData.AssetGuid = ClothAsset.AssetGuid;  // There is only one cloth asset,
 		Section.CorrespondClothAssetIndex = 0;       // this one
 
-		const ClothingMeshUtils::ClothMeshDesc TargetMesh(
-			TConstArrayView<FVector3f>(ClothCollection->RenderPosition.GetData(), NumVertices),
-			TConstArrayView<FVector3f>(ClothCollection->RenderNormal.GetData(), NumVertices),
-			TConstArrayView<FVector3f>(ClothCollection->RenderTangentU.GetData(), NumVertices),
-			TConstArrayView<uint32>(Indices));
+		TArray<FVector3f> RenderPositions;
+		TArray<FVector3f> RenderNormals;
+		TArray<FVector3f> RenderTangents;
+		for (const FSoftSkinVertex& SoftVert : Section.SoftVertices)
+		{
+			RenderPositions.Add(SoftVert.Position);
+			RenderNormals.Add(SoftVert.TangentZ);
+			RenderTangents.Add(SoftVert.TangentX);
+		}
+
+		const ClothingMeshUtils::ClothMeshDesc TargetMesh(RenderPositions, 
+			RenderNormals, 
+			RenderTangents,
+			LODModel.IndexBuffer);
 
 		ClothingMeshUtils::GenerateMeshToMeshVertData(
 			Section.ClothMappingDataLODs[0],
@@ -249,11 +264,6 @@ void UClothAssetBuilderEditor::BuildLod(FSkeletalMeshLODModel& LODModel, const U
 		FSkelMeshSourceSectionUserData::GetSourceSectionUserData(LODModel.UserSectionsData, Section);
 	}
 
-	// Remap the LOD indices with the new vertex indices
-	for (uint32& RenderIndex : LODModel.IndexBuffer)
-	{
-		RenderIndex = LodRenderIndexRemap[RenderIndex];
-	}
 
 	// Update the active bone indices on the LOD model
 	LODModel.ActiveBoneIndices = ActiveBoneIndices.Array();
