Cloth Editor: Adding a transform gizmo to cloth preview scene

diff --git a/Engine/Plugins/Experimental/ChaosClothAssetEditor/Source/ChaosClothAssetEditor/Private/ChaosClothAsset/ClothEditor3DViewportClient.cpp b/Engine/Plugins/Experimental/ChaosClothAssetEditor/Source/ChaosClothAssetEditor/Private/ChaosClothAsset/ClothEditor3DViewportClient.cpp
--- a/Engine/Plugins/Experimental/ChaosClothAssetEditor/Source/ChaosClothAssetEditor/Private/ChaosClothAsset/ClothEditor3DViewportClient.cpp
+++ b/Engine/Plugins/Experimental/ChaosClothAssetEditor/Source/ChaosClothAssetEditor/Private/ChaosClothAsset/ClothEditor3DViewportClient.cpp
@@ -9,6 +9,12 @@
 #include "Dataflow/DataflowObject.h"
 #include "Dataflow/DataflowActor.h"
 #include "Dataflow/DataflowRenderingFactory.h"
+#include "EdModeInteractiveToolsContext.h"
+#include "BaseGizmos/TransformProxy.h"
+#include "BaseGizmos/CombinedTransformGizmo.h"
+#include "BaseGizmos/TransformGizmoUtil.h"
+#include "Engine/Selection.h"
+#include "EngineUtils.h"
 
 FChaosClothAssetEditor3DViewportClient::FChaosClothAssetEditor3DViewportClient(FEditorModeTools* InModeTools,
 	FPreviewScene* InPreviewScene, 
@@ -20,8 +26,31 @@ FChaosClothAssetEditor3DViewportClient::FChaosClothAssetEditor3DViewportClient(F
 
 	// Call this once with the default value to get everything in a consistent state
 	EnableRenderMeshWireframe(bRenderMeshWireframe);
+
+	// Set up Gizmo and TransformProxy
+
+	UModeManagerInteractiveToolsContext* const InteractiveToolsContext = ModeTools->GetInteractiveToolsContext();
+	TransformProxy = NewObject<UTransformProxy>();
+
+	UInteractiveGizmoManager* const GizmoManager = InteractiveToolsContext->GizmoManager;
+	const FString GizmoIdentifier = TEXT("ChaosClothAssetEditor3DViewportClientGizmoIdentifier");
+	Gizmo = GizmoManager->Create3AxisTransformGizmo(this, GizmoIdentifier);
+
+	UInteractiveToolManager* const ToolManager = InteractiveToolsContext->ToolManager;
+	Gizmo->SetActiveTarget(TransformProxy, ToolManager);
+	Gizmo->SetVisibility(true);
+	Gizmo->bUseContextGizmoMode = false;
+	Gizmo->bUseContextCoordinateSystem = false;
+	Gizmo->ActiveGizmoMode = EToolContextTransformGizmoMode::Combined;
 }
 
+void FChaosClothAssetEditor3DViewportClient::AddReferencedObjects(FReferenceCollector& Collector)
+{
+	FEditorViewportClient::AddReferencedObjects(Collector);
+
+	Collector.AddReferencedObject(TransformProxy);
+	Collector.AddReferencedObject(Gizmo);
+}
 
 void FChaosClothAssetEditor3DViewportClient::Tick(float DeltaSeconds)
 {
@@ -88,6 +117,12 @@ void FChaosClothAssetEditor3DViewportClient::EnableRenderMeshWireframe(bool bEna
 void FChaosClothAssetEditor3DViewportClient::SetClothComponent(TObjectPtr<UChaosClothComponent> InClothComponent)
 {
 	ClothComponent = InClothComponent;
+
+	if (ClothComponent)
+	{
+		TransformProxy->AddComponent(ClothComponent);
+		TransformProxy->SetTransform(ClothComponent->GetComponentToWorld());
+	}
 }
 
 void FChaosClothAssetEditor3DViewportClient::SetClothEdMode(TObjectPtr<UChaosClothAssetEditorMode> InClothEdMode)
@@ -144,13 +179,7 @@ bool FChaosClothAssetEditor3DViewportClient::IsSimulationSuspended() const
 
 void FChaosClothAssetEditor3DViewportClient::Draw(const FSceneView* View, FPrimitiveDrawInterface* PDI)
 {
-	if (!bSimMeshWireframe)
-	{
-		return;
-	}
-
-	// TODO: Draw simulation mesh wireframe
-
+	FEditorViewportClient::Draw(View, PDI);
 }
 
 FBox FChaosClothAssetEditor3DViewportClient::PreviewBoundingBox() const
@@ -161,4 +190,81 @@ FBox FChaosClothAssetEditor3DViewportClient::PreviewBoundingBox() const
 	}
 
 	return FBox(ForceInitToZero);
-}
\ No newline at end of file
+}
+
+void FChaosClothAssetEditor3DViewportClient::ProcessClick(FSceneView& View, HHitProxy* HitProxy, FKey Key, EInputEvent Event, uint32 HitX, uint32 HitY)
+{
+	FEditorViewportClient::ProcessClick(View, HitProxy, Key, Event, HitX, HitY);
+
+	const bool bIsShiftKeyDown = Viewport->KeyState(EKeys::LeftShift) || Viewport->KeyState(EKeys::RightShift);
+	const bool bIsCtrlKeyDown = Viewport->KeyState(EKeys::LeftControl) || Viewport->KeyState(EKeys::RightControl);
+
+	USelection* SelectedComponents = ModeTools->GetSelectedComponents();
+	SelectedComponents->Modify();
+
+	if (!bIsShiftKeyDown && !bIsCtrlKeyDown)
+	{
+		SelectedComponents->DeselectAll();
+	}
+
+	if (HitProxy && HitProxy->IsA(HActor::StaticGetType()))
+	{
+		const HActor* const ActorProxy = static_cast<HActor*>(HitProxy);
+		if (ActorProxy && ActorProxy->Actor)
+		{
+			const AActor* const Actor = ActorProxy->Actor;		
+			USceneComponent* RootComponent = Actor->GetRootComponent();
+			if (RootComponent)
+			{
+				if (bIsShiftKeyDown)
+				{
+					SelectedComponents->Select(RootComponent);
+				}
+				else if (bIsCtrlKeyDown)
+				{
+					// Don't use USelection::ToggleSelect here because that checks membership in GObjectSelection, not the given USelection...
+					TArray<USceneComponent*> Components;
+					SelectedComponents->GetSelectedObjects(Components);
+					const bool bIsSelected = Components.Contains(RootComponent);	
+
+					if (bIsSelected)
+					{
+						SelectedComponents->Deselect(RootComponent);
+					}
+					else
+					{
+						SelectedComponents->Select(RootComponent);
+					}
+				}
+				else
+				{
+					SelectedComponents->Select(RootComponent);
+				}
+			}
+		}
+	}
+
+	// Update TransformProxy
+
+	TransformProxy = NewObject<UTransformProxy>();
+	TArray<USceneComponent*> Components;
+	SelectedComponents->GetSelectedObjects(Components);
+	for (USceneComponent* SelectedComponent : Components)
+	{
+		TransformProxy->AddComponent(SelectedComponent);
+	}
+
+	// Update gizmo
+
+	if (Components.Num() > 0)
+	{
+		Gizmo->SetActiveTarget(TransformProxy);
+		Gizmo->SetVisibility(true);
+	}
+	else
+	{
+		Gizmo->ClearActiveTarget();
+		Gizmo->SetVisibility(false);
+	}
+
+}
diff --git a/Engine/Plugins/Experimental/ChaosClothAssetEditor/Source/ChaosClothAssetEditor/Private/ChaosClothAsset/ClothEditor3DViewportClient.h b/Engine/Plugins/Experimental/ChaosClothAssetEditor/Source/ChaosClothAssetEditor/Private/ChaosClothAsset/ClothEditor3DViewportClient.h
--- a/Engine/Plugins/Experimental/ChaosClothAssetEditor/Source/ChaosClothAssetEditor/Private/ChaosClothAsset/ClothEditor3DViewportClient.h
+++ b/Engine/Plugins/Experimental/ChaosClothAssetEditor/Source/ChaosClothAssetEditor/Private/ChaosClothAsset/ClothEditor3DViewportClient.h
@@ -9,6 +9,8 @@
 class UChaosClothComponent;
 class UChaosClothAssetEditorMode;
 class FChaosClothAssetEditorToolkit;
+class UTransformProxy;
+class UCombinedTransformGizmo;
 
 /**
  * Viewport client for the 3d sim preview in the cloth editor. Currently same as editor viewport
@@ -23,12 +25,6 @@ public:
 
 	virtual ~FChaosClothAssetEditor3DViewportClient() = default;
 
-	virtual bool CanSetWidgetMode(UE::Widget::EWidgetMode NewMode) const override {	return false; }
-	virtual void SetWidgetMode(UE::Widget::EWidgetMode NewMode) override {}
-	virtual UE::Widget::EWidgetMode GetWidgetMode() const override { return UE::Widget::EWidgetMode::WM_None; }
-
-	virtual void Tick(float DeltaSeconds) override;
-
 	void EnableSimMeshWireframe(bool bEnable ) { bSimMeshWireframe = bEnable; }
 	bool SimMeshWireframeEnabled() const { return bSimMeshWireframe; }
 
@@ -49,21 +45,31 @@ public:
 
 private:
 
+	// FGCObject override
+	virtual void AddReferencedObjects(FReferenceCollector& Collector) override;
+
+	// FEditorViewportClient overrides
+	virtual bool CanSetWidgetMode(UE::Widget::EWidgetMode NewMode) const override { return false; }
+	virtual void SetWidgetMode(UE::Widget::EWidgetMode NewMode) override {}
+	virtual UE::Widget::EWidgetMode GetWidgetMode() const override { return UE::Widget::EWidgetMode::WM_None; }
+	virtual void Tick(float DeltaSeconds) override;
+	virtual void ProcessClick(FSceneView& View, HHitProxy* HitProxy, FKey Key, EInputEvent Event, uint32 HitX, uint32 HitY) override;
+	virtual void Draw(const FSceneView* View, FPrimitiveDrawInterface* PDI) override;
+
 	TObjectPtr<UChaosClothComponent> ClothComponent;
 
 	TObjectPtr<UChaosClothAssetEditorMode> ClothEdMode;
 
 	TSharedPtr<const FChaosClothAssetEditorToolkit> ClothToolkit;
-
-	// Debug draw of simulation meshes
-	virtual void Draw(const FSceneView* View, FPrimitiveDrawInterface* PDI) override;
 	
 	bool bSimMeshWireframe = true;
 	bool bRenderMeshWireframe = false;
 
-	//~ Dataflow render support
-
+	// Dataflow render support
 	Dataflow::FTimestamp LastModifiedTimestamp = Dataflow::FTimestamp::Invalid;
 
-	//~ Dataflow render support
+	// Gizmo support
+	TObjectPtr<UTransformProxy> TransformProxy = nullptr;
+	TObjectPtr<UCombinedTransformGizmo> Gizmo = nullptr;
+
 };
diff --git a/Engine/Plugins/Experimental/ChaosClothAssetEditor/Source/ChaosClothAssetEditor/Private/ChaosClothAsset/ClothEditorMode.cpp b/Engine/Plugins/Experimental/ChaosClothAssetEditor/Source/ChaosClothAssetEditor/Private/ChaosClothAsset/ClothEditorMode.cpp
--- a/Engine/Plugins/Experimental/ChaosClothAssetEditor/Source/ChaosClothAssetEditor/Private/ChaosClothAsset/ClothEditorMode.cpp
+++ b/Engine/Plugins/Experimental/ChaosClothAssetEditor/Source/ChaosClothAssetEditor/Private/ChaosClothAsset/ClothEditorMode.cpp
@@ -590,12 +590,7 @@ void UChaosClothAssetEditorMode::InitializeTargets(const TArray<TObjectPtr<UObje
 	{
 		if (UChaosClothAsset* ChaosClothAsset = Cast<UChaosClothAsset>(Asset))
 		{
-			PreviewScene->CreateClothComponent(ChaosClothAsset);
-
-			USelection* SelectedComponents = GetModeManager()->GetSelectedComponents();
-			SelectedComponents->Modify();
-			SelectedComponents->Select(PreviewScene->ClothComponent);
-
+			PreviewScene->CreateClothActor(ChaosClothAsset);
 			break;
 		}
 	}
diff --git a/Engine/Plugins/Experimental/ChaosClothAssetEditor/Source/ChaosClothAssetEditor/Private/ChaosClothAsset/ClothEditorPreviewScene.cpp b/Engine/Plugins/Experimental/ChaosClothAssetEditor/Source/ChaosClothAssetEditor/Private/ChaosClothAsset/ClothEditorPreviewScene.cpp
--- a/Engine/Plugins/Experimental/ChaosClothAssetEditor/Source/ChaosClothAssetEditor/Private/ChaosClothAsset/ClothEditorPreviewScene.cpp
+++ b/Engine/Plugins/Experimental/ChaosClothAssetEditor/Source/ChaosClothAssetEditor/Private/ChaosClothAsset/ClothEditorPreviewScene.cpp
@@ -30,7 +30,7 @@ FChaosClothPreviewScene::FChaosClothPreviewScene(FPreviewScene::ConstructionValu
 
 	if (PreviewSceneDescription->SkeletalMeshAsset)
 	{
-		InitializeSkeletalMeshActor();
+		CreateSkeletalMeshActor();
 	}
 }
 
@@ -38,6 +38,11 @@ FChaosClothPreviewScene::~FChaosClothPreviewScene()
 {
 	if (SkeletalMeshActor)
 	{
+		if (SkeletalMeshActor->GetSkeletalMeshComponent())
+		{
+			SkeletalMeshActor->GetSkeletalMeshComponent()->TransformUpdated.RemoveAll(this);
+		}
+
 		SkeletalMeshActor->UnregisterAllComponents();
 	}
 
@@ -54,6 +59,7 @@ void FChaosClothPreviewScene::AddReferencedObjects(FReferenceCollector& Collecto
 	Collector.AddReferencedObject(PreviewSceneDescription);
 	Collector.AddReferencedObject(ClothComponent);
 	Collector.AddReferencedObject(SkeletalMeshActor);
+	Collector.AddReferencedObject(ClothActor);
 }
 
 void FChaosClothPreviewScene::SceneDescriptionPropertyChanged(struct FPropertyChangedEvent& PropertyChangedEvent)
@@ -62,7 +68,7 @@ void FChaosClothPreviewScene::SceneDescriptionPropertyChanged(struct FPropertyCh
 	{
 		if (PreviewSceneDescription->SkeletalMeshAsset)
 		{
-			InitializeSkeletalMeshActor();
+			CreateSkeletalMeshActor();
 		}
 	}
 
@@ -70,16 +76,28 @@ void FChaosClothPreviewScene::SceneDescriptionPropertyChanged(struct FPropertyCh
 	{
 		if (SkeletalMeshActor)
 		{
-			SkeletalMeshActor->SetActorTransform(PreviewSceneDescription->SkeletalMeshTransform);
+			SkeletalMeshActor->GetSkeletalMeshComponent()->SetComponentToWorld(PreviewSceneDescription->SkeletalMeshTransform);
 		}
 	}
 }
 
 
-void FChaosClothPreviewScene::InitializeSkeletalMeshActor()
+void FChaosClothPreviewScene::SkeletalMeshTransformChanged(USceneComponent* UpdatedComponent, EUpdateTransformFlags UpdateTransformFlags, ETeleportType Teleport)
+{
+	ensure(UpdatedComponent == SkeletalMeshActor->GetSkeletalMeshComponent());
+	PreviewSceneDescription->SkeletalMeshTransform = UpdatedComponent->GetComponentToWorld();
+}
+
+
+void FChaosClothPreviewScene::CreateSkeletalMeshActor()
 {
 	if (SkeletalMeshActor)
 	{
+		if (SkeletalMeshActor->GetSkeletalMeshComponent())
+		{
+			SkeletalMeshActor->GetSkeletalMeshComponent()->TransformUpdated.RemoveAll(this);
+		}
+
 		SkeletalMeshActor->UnregisterAllComponents();
 	}
 
@@ -92,14 +110,27 @@ void FChaosClothPreviewScene::InitializeSkeletalMeshActor()
 	}
 
 	SkeletalMeshActor->RegisterAllComponents();
+
+	SkeletalMeshActor->GetSkeletalMeshComponent()->TransformUpdated.AddRaw(this, &FChaosClothPreviewScene::SkeletalMeshTransformChanged);
 }
 
 
-void FChaosClothPreviewScene::CreateClothComponent(UChaosClothAsset* Asset)
+void FChaosClothPreviewScene::CreateClothActor(UChaosClothAsset* Asset)
 {
-	ClothComponent = NewObject<UChaosClothComponent>();
+	if (ClothActor)
+	{
+		ClothActor->UnregisterAllComponents();
+	}
+
+	ClothActor = GetWorld()->SpawnActor<AActor>(AActor::StaticClass());
+
+	check(ClothActor);
+
+	ClothComponent = NewObject<UChaosClothComponent>(ClothActor);
 	ClothComponent->SetClothAsset(Asset);
-	ClothComponent->RegisterComponentWithWorld(PreviewWorld);
+
+	ClothActor->SetRootComponent(ClothComponent);
+	ClothActor->RegisterAllComponents();
 }
 
 
diff --git a/Engine/Plugins/Experimental/ChaosClothAssetEditor/Source/ChaosClothAssetEditor/Private/ChaosClothAsset/ClothEditorPreviewScene.h b/Engine/Plugins/Experimental/ChaosClothAssetEditor/Source/ChaosClothAssetEditor/Private/ChaosClothAsset/ClothEditorPreviewScene.h
--- a/Engine/Plugins/Experimental/ChaosClothAssetEditor/Source/ChaosClothAssetEditor/Private/ChaosClothAsset/ClothEditorPreviewScene.h
+++ b/Engine/Plugins/Experimental/ChaosClothAssetEditor/Source/ChaosClothAssetEditor/Private/ChaosClothAsset/ClothEditorPreviewScene.h
@@ -60,13 +60,13 @@ public:
 	const UChaosClothPreviewSceneDescription* GetPreviewSceneDescription() const { return PreviewSceneDescription; }
 	UChaosClothPreviewSceneDescription* GetPreviewSceneDescription() { return PreviewSceneDescription; }
 
+	void CreateClothActor(UChaosClothAsset* Asset);
+
 	// Update Scene in response to the SceneDescription changing
 	void SceneDescriptionPropertyChanged(struct FPropertyChangedEvent& PropertyChangedEvent);
 
-	void CreateClothComponent(UChaosClothAsset* Asset);
-
-
 	// Preview simulation mesh
+	TObjectPtr<AActor> ClothActor;
 	TObjectPtr<UChaosClothComponent> ClothComponent;
 
 	// Skeletal Mesh
@@ -74,9 +74,11 @@ public:
 
 private:
 
-	TObjectPtr<UChaosClothPreviewSceneDescription> PreviewSceneDescription;
+	void SkeletalMeshTransformChanged(USceneComponent* UpdatedComponent, EUpdateTransformFlags UpdateTransformFlags, ETeleportType Teleport);
+
+	void CreateSkeletalMeshActor();
 
-	void InitializeSkeletalMeshActor();
+	TObjectPtr<UChaosClothPreviewSceneDescription> PreviewSceneDescription;
 };
 
 
