Cloth Editor: re-enable automatic selection of the cloth component in the EdMode's Mode Tools

diff --git a/Engine/Plugins/Experimental/ChaosClothAssetEditor/Source/ChaosClothAssetEditor/Private/ChaosClothAsset/ClothEditorMode.cpp b/Engine/Plugins/Experimental/ChaosClothAssetEditor/Source/ChaosClothAssetEditor/Private/ChaosClothAsset/ClothEditorMode.cpp
--- a/Engine/Plugins/Experimental/ChaosClothAssetEditor/Source/ChaosClothAssetEditor/Private/ChaosClothAsset/ClothEditorMode.cpp
+++ b/Engine/Plugins/Experimental/ChaosClothAssetEditor/Source/ChaosClothAssetEditor/Private/ChaosClothAsset/ClothEditorMode.cpp
@@ -591,6 +591,11 @@ void UChaosClothAssetEditorMode::InitializeTargets(const TArray<TObjectPtr<UObje
 		if (UChaosClothAsset* ChaosClothAsset = Cast<UChaosClothAsset>(Asset))
 		{
 			PreviewScene->CreateClothActor(ChaosClothAsset);
+
+			USelection* SelectedComponents = GetModeManager()->GetSelectedComponents();
+			SelectedComponents->Modify();
+			SelectedComponents->Select(PreviewScene->ClothComponent);
+
 			break;
 		}
 	}
