Cloth Editor: enable selection outline rendering in viewports

diff --git a/Engine/Plugins/Experimental/ChaosClothAssetEditor/Source/ChaosClothAssetEditor/ChaosClothAssetEditor.Build.cs b/Engine/Plugins/Experimental/ChaosClothAssetEditor/Source/ChaosClothAssetEditor/ChaosClothAssetEditor.Build.cs
--- a/Engine/Plugins/Experimental/ChaosClothAssetEditor/Source/ChaosClothAssetEditor/ChaosClothAssetEditor.Build.cs
+++ b/Engine/Plugins/Experimental/ChaosClothAssetEditor/Source/ChaosClothAssetEditor/ChaosClothAssetEditor.Build.cs
@@ -59,7 +59,9 @@ public class ChaosClothAssetEditor : ModuleRules
 				"DataflowEditor",
 				"DataflowEngine",
 				"DataflowEnginePlugin",
-				"RenderCore"
+				"RenderCore",
+				"TypedElementFramework",
+				"TypedElementRuntime"
 				// ... add private dependencies that you statically link with here ...	
 			}
 			);
diff --git a/Engine/Plugins/Experimental/ChaosClothAssetEditor/Source/ChaosClothAssetEditor/Private/ChaosClothAsset/ClothEditor3DViewportClient.cpp b/Engine/Plugins/Experimental/ChaosClothAssetEditor/Source/ChaosClothAssetEditor/Private/ChaosClothAsset/ClothEditor3DViewportClient.cpp
--- a/Engine/Plugins/Experimental/ChaosClothAssetEditor/Source/ChaosClothAssetEditor/Private/ChaosClothAsset/ClothEditor3DViewportClient.cpp
+++ b/Engine/Plugins/Experimental/ChaosClothAssetEditor/Source/ChaosClothAssetEditor/Private/ChaosClothAsset/ClothEditor3DViewportClient.cpp
@@ -15,6 +15,7 @@
 #include "BaseGizmos/TransformGizmoUtil.h"
 #include "Engine/Selection.h"
 #include "EngineUtils.h"
+#include "Animation/SkeletalMeshActor.h"
 
 FChaosClothAssetEditor3DViewportClient::FChaosClothAssetEditor3DViewportClient(FEditorModeTools* InModeTools,
 	FPreviewScene* InPreviewScene, 
@@ -27,6 +28,8 @@ FChaosClothAssetEditor3DViewportClient::FChaosClothAssetEditor3DViewportClient(F
 	// Call this once with the default value to get everything in a consistent state
 	EnableRenderMeshWireframe(bRenderMeshWireframe);
 
+	EngineShowFlags.SetSelectionOutline(true);
+
 	// Set up Gizmo and TransformProxy
 
 	UModeManagerInteractiveToolsContext* const InteractiveToolsContext = ModeTools->GetInteractiveToolsContext();
@@ -200,6 +203,10 @@ void FChaosClothAssetEditor3DViewportClient::ProcessClick(FSceneView& View, HHit
 	const bool bIsCtrlKeyDown = Viewport->KeyState(EKeys::LeftControl) || Viewport->KeyState(EKeys::RightControl);
 
 	USelection* SelectedComponents = ModeTools->GetSelectedComponents();
+
+	TArray<UPrimitiveComponent*> PreviouslySelectedComponents;
+	SelectedComponents->GetSelectedObjects<UPrimitiveComponent>(PreviouslySelectedComponents);
+
 	SelectedComponents->Modify();
 
 	if (!bIsShiftKeyDown && !bIsCtrlKeyDown)
@@ -240,10 +247,20 @@ void FChaosClothAssetEditor3DViewportClient::ProcessClick(FSceneView& View, HHit
 				{
 					SelectedComponents->Select(RootComponent);
 				}
+
+				if (UPrimitiveComponent* const PrimitiveComponent = Cast<UPrimitiveComponent>(RootComponent))
+				{
+					PrimitiveComponent->PushSelectionToProxy();
+				}
 			}
 		}
 	}
 
+	for (UPrimitiveComponent* const Component : PreviouslySelectedComponents)
+	{
+		Component->PushSelectionToProxy();
+	}
+
 	// Update TransformProxy
 
 	TransformProxy = NewObject<UTransformProxy>();
diff --git a/Engine/Plugins/Experimental/ChaosClothAssetEditor/Source/ChaosClothAssetEditor/Private/ChaosClothAsset/ClothEditorMode.cpp b/Engine/Plugins/Experimental/ChaosClothAssetEditor/Source/ChaosClothAssetEditor/Private/ChaosClothAsset/ClothEditorMode.cpp
--- a/Engine/Plugins/Experimental/ChaosClothAssetEditor/Source/ChaosClothAssetEditor/Private/ChaosClothAsset/ClothEditorMode.cpp
+++ b/Engine/Plugins/Experimental/ChaosClothAssetEditor/Source/ChaosClothAssetEditor/Private/ChaosClothAsset/ClothEditorMode.cpp
@@ -51,6 +51,7 @@
 #include "BoneWeights.h"
 #include "ToolSetupUtil.h"
 #include "Dataflow/DataflowComponent.h"
+#include "Elements/Framework/EngineElementsLibrary.h"
 
 #include UE_INLINE_GENERATED_CPP_BY_NAME(ClothEditorMode)
 
@@ -254,6 +255,12 @@ void UChaosClothAssetEditorMode::Exit()
 	USelection::SelectionChangedEvent.Remove(SelectionModifiedEventHandle);
 	UActorComponent::MarkRenderStateDirtyEvent.RemoveAll(this);
 
+	for (TObjectPtr<UDynamicMeshComponent> DynamicMeshComp : DynamicMeshComponents)
+	{
+		DynamicMeshComp->UnregisterComponent();
+		DynamicMeshComp->SelectionOverrideDelegate.Unbind();
+	}
+
 	for (TObjectPtr<UMeshElementsVisualizer> WireframeDisplay : WireframesToTick)
 	{
 		WireframeDisplay->Disconnect();
@@ -363,6 +370,24 @@ void UChaosClothAssetEditorMode::CreateToolTargets(const TArray<TObjectPtr<UObje
 }
 
 
+bool UChaosClothAssetEditorMode::IsComponentSelected(const UPrimitiveComponent* InComponent)
+{
+	if (const FEditorModeTools* const ModeManager = GetModeManager())
+	{
+		if (const UTypedElementSelectionSet* const TypedElementSelectionSet = ModeManager->GetEditorSelectionSet())
+		{
+			if (const FTypedElementHandle ComponentElement = UEngineElementsLibrary::AcquireEditorComponentElementHandle(InComponent))
+			{
+				const bool bElementSelected = TypedElementSelectionSet->IsElementSelected(ComponentElement, FTypedElementIsSelectedOptions());
+				return bElementSelected;
+			}
+		}
+	}
+
+	return false;
+}
+
+
 void UChaosClothAssetEditorMode::ReinitializeDynamicMeshComponents()
 {
 	// Clean up any existing DynamicMeshComponents
@@ -377,6 +402,7 @@ void UChaosClothAssetEditorMode::ReinitializeDynamicMeshComponents()
 
 		TObjectPtr<UDynamicMeshComponent> DynamicMeshComp = DynamicMeshComponents[DynamicMeshComponentIndex];
 		DynamicMeshComp->UnregisterComponent();
+		DynamicMeshComp->SelectionOverrideDelegate.Unbind();
 
 		if (SelectedComponents->IsSelected(DynamicMeshComp))
 		{
@@ -440,6 +466,8 @@ void UChaosClothAssetEditorMode::ReinitializeDynamicMeshComponents()
 				TObjectPtr<UDynamicMeshComponent> PatternMeshComponent = NewObject<UDynamicMeshComponent>(ParentActor);
 				PatternMeshComponent->SetMesh(MoveTemp(PatternMesh));
 
+				PatternMeshComponent->SelectionOverrideDelegate = UPrimitiveComponent::FSelectionOverride::CreateUObject(this, &UChaosClothAssetEditorMode::IsComponentSelected);
+
 				UMaterialInterface* Material = ToolSetupUtil::GetDefaultSculptMaterial(GetToolManager());
 				UMaterialInstanceDynamic* MatInstance = UMaterialInstanceDynamic::Create(Material, GetToolManager());
 				MatInstance->SetVectorParameterValue(TEXT("Color"), FLinearColor::Gray);
diff --git a/Engine/Plugins/Experimental/ChaosClothAssetEditor/Source/ChaosClothAssetEditor/Private/ChaosClothAsset/ClothEditorMode.h b/Engine/Plugins/Experimental/ChaosClothAssetEditor/Source/ChaosClothAssetEditor/Private/ChaosClothAsset/ClothEditorMode.h
--- a/Engine/Plugins/Experimental/ChaosClothAssetEditor/Source/ChaosClothAssetEditor/Private/ChaosClothAsset/ClothEditorMode.h
+++ b/Engine/Plugins/Experimental/ChaosClothAssetEditor/Source/ChaosClothAssetEditor/Private/ChaosClothAsset/ClothEditorMode.h
@@ -147,6 +147,8 @@ private:
 	// TODO: Expose this to the user
 	bool bCombineAllPatterns = false;
 
+	bool IsComponentSelected(const UPrimitiveComponent* InComponent);
+
 	// Create dynamic mesh components from the cloth component's rest space info
 	void ReinitializeDynamicMeshComponents();
 
diff --git a/Engine/Plugins/Experimental/ChaosClothAssetEditor/Source/ChaosClothAssetEditor/Private/ChaosClothAsset/ClothEditorPreviewScene.cpp b/Engine/Plugins/Experimental/ChaosClothAssetEditor/Source/ChaosClothAssetEditor/Private/ChaosClothAsset/ClothEditorPreviewScene.cpp
--- a/Engine/Plugins/Experimental/ChaosClothAssetEditor/Source/ChaosClothAssetEditor/Private/ChaosClothAsset/ClothEditorPreviewScene.cpp
+++ b/Engine/Plugins/Experimental/ChaosClothAssetEditor/Source/ChaosClothAssetEditor/Private/ChaosClothAsset/ClothEditorPreviewScene.cpp
@@ -4,6 +4,8 @@
 #include "ChaosClothAsset/ClothComponent.h"
 #include "Components/SkeletalMeshComponent.h"
 #include "Animation/SkeletalMeshActor.h"
+#include "AssetEditorModeManager.h"
+#include "Elements/Framework/EngineElementsLibrary.h"
 
 #define LOCTEXT_NAMESPACE "UChaosClothEditorPreviewScene"
 
@@ -41,6 +43,7 @@ FChaosClothPreviewScene::~FChaosClothPreviewScene()
 		if (SkeletalMeshActor->GetSkeletalMeshComponent())
 		{
 			SkeletalMeshActor->GetSkeletalMeshComponent()->TransformUpdated.RemoveAll(this);
+			SkeletalMeshActor->GetSkeletalMeshComponent()->SelectionOverrideDelegate.Unbind();
 		}
 
 		SkeletalMeshActor->UnregisterAllComponents();
@@ -48,6 +51,7 @@ FChaosClothPreviewScene::~FChaosClothPreviewScene()
 
 	if (ClothComponent)
 	{
+		ClothComponent->SelectionOverrideDelegate.Unbind();
 		ClothComponent->UnregisterComponent();
 	}
 }
@@ -81,6 +85,10 @@ void FChaosClothPreviewScene::SceneDescriptionPropertyChanged(struct FPropertyCh
 	}
 }
 
+void FChaosClothPreviewScene::SetModeManager(TSharedPtr<FAssetEditorModeManager> InClothPreviewEditorModeManager)
+{
+	ClothPreviewEditorModeManager = InClothPreviewEditorModeManager;
+}
 
 void FChaosClothPreviewScene::SkeletalMeshTransformChanged(USceneComponent* UpdatedComponent, EUpdateTransformFlags UpdateTransformFlags, ETeleportType Teleport)
 {
@@ -96,6 +104,7 @@ void FChaosClothPreviewScene::CreateSkeletalMeshActor()
 		if (SkeletalMeshActor->GetSkeletalMeshComponent())
 		{
 			SkeletalMeshActor->GetSkeletalMeshComponent()->TransformUpdated.RemoveAll(this);
+			SkeletalMeshActor->GetSkeletalMeshComponent()->SelectionOverrideDelegate.Unbind();
 		}
 
 		SkeletalMeshActor->UnregisterAllComponents();
@@ -112,11 +121,30 @@ void FChaosClothPreviewScene::CreateSkeletalMeshActor()
 	SkeletalMeshActor->RegisterAllComponents();
 
 	SkeletalMeshActor->GetSkeletalMeshComponent()->TransformUpdated.AddRaw(this, &FChaosClothPreviewScene::SkeletalMeshTransformChanged);
+	SkeletalMeshActor->GetSkeletalMeshComponent()->SelectionOverrideDelegate = UPrimitiveComponent::FSelectionOverride::CreateRaw(this, &FChaosClothPreviewScene::IsComponentSelected);
 }
 
+bool FChaosClothPreviewScene::IsComponentSelected(const UPrimitiveComponent* InComponent)
+{
+	if (const UTypedElementSelectionSet* const TypedElementSelectionSet = ClothPreviewEditorModeManager->GetEditorSelectionSet())
+	{
+		if (const FTypedElementHandle ComponentElement = UEngineElementsLibrary::AcquireEditorComponentElementHandle(InComponent))
+		{
+			const bool bElementSelected = TypedElementSelectionSet->IsElementSelected(ComponentElement, FTypedElementIsSelectedOptions());
+			return bElementSelected;
+		}
+	}
+
+	return false;
+}
 
 void FChaosClothPreviewScene::CreateClothActor(UChaosClothAsset* Asset)
 {
+	if (ClothComponent)
+	{
+		ClothComponent->SelectionOverrideDelegate.Unbind();
+	}
+
 	if (ClothActor)
 	{
 		ClothActor->UnregisterAllComponents();
@@ -129,6 +157,8 @@ void FChaosClothPreviewScene::CreateClothActor(UChaosClothAsset* Asset)
 	ClothComponent = NewObject<UChaosClothComponent>(ClothActor);
 	ClothComponent->SetClothAsset(Asset);
 
+	ClothComponent->SelectionOverrideDelegate = UPrimitiveComponent::FSelectionOverride::CreateRaw(this, &FChaosClothPreviewScene::IsComponentSelected);
+
 	ClothActor->SetRootComponent(ClothComponent);
 	ClothActor->RegisterAllComponents();
 }
diff --git a/Engine/Plugins/Experimental/ChaosClothAssetEditor/Source/ChaosClothAssetEditor/Private/ChaosClothAsset/ClothEditorPreviewScene.h b/Engine/Plugins/Experimental/ChaosClothAssetEditor/Source/ChaosClothAssetEditor/Private/ChaosClothAsset/ClothEditorPreviewScene.h
--- a/Engine/Plugins/Experimental/ChaosClothAssetEditor/Source/ChaosClothAssetEditor/Private/ChaosClothAsset/ClothEditorPreviewScene.h
+++ b/Engine/Plugins/Experimental/ChaosClothAssetEditor/Source/ChaosClothAssetEditor/Private/ChaosClothAsset/ClothEditorPreviewScene.h
@@ -13,6 +13,7 @@ class FChaosClothPreviewScene;
 class UPhysicsAsset;
 class UAnimSequence;
 class ASkeletalMeshActor;
+class FAssetEditorModeManager;
 
 ///
 /// The UChaosClothPreviewSceneDescription is a description of the Preview scene contents, intended to be editable in an FAdvancedPreviewSettingsWidget
@@ -72,13 +73,19 @@ public:
 	// Skeletal Mesh
 	TObjectPtr<ASkeletalMeshActor> SkeletalMeshActor;
 
+	void SetModeManager(TSharedPtr<FAssetEditorModeManager> InClothPreviewEditorModeManager);
+
 private:
 
 	void SkeletalMeshTransformChanged(USceneComponent* UpdatedComponent, EUpdateTransformFlags UpdateTransformFlags, ETeleportType Teleport);
 
 	void CreateSkeletalMeshActor();
 
+	bool IsComponentSelected(const UPrimitiveComponent* InComponent);
+
 	TObjectPtr<UChaosClothPreviewSceneDescription> PreviewSceneDescription;
+
+	TSharedPtr<FAssetEditorModeManager> ClothPreviewEditorModeManager;
 };
 
 
diff --git a/Engine/Plugins/Experimental/ChaosClothAssetEditor/Source/ChaosClothAssetEditor/Private/ChaosClothAsset/ClothEditorRestSpaceViewportClient.cpp b/Engine/Plugins/Experimental/ChaosClothAssetEditor/Source/ChaosClothAssetEditor/Private/ChaosClothAsset/ClothEditorRestSpaceViewportClient.cpp
--- a/Engine/Plugins/Experimental/ChaosClothAssetEditor/Source/ChaosClothAssetEditor/Private/ChaosClothAsset/ClothEditorRestSpaceViewportClient.cpp
+++ b/Engine/Plugins/Experimental/ChaosClothAssetEditor/Source/ChaosClothAssetEditor/Private/ChaosClothAsset/ClothEditorRestSpaceViewportClient.cpp
@@ -22,6 +22,8 @@ FChaosClothEditorRestSpaceViewportClient::FChaosClothEditorRestSpaceViewportClie
 
 	// Don't automatically switch to wireframe rendering in ortho mode
 	SetViewModes(EViewModeIndex::VMI_Lit, EViewModeIndex::VMI_Lit);
+
+	EngineShowFlags.SetSelectionOutline(true);
 }
 
 void FChaosClothEditorRestSpaceViewportClient::Set2DMode(bool In2DMode)
@@ -76,6 +78,9 @@ void FChaosClothEditorRestSpaceViewportClient::ProcessClick(FSceneView& View, HH
 	USelection* SelectedComponents = ModeTools->GetSelectedComponents();
 	SelectedComponents->Modify();
 	SelectedComponents->BeginBatchSelectOperation();
+
+	TArray<UDynamicMeshComponent*> PreviouslySelectedComponents;
+	SelectedComponents->GetSelectedObjects<UDynamicMeshComponent>(PreviouslySelectedComponents);
 	SelectedComponents->DeselectAll(UDynamicMeshComponent::StaticClass());
 
 	if (HitProxy && HitProxy->IsA(HActor::StaticGetType()))
@@ -90,6 +95,7 @@ void FChaosClothEditorRestSpaceViewportClient::ProcessClick(FSceneView& View, HH
 				if (UDynamicMeshComponent* DynamicMeshComp = Cast<UDynamicMeshComponent>(Component))
 				{
 					SelectedComponents->Select(DynamicMeshComp);
+					DynamicMeshComp->PushSelectionToProxy();
 				}
 			}
 		}
@@ -97,6 +103,11 @@ void FChaosClothEditorRestSpaceViewportClient::ProcessClick(FSceneView& View, HH
 
 	SelectedComponents->EndBatchSelectOperation();
 
+	for (UDynamicMeshComponent* Component : PreviouslySelectedComponents)
+	{
+		Component->PushSelectionToProxy();
+	}
+
 }
 
 void FChaosClothEditorRestSpaceViewportClient::SetEditorViewportWidget(TWeakPtr<SEditorViewport> InEditorViewportWidget)
diff --git a/Engine/Plugins/Experimental/ChaosClothAssetEditor/Source/ChaosClothAssetEditor/Private/ChaosClothAsset/ClothEditorToolkit.cpp b/Engine/Plugins/Experimental/ChaosClothAssetEditor/Source/ChaosClothAssetEditor/Private/ChaosClothAsset/ClothEditorToolkit.cpp
--- a/Engine/Plugins/Experimental/ChaosClothAssetEditor/Source/ChaosClothAssetEditor/Private/ChaosClothAsset/ClothEditorToolkit.cpp
+++ b/Engine/Plugins/Experimental/ChaosClothAssetEditor/Source/ChaosClothAssetEditor/Private/ChaosClothAsset/ClothEditorToolkit.cpp
@@ -171,6 +171,7 @@ FChaosClothAssetEditorToolkit::FChaosClothAssetEditorToolkit(UAssetEditor* InOwn
 	ClothPreviewScene->SetFloorVisibility(false, true);
 	ClothPreviewEditorModeManager = MakeShared<FAssetEditorModeManager>();
 	ClothPreviewEditorModeManager->SetPreviewScene(ClothPreviewScene.Get());
+	ClothPreviewScene->SetModeManager(ClothPreviewEditorModeManager);
 
 	//ClothPreviewInputRouter = ClothPreviewEditorModeManager->GetInteractiveToolsContext()->InputRouter;
 	ClothPreviewTabContent = MakeShareable(new FEditorViewportTabContent());
