ChaosClothAsset: Switching from FAssetTypeActions to UAssetDefinition

diff --git a/Engine/Plugins/Experimental/ChaosClothAssetEditor/Source/ChaosClothAssetEditor/ChaosClothAssetEditor.Build.cs b/Engine/Plugins/Experimental/ChaosClothAssetEditor/Source/ChaosClothAssetEditor/ChaosClothAssetEditor.Build.cs
--- a/Engine/Plugins/Experimental/ChaosClothAssetEditor/Source/ChaosClothAssetEditor/ChaosClothAssetEditor.Build.cs
+++ b/Engine/Plugins/Experimental/ChaosClothAssetEditor/Source/ChaosClothAssetEditor/ChaosClothAssetEditor.Build.cs
@@ -61,7 +61,8 @@ public class ChaosClothAssetEditor : ModuleRules
 				"DataflowEnginePlugin",
 				"RenderCore",
 				"TypedElementFramework",
-				"TypedElementRuntime"
+				"TypedElementRuntime",
+				"AssetDefinition",
 				// ... add private dependencies that you statically link with here ...	
 			}
 			);
diff --git a/Engine/Plugins/Experimental/ChaosClothAssetEditor/Source/ChaosClothAssetEditor/Private/ChaosClothAsset/AssetDefinition_ClothAsset.cpp b/Engine/Plugins/Experimental/ChaosClothAssetEditor/Source/ChaosClothAssetEditor/Private/ChaosClothAsset/AssetDefinition_ClothAsset.cpp
new file mode 100644
--- /dev/null
+++ b/Engine/Plugins/Experimental/ChaosClothAssetEditor/Source/ChaosClothAssetEditor/Private/ChaosClothAsset/AssetDefinition_ClothAsset.cpp
@@ -0,0 +1,137 @@
+// Copyright Epic Games, Inc. All Rights Reserved.
+
+#include "AssetDefinition_ClothAsset.h"
+#include "ChaosClothAsset/ClothAsset.h"
+#include "ChaosClothAsset/ClothEditor.h"
+#include "Dataflow/DataflowObject.h"
+#include "Misc/MessageDialog.h"
+#include "Misc/FileHelper.h"
+#include "IContentBrowserSingleton.h"
+#include "ContentBrowserModule.h"
+#include "AssetRegistry/AssetRegistryModule.h"
+#include "FileHelpers.h"
+#include "ThumbnailRendering/SceneThumbnailInfo.h"
+#include "Misc/WarnIfAssetsLoadedInScope.h"
+
+#define LOCTEXT_NAMESPACE "AssetDefinition_ClothAsset"
+
+namespace ClothAssetDefinitionHelpers
+{
+	// Create a new UDataflow if one doesn't already exist for the Cloth Asset
+	UObject* CreateNewDataflowAsset(const UChaosClothAsset* ClothAsset)
+	{
+		if (FMessageDialog::Open(EAppMsgType::YesNo, LOCTEXT("MissingDataflow", "This Cloth asset currently has no Dataflow graph. Would you like to create a new one?")) == EAppReturnType::Yes)
+		{
+			const UClass* const DataflowClass = UDataflow::StaticClass();
+
+			FSaveAssetDialogConfig NewDataflowAssetDialogConfig;
+			{
+				const FString PackageName = ClothAsset->GetOutermost()->GetName();
+				NewDataflowAssetDialogConfig.DefaultPath = FPackageName::GetLongPackagePath(PackageName);
+				NewDataflowAssetDialogConfig.AssetClassNames.Add(DataflowClass->GetClassPathName());
+				NewDataflowAssetDialogConfig.ExistingAssetPolicy = ESaveAssetDialogExistingAssetPolicy::AllowButWarn;
+				NewDataflowAssetDialogConfig.DialogTitleOverride = LOCTEXT("NewDataflowAssetDialogTitle", "Save Dataflow Asset As");
+			}
+
+			FContentBrowserModule& ContentBrowserModule = FModuleManager::LoadModuleChecked<FContentBrowserModule>("ContentBrowser");
+
+			FString NewPackageName;
+			FText OutError;
+			for (bool bFilenameValid = false; !bFilenameValid; bFilenameValid = FFileHelper::IsFilenameValidForSaving(NewPackageName, OutError))
+			{
+				const FString AssetSavePath = ContentBrowserModule.Get().CreateModalSaveAssetDialog(NewDataflowAssetDialogConfig);
+				if (AssetSavePath.IsEmpty())
+				{
+					return nullptr;
+				}
+				NewPackageName = FPackageName::ObjectPathToPackageName(AssetSavePath);
+			}
+
+			const FName NewAssetName(FPackageName::GetLongPackageAssetName(NewPackageName));
+			UPackage* const NewPackage = CreatePackage(*NewPackageName);
+			UObject* const NewAsset = NewObject<UObject>(NewPackage, DataflowClass, NewAssetName, RF_Public | RF_Standalone | RF_Transactional);
+
+			NewAsset->MarkPackageDirty();
+
+			// Notify the asset registry
+			FAssetRegistryModule::AssetCreated(NewAsset);
+
+			// Save the package
+			TArray<UPackage*> PackagesToSave;
+			PackagesToSave.Add(NewAsset->GetOutermost());
+			constexpr bool bCheckDirty = false;
+			constexpr bool bPromptToSave = false;
+			FEditorFileUtils::PromptForCheckoutAndSave(PackagesToSave, bCheckDirty, bPromptToSave);
+
+			return NewAsset;
+		}
+
+		return nullptr;
+	}
+}
+
+
+FText UAssetDefinition_ClothAsset::GetAssetDisplayName() const
+{
+	return NSLOCTEXT("AssetTypeActions", "AssetTypeActions_ClothAsset", "ClothAsset");
+}
+
+TSoftClassPtr<UObject> UAssetDefinition_ClothAsset::GetAssetClass() const
+{
+	return UChaosClothAsset::StaticClass();
+}
+
+FLinearColor UAssetDefinition_ClothAsset::GetAssetColor() const
+{
+	return FLinearColor(FColor(180, 120, 110));
+}
+
+TConstArrayView<FAssetCategoryPath> UAssetDefinition_ClothAsset::GetAssetCategories() const
+{
+	static const auto Categories = { EAssetCategoryPaths::Physics };
+	return Categories;
+}
+
+UThumbnailInfo* UAssetDefinition_ClothAsset::LoadThumbnailInfo(const FAssetData& InAsset) const
+{
+	return UE::Editor::FindOrCreateThumbnailInfo(InAsset.GetAsset(), USceneThumbnailInfo::StaticClass());
+}
+
+EAssetCommandResult UAssetDefinition_ClothAsset::OpenAssets(const FAssetOpenArgs& OpenArgs) const
+{
+	TArray<UChaosClothAsset*> ClothObjects = OpenArgs.LoadObjects<UChaosClothAsset>();
+
+	// For now the cloth editor only works on one asset at a time
+	ensure(ClothObjects.Num() == 0 || ClothObjects.Num() == 1);
+
+	if (ClothObjects.Num() > 0)
+	{
+		UAssetEditorSubsystem* const AssetEditorSubsystem = GEditor->GetEditorSubsystem<UAssetEditorSubsystem>();
+		UChaosClothAssetEditor* const AssetEditor = NewObject<UChaosClothAssetEditor>(AssetEditorSubsystem, NAME_None, RF_Transient);
+
+		// Validate the asset
+		UChaosClothAsset* const ClothAsset = CastChecked<UChaosClothAsset>(ClothObjects[0]);
+
+		if (!ClothAsset->DataflowAsset)
+		{
+			if (UObject* const NewDataflowAsset = ClothAssetDefinitionHelpers::CreateNewDataflowAsset(ClothAsset))
+			{
+				ClothAsset->DataflowAsset = CastChecked<UDataflow>(NewDataflowAsset);
+			}
+		}
+
+		TArray<TObjectPtr<UObject>> Objects;
+		for (UChaosClothAsset* const ClothObject : ClothObjects)
+		{
+			Objects.Add(ClothObject);
+		}
+		AssetEditor->Initialize(Objects);
+
+		return EAssetCommandResult::Handled;
+	}
+
+	return EAssetCommandResult::Unhandled;
+}
+
+
+#undef LOCTEXT_NAMESPACE
diff --git a/Engine/Plugins/Experimental/ChaosClothAssetEditor/Source/ChaosClothAssetEditor/Private/ChaosClothAsset/AssetDefinition_ClothAsset.h b/Engine/Plugins/Experimental/ChaosClothAssetEditor/Source/ChaosClothAssetEditor/Private/ChaosClothAsset/AssetDefinition_ClothAsset.h
new file mode 100644
--- /dev/null
+++ b/Engine/Plugins/Experimental/ChaosClothAssetEditor/Source/ChaosClothAssetEditor/Private/ChaosClothAsset/AssetDefinition_ClothAsset.h
@@ -0,0 +1,23 @@
+// Copyright Epic Games, Inc. All Rights Reserved.
+
+#pragma once
+
+#include "AssetDefinitionDefault.h"
+#include "AssetDefinition_ClothAsset.generated.h"
+
+UCLASS()
+class UAssetDefinition_ClothAsset : public UAssetDefinitionDefault
+{
+	GENERATED_BODY()
+
+private:
+
+	virtual FText GetAssetDisplayName() const override;
+	virtual TSoftClassPtr<UObject> GetAssetClass() const override;
+	virtual FLinearColor GetAssetColor() const override;
+	virtual TConstArrayView<FAssetCategoryPath> GetAssetCategories() const override;
+	virtual UThumbnailInfo* LoadThumbnailInfo(const FAssetData& InAssetData) const override;
+
+	virtual EAssetCommandResult OpenAssets(const FAssetOpenArgs& OpenArgs) const override;
+};
+
diff --git a/Engine/Plugins/Experimental/ChaosClothAssetEditor/Source/ChaosClothAssetEditor/Private/ChaosClothAsset/AssetTypeActions_ClothAsset.cpp b/Engine/Plugins/Experimental/ChaosClothAssetEditor/Source/ChaosClothAssetEditor/Private/ChaosClothAsset/AssetTypeActions_ClothAsset.cpp
deleted file mode 100644
diff --git a/Engine/Plugins/Experimental/ChaosClothAssetEditor/Source/ChaosClothAssetEditor/Private/ChaosClothAsset/AssetTypeActions_ClothAsset.h b/Engine/Plugins/Experimental/ChaosClothAssetEditor/Source/ChaosClothAssetEditor/Private/ChaosClothAsset/AssetTypeActions_ClothAsset.h
deleted file mode 100644
diff --git a/Engine/Plugins/Experimental/ChaosClothAssetEditor/Source/ChaosClothAssetEditor/Private/ChaosClothAsset/ChaosClothAssetEditorModule.cpp b/Engine/Plugins/Experimental/ChaosClothAssetEditor/Source/ChaosClothAssetEditor/Private/ChaosClothAsset/ChaosClothAssetEditorModule.cpp
--- a/Engine/Plugins/Experimental/ChaosClothAssetEditor/Source/ChaosClothAssetEditor/Private/ChaosClothAsset/ChaosClothAssetEditorModule.cpp
+++ b/Engine/Plugins/Experimental/ChaosClothAssetEditor/Source/ChaosClothAssetEditor/Private/ChaosClothAsset/ChaosClothAssetEditorModule.cpp
@@ -12,7 +12,6 @@
 #include "Editor.h"
 #include "Styling/AppStyle.h"
 #include "AssetToolsModule.h"
-#include "AssetTypeActions_ClothAsset.h"
 
 #define LOCTEXT_NAMESPACE "FChaosClothAssetEditorModule"
 
@@ -28,9 +27,6 @@ void FChaosClothAssetEditorModule::StartupModule()
 	FAssetToolsModule& AssetToolsModule = FAssetToolsModule::GetModule();
 	IAssetTools& AssetTools = AssetToolsModule.Get();
 
-	AssetTypeActions_ClothAsset = new FAssetTypeActions_ClothAsset();
-	AssetTools.RegisterAssetTypeActions(MakeShareable(AssetTypeActions_ClothAsset));
-
 	Dataflow::RegisterClothDataflowNodes();
 
 	// TODO: Register details view customizations
@@ -42,14 +38,6 @@ void FChaosClothAssetEditorModule::ShutdownModule()
 
 	FEditorModeRegistry::Get().UnregisterMode(UChaosClothAssetEditorMode::EM_ChaosClothAssetEditorModeId);
 
-	if (UObjectInitialized())
-	{
-		// Unregister asset actions
-		FAssetToolsModule& AssetToolsModule = FAssetToolsModule::GetModule();
-		IAssetTools& AssetTools = AssetToolsModule.Get();
-		AssetTools.UnregisterAssetTypeActions(AssetTypeActions_ClothAsset->AsShared());
-	}
-
 	// TODO: Unregister details view customizations
 }
 
