Cloth Editor: enable tools to be run in either viewport. There already exist separate UWorlds and EditorModeManagers for each viewport, this change just introduces bookkeeping to allow the EdMode to specify which viewport to use on a per-tool basis.

diff --git a/Engine/Plugins/Experimental/ChaosClothAssetEditor/Source/ChaosClothAssetEditor/Private/ChaosClothAsset/ClothEditorMode.cpp b/Engine/Plugins/Experimental/ChaosClothAssetEditor/Source/ChaosClothAssetEditor/Private/ChaosClothAsset/ClothEditorMode.cpp
--- a/Engine/Plugins/Experimental/ChaosClothAssetEditor/Source/ChaosClothAssetEditor/Private/ChaosClothAsset/ClothEditorMode.cpp
+++ b/Engine/Plugins/Experimental/ChaosClothAssetEditor/Source/ChaosClothAssetEditor/Private/ChaosClothAsset/ClothEditorMode.cpp
@@ -52,6 +52,7 @@
 #include "ToolSetupUtil.h"
 #include "Dataflow/DataflowComponent.h"
 #include "Elements/Framework/EngineElementsLibrary.h"
+#include "ToolTargets/SkeletalMeshComponentToolTarget.h"
 
 #include UE_INLINE_GENERATED_CPP_BY_NAME(ClothEditorMode)
 
@@ -167,6 +168,54 @@ void UChaosClothAssetEditorMode::AddToolTargetFactories()
 	GetInteractiveToolsContext()->TargetManager->AddTargetFactory(NewObject<UClothComponentToolTargetFactory>(GetToolManager()));
 }
 
+
+void UChaosClothAssetEditorMode::RegisterClothTool(TSharedPtr<FUICommandInfo> UICommand, FString ToolIdentifier, UInteractiveToolBuilder* Builder, UEditorInteractiveToolsContext* const ToolsContext, EToolsContextScope ToolScope)
+{
+	if (!Toolkit.IsValid())
+	{
+		return;
+	}
+
+	if (!ToolsContext)
+	{
+		return;
+	}
+
+	if (ToolScope == EToolsContextScope::Default)
+	{
+		ToolScope = GetDefaultToolScope();
+	}
+	ensure(ToolScope != EToolsContextScope::Editor);
+
+	ToolsContext->ToolManager->RegisterToolType(ToolIdentifier, Builder);
+
+	const TSharedRef<FUICommandList>& CommandList = Toolkit->GetToolkitCommands();
+
+	CommandList->MapAction(UICommand,
+		FExecuteAction::CreateWeakLambda(ToolsContext, [this, ToolsContext, ToolIdentifier]()
+		{
+			ActiveToolsContext = ToolsContext;
+			ToolsContext->StartTool(ToolIdentifier);
+		}),
+		FCanExecuteAction::CreateWeakLambda(ToolsContext, [this, ToolIdentifier, ToolsContext]() 
+		{
+			return ShouldToolStartBeAllowed(ToolIdentifier) &&
+			ToolsContext->ToolManager->CanActivateTool(EToolSide::Mouse, ToolIdentifier);
+		}),
+		FIsActionChecked::CreateUObject(ToolsContext, &UEdModeInteractiveToolsContext::IsToolActive, EToolSide::Mouse, ToolIdentifier),
+		EUIActionRepeatMode::RepeatDisabled);
+
+}
+
+void UChaosClothAssetEditorMode::RegisterPreviewTools()
+{
+	UEditorInteractiveToolsContext* const PreviewSceneToolsContext = PreviewScene->GetClothPreviewEditorModeManager()->GetInteractiveToolsContext();
+
+	const FChaosClothAssetEditorCommands& CommandInfos = FChaosClothAssetEditorCommands::Get();
+	RegisterClothTool(CommandInfos.BeginClothTrainingTool, FChaosClothAssetEditorCommands::BeginClothTrainingToolIdentifier, NewObject<UClothTrainingToolBuilder>(), PreviewSceneToolsContext);
+	RegisterClothTool(CommandInfos.BeginTransferSkinWeightsTool, FChaosClothAssetEditorCommands::BeginTransferSkinWeightsToolIdentifier, NewObject<UClothTransferSkinWeightsToolBuilder>(), PreviewSceneToolsContext);
+}
+
 void UChaosClothAssetEditorMode::RegisterTools()
 {
 	const FChaosClothAssetEditorCommands& CommandInfos = FChaosClothAssetEditorCommands::Get();
@@ -178,20 +227,28 @@ void UChaosClothAssetEditorMode::RegisterTools()
 	// TODO: Re-add the remesh tool when we have a way to remesh both 2d and 3d rest space meshes at the same time
 	//RegisterTool(CommandInfos.BeginRemeshTool, FChaosClothAssetEditorCommands::BeginRemeshToolIdentifier, NewObject<URemeshMeshToolBuilder>());
 
-	RegisterTool(CommandInfos.BeginAttributeEditorTool, FChaosClothAssetEditorCommands::BeginAttributeEditorToolIdentifier, NewObject<UAttributeEditorToolBuilder>());
-	RegisterTool(CommandInfos.BeginWeightMapPaintTool, FChaosClothAssetEditorCommands::BeginWeightMapPaintToolIdentifier, NewObject<UClothEditorWeightMapPaintToolBuilder>());
-	RegisterTool(CommandInfos.BeginClothTrainingTool, FChaosClothAssetEditorCommands::BeginClothTrainingToolIdentifier, NewObject<UClothTrainingToolBuilder>());
-	RegisterTool(CommandInfos.BeginTransferSkinWeightsTool, FChaosClothAssetEditorCommands::BeginTransferSkinWeightsToolIdentifier, NewObject<UClothTransferSkinWeightsToolBuilder>());
+	UEditorInteractiveToolsContext* const RestSpaceToolsContext = GetInteractiveToolsContext();
+
+	RegisterClothTool(CommandInfos.BeginAttributeEditorTool, FChaosClothAssetEditorCommands::BeginAttributeEditorToolIdentifier, NewObject<UAttributeEditorToolBuilder>(), RestSpaceToolsContext);
+	RegisterClothTool(CommandInfos.BeginWeightMapPaintTool, FChaosClothAssetEditorCommands::BeginWeightMapPaintToolIdentifier, NewObject<UClothEditorWeightMapPaintToolBuilder>(), RestSpaceToolsContext);
 }
 
 bool UChaosClothAssetEditorMode::ShouldToolStartBeAllowed(const FString& ToolIdentifier) const
 {
-	// For now we've decided to disallow switch-away on accept/cancel tools in the Cloth editor.
-	if (GetInteractiveToolsContext()->ActiveToolHasAccept())
+	// For now we've decided to disallow switch-away on any tools in the Cloth editor.
+	if (GetInteractiveToolsContext()->HasActiveTool())
 	{
 		return false;
 	}
 	
+	if (PreviewScene && PreviewScene->GetClothPreviewEditorModeManager() && PreviewScene->GetClothPreviewEditorModeManager()->GetInteractiveToolsContext())
+	{
+		if (PreviewScene->GetClothPreviewEditorModeManager()->GetInteractiveToolsContext()->HasActiveTool())
+		{
+			return false;
+		}
+	}
+
 	return Super::ShouldToolStartBeAllowed(ToolIdentifier);
 }
 
@@ -212,6 +269,7 @@ void UChaosClothAssetEditorMode::OnToolEnded(UInteractiveToolManager* Manager, U
 	FChaosClothAssetEditorCommands::UpdateToolCommandBinding(Tool, ToolCommandList, true);
 
 	bCanTogglePattern2DMode = true;
+	ActiveToolsContext = nullptr;
 
 	UpdateSimulationMeshes();
 	ReinitializeDynamicMeshComponents();
@@ -276,6 +334,26 @@ void UChaosClothAssetEditorMode::Exit()
 void UChaosClothAssetEditorMode::SetPreviewScene(FChaosClothPreviewScene* InPreviewScene)
 {
 	PreviewScene = InPreviewScene;
+
+
+	UEditorInteractiveToolsContext* const PreviewToolsContext = PreviewScene->GetClothPreviewEditorModeManager()->GetInteractiveToolsContext();
+	UInteractiveToolManager* const PreviewToolManager = PreviewToolsContext->ToolManager;
+	PreviewToolsContext->TargetManager->AddTargetFactory(NewObject<UClothComponentToolTargetFactory>(PreviewToolManager));
+	PreviewToolsContext->TargetManager->AddTargetFactory(NewObject<USkeletalMeshComponentToolTargetFactory>(PreviewToolManager));
+
+	PreviewToolManager->OnToolStarted.AddUObject(this, &UChaosClothAssetEditorMode::OnToolStarted);
+	PreviewToolManager->OnToolEnded.AddUObject(this, &UChaosClothAssetEditorMode::OnToolEnded);
+
+	check(Toolkit.IsValid());
+	
+	// FBaseToolkit's OnToolStarted and OnToolEnded are protected, so we use the subclass to get at them
+	FChaosClothAssetEditorModeToolkit* const ClothModeToolkit = static_cast<FChaosClothAssetEditorModeToolkit*>(Toolkit.Get());	
+
+	PreviewToolManager->OnToolStarted.AddSP(ClothModeToolkit, &FChaosClothAssetEditorModeToolkit::OnToolStarted);
+	PreviewToolManager->OnToolEnded.AddSP(ClothModeToolkit, &FChaosClothAssetEditorModeToolkit::OnToolEnded);
+
+	// TODO: It would be nice if the PreviewScene could be specified before the UBaseCharacterFXEditorMode::Enter() function is called. Then we could register both sets of tools from there.
+	RegisterPreviewTools();
 }
 
 void UChaosClothAssetEditorMode::UpdateSimulationMeshes()
diff --git a/Engine/Plugins/Experimental/ChaosClothAssetEditor/Source/ChaosClothAssetEditor/Private/ChaosClothAsset/ClothEditorMode.h b/Engine/Plugins/Experimental/ChaosClothAssetEditor/Source/ChaosClothAssetEditor/Private/ChaosClothAsset/ClothEditorMode.h
--- a/Engine/Plugins/Experimental/ChaosClothAssetEditor/Source/ChaosClothAssetEditor/Private/ChaosClothAsset/ClothEditorMode.h
+++ b/Engine/Plugins/Experimental/ChaosClothAssetEditor/Source/ChaosClothAssetEditor/Private/ChaosClothAsset/ClothEditorMode.h
@@ -28,6 +28,7 @@ class FChaosClothEditorRestSpaceViewportClient;
 class FViewport;
 class UDataflowComponent;
 class FChaosClothPreviewScene;
+class FChaosClothAssetEditor3DViewportClient;
 
 /**
  * The cloth editor mode is the mode used in the cloth asset editor. It holds most of the inter-tool state.
@@ -64,9 +65,15 @@ public:
 
 	UDataflowComponent* GetDataflowComponent() const;
 
+	TObjectPtr<UEditorInteractiveToolsContext> GetActiveToolsContext()
+	{
+		return ActiveToolsContext;
+	}
+
 private:
 
 	friend class FChaosClothAssetEditorToolkit;
+	friend class FChaosClothAssetEditorModeToolkit;
 
 	// UEdMode
 	virtual void Enter() final;
@@ -92,6 +99,16 @@ private:
 	// to turn them into the input objects that tools get (since these need preview meshes, etc).
 	static const FToolTargetTypeRequirements& GetToolTargetRequirements();
 
+	// Use this function to register tools rather than UEdMode::RegisterTool() because we need to specify the ToolsContext
+	void RegisterClothTool(TSharedPtr<FUICommandInfo> UICommand, 
+		FString ToolIdentifier, 
+		UInteractiveToolBuilder* Builder, 
+		UEditorInteractiveToolsContext* UseToolsContext, 
+		EToolsContextScope ToolScope = EToolsContextScope::Default);
+
+	// Register the set of tools that operate on objects in the 3D preview world 
+	void RegisterPreviewTools();
+
 	void SetPreviewScene(FChaosClothPreviewScene* PreviewScene);
 
 	// Bounding box for rest space meshes
@@ -162,5 +179,8 @@ private:
 
 	UPROPERTY()
 	TObjectPtr<UDataflowComponent> DataflowComponent = nullptr;
+
+	UPROPERTY()
+	TObjectPtr<UEditorInteractiveToolsContext> ActiveToolsContext = nullptr;
 };
 
diff --git a/Engine/Plugins/Experimental/ChaosClothAssetEditor/Source/ChaosClothAssetEditor/Private/ChaosClothAsset/ClothEditorModeToolkit.cpp b/Engine/Plugins/Experimental/ChaosClothAssetEditor/Source/ChaosClothAssetEditor/Private/ChaosClothAsset/ClothEditorModeToolkit.cpp
--- a/Engine/Plugins/Experimental/ChaosClothAssetEditor/Source/ChaosClothAssetEditor/Private/ChaosClothAsset/ClothEditorModeToolkit.cpp
+++ b/Engine/Plugins/Experimental/ChaosClothAssetEditor/Source/ChaosClothAssetEditor/Private/ChaosClothAsset/ClothEditorModeToolkit.cpp
@@ -3,8 +3,24 @@
 #include "ChaosClothAsset/ClothEditorModeToolkit.h"
 #include "ChaosClothAsset/ClothEditorCommands.h"
 #include "ChaosClothAsset/ClothEditorStyle.h"
-#include "Framework/MultiBox/MultiBoxBuilder.h"
+#include "ChaosClothAsset/ClothEditorMode.h"
+#include "ChaosClothAsset/ClothEditorPreviewScene.h"
+#include "ChaosClothAsset/SClothEditorRestSpaceViewport.h"
+#include "ChaosClothAsset/ClothEditorRestSpaceViewportClient.h"
+#include "ChaosClothAsset/SClothEditor3DViewport.h"
+#include "ChaosClothAsset/ClothEditor3DViewportClient.h"
 
+#include "AssetEditorModeManager.h"
+#include "Framework/MultiBox/MultiBoxBuilder.h"
+#include "InteractiveToolManager.h"
+#include "SEditorViewport.h"
+#include "EdModeInteractiveToolsContext.h"
+#include "SBaseCharacterFXEditorViewport.h"
+#include "Widgets/Input/SButton.h"
+#include "SPrimaryButton.h"
+#include "IDetailsView.h"
+#include "Widgets/Images/SImage.h"
+	
 #define LOCTEXT_NAMESPACE "FChaosClothAssetEditorModeToolkit"
 
 FName FChaosClothAssetEditorModeToolkit::GetToolkitFName() const
@@ -17,6 +33,110 @@ FText FChaosClothAssetEditorModeToolkit::GetBaseToolkitName() const
 	return NSLOCTEXT("ChaosClothAssetEditorModeToolkit", "DisplayName", "ChaosClothAssetEditorMode");
 }
 
+
+void FChaosClothAssetEditorModeToolkit::Init(const TSharedPtr<IToolkitHost>& InitToolkitHost, TWeakObjectPtr<UEdMode> InOwningMode)
+{
+	FModeToolkit::Init(InitToolkitHost, InOwningMode);	
+
+	// Set up tool message areas
+	ClearNotification();
+	ClearWarning();
+	GetScriptableEditorMode()->GetInteractiveToolsContext()->OnToolNotificationMessage.AddSP(this, &FBaseCharacterFXEditorModeToolkit::PostNotification);
+	GetScriptableEditorMode()->GetInteractiveToolsContext()->OnToolWarningMessage.AddSP(this, &FBaseCharacterFXEditorModeToolkit::PostWarning);
+
+	// Hook up the tool detail panel
+	ToolDetailsContainer->SetContent(DetailsView.ToSharedRef());
+	
+	// Set up the overlay -- it will poll different ToolsContexts depending on which one launched the currently active tool
+
+	SAssignNew(ViewportOverlayWidget, SHorizontalBox)
+
+	+SHorizontalBox::Slot()
+	.HAlign(HAlign_Center)
+	.VAlign(VAlign_Bottom)
+	.Padding(FMargin(0.0f, 0.0f, 0.f, 15.f))
+	[
+		SNew(SBorder)
+		.BorderImage(FAppStyle::Get().GetBrush("EditorViewport.OverlayBrush"))
+		.Padding(8.f)
+		[
+			SNew(SHorizontalBox)
+
+			+SHorizontalBox::Slot()
+			.AutoWidth()
+			.VAlign(VAlign_Center)
+			.Padding(FMargin(0.f, 0.f, 8.f, 0.f))
+			[
+				SNew(SImage)
+				.Image_Lambda([this] () { return ActiveToolIcon; })
+			]
+
+			+SHorizontalBox::Slot()
+			.AutoWidth()
+			.VAlign(VAlign_Center)
+			.Padding(FMargin(0.f, 0.f, 8.f, 0.f))
+			[
+				SNew(STextBlock)
+				.Text(this, &FBaseCharacterFXEditorModeToolkit::GetActiveToolDisplayName)
+			]
+
+			+SHorizontalBox::Slot()
+			.AutoWidth()
+			.Padding(FMargin(0.0, 0.f, 2.f, 0.f))
+			[
+				SNew(SPrimaryButton)
+				.Text(LOCTEXT("OverlayAccept", "Accept"))
+				.ToolTipText(LOCTEXT("OverlayAcceptTooltip", "Accept/Commit the results of the active Tool [Enter]"))
+				.OnClicked_Lambda([this]() { 
+				    GetCurrentToolsContext()->EndTool(EToolShutdownType::Accept);
+					return FReply::Handled(); 
+					})
+				.IsEnabled_Lambda([this]() { return GetCurrentToolsContext()->CanAcceptActiveTool(); })
+				.Visibility_Lambda([this]() { return GetCurrentToolsContext()->ActiveToolHasAccept() ? EVisibility::Visible : EVisibility::Collapsed; })
+			]
+
+			+SHorizontalBox::Slot()
+			.AutoWidth()
+			.Padding(FMargin(2.0, 0.f, 0.f, 0.f))
+			[
+				SNew(SButton)
+				.ButtonStyle(FAppStyle::Get(), "Button")
+				.TextStyle( FAppStyle::Get(), "DialogButtonText" )
+				.Text(LOCTEXT("OverlayCancel", "Cancel"))
+				.ToolTipText(LOCTEXT("OverlayCancelTooltip", "Cancel the active Tool [Esc]"))
+				.HAlign(HAlign_Center)
+				.OnClicked_Lambda([this]() { 
+					GetCurrentToolsContext()->EndTool(EToolShutdownType::Cancel);
+					return FReply::Handled(); 
+					})
+				.IsEnabled_Lambda([this]() { return GetCurrentToolsContext()->CanCancelActiveTool(); })
+				.Visibility_Lambda([this]() { return GetCurrentToolsContext()->ActiveToolHasAccept() ? EVisibility::Visible : EVisibility::Collapsed; })
+			]
+
+			+ SHorizontalBox::Slot()
+			.AutoWidth()
+			.Padding(FMargin(2.0, 0.f, 0.f, 0.f))
+			[
+				SNew(SButton)
+				.ButtonStyle(FAppStyle::Get(), "PrimaryButton")
+				.TextStyle(FAppStyle::Get(), "DialogButtonText")
+				.Text(LOCTEXT("OverlayComplete", "Complete"))
+				.ToolTipText(LOCTEXT("OverlayCompleteTooltip", "Exit the active Tool [Enter]"))
+				.HAlign(HAlign_Center)
+				.OnClicked_Lambda([this]() { 
+					GetCurrentToolsContext()->EndTool(EToolShutdownType::Completed);
+					return FReply::Handled(); 
+					})
+				.IsEnabled_Lambda([this]() {
+					return GetCurrentToolsContext()->CanCompleteActiveTool();
+				})
+				.Visibility_Lambda([this]() { return GetCurrentToolsContext()->CanCompleteActiveTool() ? EVisibility::Visible : EVisibility::Collapsed; })
+			]
+		]	
+	];
+}
+
+
 void FChaosClothAssetEditorModeToolkit::BuildToolPalette(FName PaletteIndex, class FToolBarBuilder& ToolbarBuilder)
 {
 	const FChaosClothAssetEditorCommands& Commands = FChaosClothAssetEditorCommands::Get();
@@ -40,4 +160,119 @@ const FSlateBrush* FChaosClothAssetEditorModeToolkit::GetActiveToolIcon(const FS
 	return FChaosClothAssetEditorStyle::Get().GetOptionalBrush(ActiveToolIconName);
 }
 
+SBaseCharacterFXEditorViewport* FChaosClothAssetEditorModeToolkit::GetViewportWidgetForManager(UInteractiveToolManager* Manager)
+{
+	if (OwningEditorMode.IsValid(false))
+	{
+		const UEdMode* const Mode = OwningEditorMode.Get();
+
+		if (const UChaosClothAssetEditorMode* const ClothEdMode = Cast<UChaosClothAssetEditorMode>(Mode))
+		{
+			const FChaosClothPreviewScene* const PreviewScene = ClothEdMode->PreviewScene;
+			const UEditorInteractiveToolsContext* const PreviewToolsContext = PreviewScene->GetClothPreviewEditorModeManager()->GetInteractiveToolsContext();
+			const UInteractiveToolManager* const PreviewToolManager = PreviewToolsContext->ToolManager;
+			 
+			if (Manager == PreviewToolManager)
+			{
+				if (const TSharedPtr<SChaosClothAssetEditor3DViewport> Widget = PreviewViewportWidget.Pin())
+				{
+					return Widget.Get();
+				}
+			}
+			else
+			{
+				if (const TSharedPtr<SChaosClothAssetEditorRestSpaceViewport> Widget = RestSpaceViewportWidget.Pin())
+				{
+					return Widget.Get();
+				}
+			}
+		}
+	}
+
+	return nullptr;
+}
+
+UEditorInteractiveToolsContext* FChaosClothAssetEditorModeToolkit::GetCurrentToolsContext()
+{
+	if (OwningEditorMode.IsValid(/*bEvenIfPendingKill*/ false))
+	{
+		UEdMode* const Mode = OwningEditorMode.Get();
+		if (UChaosClothAssetEditorMode* const ClothEdMode = Cast<UChaosClothAssetEditorMode>(Mode))
+		{
+			return ClothEdMode->GetActiveToolsContext();
+		}
+	}
+
+	// this should not happen, but don't crash if it does
+	return GetEditorModeManager().GetInteractiveToolsContext();
+}
+
+void FChaosClothAssetEditorModeToolkit::OnToolStarted(UInteractiveToolManager* Manager, UInteractiveTool* Tool)
+{
+	ensure(Tool == Manager->GetActiveTool(EToolSide::Left));
+
+	if (!Tool)
+	{
+		return;
+	}
+
+	if (DetailsView)
+	{
+		DetailsView->SetObjects(Tool->GetToolProperties());
+
+		Tool->OnPropertySetsModified.AddLambda([this, Manager]()
+		{
+			const UInteractiveTool* const Tool = Manager->GetActiveTool(EToolSide::Left);
+			if (Tool != nullptr)
+			{
+				DetailsView->SetObjects(Tool->GetToolProperties(true));
+			}
+		});
+	}
+
+	Tool->OnPropertyModifiedDirectlyByTool.AddSP(this, &FChaosClothAssetEditorModeToolkit::InvalidateCachedDetailPanelState);
+
+	ActiveToolName = Tool->GetToolInfo().ToolDisplayName;
+	FString ActiveToolIdentifier = Manager->GetActiveToolName(EToolSide::Mouse);
+	ActiveToolIdentifier.InsertAt(0, ".");
+	ActiveToolIcon = GetActiveToolIcon(ActiveToolIdentifier);
+
+	if (SBaseCharacterFXEditorViewport* Widget = GetViewportWidgetForManager(Manager))
+	{
+		Widget->AddOverlayWidget(ViewportOverlayWidget.ToSharedRef());
+	}
+}
+
+
+void FChaosClothAssetEditorModeToolkit::OnToolEnded(UInteractiveToolManager* Manager, UInteractiveTool* Tool)
+{
+	FModeToolkit::OnToolEnded(Manager, Tool);
+
+	ActiveToolName = FText::GetEmpty();
+	ClearNotification();
+	ClearWarning();
+	
+	if (SBaseCharacterFXEditorViewport* const Widget = GetViewportWidgetForManager(Manager))
+	{
+		Widget->RemoveOverlayWidget(ViewportOverlayWidget.ToSharedRef());
+	}
+
+	if (UInteractiveTool* const CurTool = Manager->GetActiveTool(EToolSide::Left))
+	{
+		CurTool->OnPropertySetsModified.RemoveAll(this);
+		CurTool->OnPropertyModifiedDirectlyByTool.RemoveAll(this);
+	}
+}
+
+void FChaosClothAssetEditorModeToolkit::SetRestSpaceViewportWidget(TWeakPtr<SChaosClothAssetEditorRestSpaceViewport> InRestSpaceViewportWidget)
+{
+	RestSpaceViewportWidget = InRestSpaceViewportWidget;
+}
+
+void FChaosClothAssetEditorModeToolkit::SetPreviewViewportWidget(TWeakPtr<SChaosClothAssetEditor3DViewport> InPreviewViewportWidget)
+{
+	PreviewViewportWidget = InPreviewViewportWidget;
+}
+
+
 #undef LOCTEXT_NAMESPACE
diff --git a/Engine/Plugins/Experimental/ChaosClothAssetEditor/Source/ChaosClothAssetEditor/Private/ChaosClothAsset/ClothEditorModeToolkit.h b/Engine/Plugins/Experimental/ChaosClothAssetEditor/Source/ChaosClothAssetEditor/Private/ChaosClothAsset/ClothEditorModeToolkit.h
--- a/Engine/Plugins/Experimental/ChaosClothAssetEditor/Source/ChaosClothAssetEditor/Private/ChaosClothAsset/ClothEditorModeToolkit.h
+++ b/Engine/Plugins/Experimental/ChaosClothAssetEditor/Source/ChaosClothAssetEditor/Private/ChaosClothAsset/ClothEditorModeToolkit.h
@@ -4,20 +4,46 @@
 
 #include "BaseCharacterFXEditorModeToolkit.h"
 
+class UEditorInteractiveToolsContext;
+class SChaosClothAssetEditorRestSpaceViewport;
+class SChaosClothAssetEditor3DViewport;
+class SBaseCharacterFXEditorViewport;
+
 /**
  * The cloth editor mode toolkit is responsible for the panel on the side in the cloth editor
  * that shows mode and tool properties. Tool buttons would go in Init().
+ * NOTE: the cloth editor has two separate viewports/worlds/modemanagers/toolscontexts, so we need to track which
+ * one is currently active.
  */
 class CHAOSCLOTHASSETEDITOR_API FChaosClothAssetEditorModeToolkit : public FBaseCharacterFXEditorModeToolkit
 {
 public:
 
+	void Init(const TSharedPtr<IToolkitHost>& InitToolkitHost, TWeakObjectPtr<UEdMode> InOwningMode);
+
     /** For a specific tool palette category, construct and fill ToolbarBuilder with the category's tools **/
 	virtual void BuildToolPalette(FName PaletteName, class FToolBarBuilder& ToolbarBuilder) override;
 
 	virtual const FSlateBrush* GetActiveToolIcon(const FString& Identifier) const override;
 
+	virtual void OnToolStarted(UInteractiveToolManager* Manager, UInteractiveTool* Tool) override;
+	virtual void OnToolEnded(UInteractiveToolManager* Manager, UInteractiveTool* Tool) override;
+
 	// IToolkit
 	virtual FName GetToolkitFName() const override;
 	virtual FText GetBaseToolkitName() const override;
+
+	void SetRestSpaceViewportWidget(TWeakPtr<SChaosClothAssetEditorRestSpaceViewport>);
+	void SetPreviewViewportWidget(TWeakPtr<SChaosClothAssetEditor3DViewport>);
+
+private:
+
+	// Get the viewport widget associated with the given manager
+	SBaseCharacterFXEditorViewport* GetViewportWidgetForManager(UInteractiveToolManager* Manager);
+
+	UEditorInteractiveToolsContext* GetCurrentToolsContext();
+
+	TWeakPtr<SChaosClothAssetEditorRestSpaceViewport> RestSpaceViewportWidget;
+	TWeakPtr<SChaosClothAssetEditor3DViewport> PreviewViewportWidget;
+
 };
diff --git a/Engine/Plugins/Experimental/ChaosClothAssetEditor/Source/ChaosClothAssetEditor/Private/ChaosClothAsset/ClothEditorPreviewScene.cpp b/Engine/Plugins/Experimental/ChaosClothAssetEditor/Source/ChaosClothAssetEditor/Private/ChaosClothAsset/ClothEditorPreviewScene.cpp
--- a/Engine/Plugins/Experimental/ChaosClothAssetEditor/Source/ChaosClothAssetEditor/Private/ChaosClothAsset/ClothEditorPreviewScene.cpp
+++ b/Engine/Plugins/Experimental/ChaosClothAssetEditor/Source/ChaosClothAssetEditor/Private/ChaosClothAsset/ClothEditorPreviewScene.cpp
@@ -90,6 +90,11 @@ void FChaosClothPreviewScene::SetModeManager(TSharedPtr<FAssetEditorModeManager>
 	ClothPreviewEditorModeManager = InClothPreviewEditorModeManager;
 }
 
+const TSharedPtr<const FAssetEditorModeManager> FChaosClothPreviewScene::GetClothPreviewEditorModeManager() const
+{
+	return ClothPreviewEditorModeManager;
+}
+
 void FChaosClothPreviewScene::SkeletalMeshTransformChanged(USceneComponent* UpdatedComponent, EUpdateTransformFlags UpdateTransformFlags, ETeleportType Teleport)
 {
 	ensure(UpdatedComponent == SkeletalMeshActor->GetSkeletalMeshComponent());
diff --git a/Engine/Plugins/Experimental/ChaosClothAssetEditor/Source/ChaosClothAssetEditor/Private/ChaosClothAsset/ClothEditorPreviewScene.h b/Engine/Plugins/Experimental/ChaosClothAssetEditor/Source/ChaosClothAssetEditor/Private/ChaosClothAsset/ClothEditorPreviewScene.h
--- a/Engine/Plugins/Experimental/ChaosClothAssetEditor/Source/ChaosClothAssetEditor/Private/ChaosClothAsset/ClothEditorPreviewScene.h
+++ b/Engine/Plugins/Experimental/ChaosClothAssetEditor/Source/ChaosClothAssetEditor/Private/ChaosClothAsset/ClothEditorPreviewScene.h
@@ -75,6 +75,8 @@ public:
 
 	void SetModeManager(TSharedPtr<FAssetEditorModeManager> InClothPreviewEditorModeManager);
 
+	const TSharedPtr<const FAssetEditorModeManager> GetClothPreviewEditorModeManager() const;
+
 private:
 
 	void SkeletalMeshTransformChanged(USceneComponent* UpdatedComponent, EUpdateTransformFlags UpdateTransformFlags, ETeleportType Teleport);
diff --git a/Engine/Plugins/Experimental/ChaosClothAssetEditor/Source/ChaosClothAssetEditor/Private/ChaosClothAsset/ClothEditorToolkit.cpp b/Engine/Plugins/Experimental/ChaosClothAssetEditor/Source/ChaosClothAssetEditor/Private/ChaosClothAsset/ClothEditorToolkit.cpp
--- a/Engine/Plugins/Experimental/ChaosClothAssetEditor/Source/ChaosClothAssetEditor/Private/ChaosClothAsset/ClothEditorToolkit.cpp
+++ b/Engine/Plugins/Experimental/ChaosClothAssetEditor/Source/ChaosClothAssetEditor/Private/ChaosClothAsset/ClothEditorToolkit.cpp
@@ -12,6 +12,7 @@
 #include "ChaosClothAsset/ClothComponent.h"
 #include "ChaosClothAsset/ClothAsset.h"
 #include "ChaosClothAsset/ClothCollection.h"
+#include "ChaosClothAsset/ClothEditorModeToolkit.h"
 #include "Framework/Docking/LayoutExtender.h"
 #include "AssetEditorModeManager.h"
 #include "AdvancedPreviewScene.h"
@@ -179,7 +180,7 @@ FChaosClothAssetEditorToolkit::FChaosClothAssetEditorToolkit(UAssetEditor* InOwn
 
 	ClothPreviewViewportDelegate = [this](FAssetEditorViewportConstructionArgs InArgs)
 	{
-		return SNew(SChaosClothAssetEditor3DViewport, InArgs)
+		return SAssignNew(PreviewViewportWidget, SChaosClothAssetEditor3DViewport, InArgs)
 			.EditorViewportClient(ClothPreviewViewportClient);
 	};
 
@@ -253,6 +254,13 @@ void FChaosClothAssetEditorToolkit::InitializeEdMode(UBaseCharacterFXEditorMode*
 	OwningAssetEditor->GetObjectsToEdit(ObjectsToEdit);
 
 	ClothMode->InitializeTargets(ObjectsToEdit);
+
+	if (TSharedPtr<FModeToolkit> ModeToolkit = ClothMode->GetToolkit().Pin())
+	{
+		FChaosClothAssetEditorModeToolkit* ClothModeToolkit = static_cast<FChaosClothAssetEditorModeToolkit*>(ModeToolkit.Get());
+		ClothModeToolkit->SetRestSpaceViewportWidget(RestSpaceViewportWidget);
+		ClothModeToolkit->SetPreviewViewportWidget(PreviewViewportWidget);
+	}
 }
 
 void FChaosClothAssetEditorToolkit::CreateEditorModeUILayer()
@@ -296,7 +304,7 @@ AssetEditorViewportFactoryFunction FChaosClothAssetEditorToolkit::GetViewportDel
 {
 	AssetEditorViewportFactoryFunction TempViewportDelegate = [this](FAssetEditorViewportConstructionArgs InArgs)
 	{
-		return SAssignNew(RestSpaceViewport, SChaosClothAssetEditorRestSpaceViewport, InArgs)
+		return SAssignNew(RestSpaceViewportWidget, SChaosClothAssetEditorRestSpaceViewport, InArgs)
 			.EditorViewportClient(ViewportClient);
 	};
 
@@ -388,7 +396,7 @@ void FChaosClothAssetEditorToolkit::PostInitAssetEditor()
 	{
 		// when CreateEditorViewportClient() is called, RestSpaceViewport is null. Set it here instead
 		TSharedPtr<FChaosClothEditorRestSpaceViewportClient> VC = StaticCastSharedPtr<FChaosClothEditorRestSpaceViewportClient>(ViewportClient);
-		VC->SetEditorViewportWidget(RestSpaceViewport);
+		VC->SetEditorViewportWidget(RestSpaceViewportWidget);
 	}
 
 	SetCommonViewportClientOptions(ViewportClient.Get());
diff --git a/Engine/Plugins/Experimental/ChaosClothAssetEditor/Source/ChaosClothAssetEditor/Private/ChaosClothAsset/ClothEditorToolkit.h b/Engine/Plugins/Experimental/ChaosClothAssetEditor/Source/ChaosClothAssetEditor/Private/ChaosClothAsset/ClothEditorToolkit.h
--- a/Engine/Plugins/Experimental/ChaosClothAssetEditor/Source/ChaosClothAssetEditor/Private/ChaosClothAsset/ClothEditorToolkit.h
+++ b/Engine/Plugins/Experimental/ChaosClothAssetEditor/Source/ChaosClothAssetEditor/Private/ChaosClothAsset/ClothEditorToolkit.h
@@ -16,6 +16,8 @@ class SDataflowGraphEditor;
 class IStructureDetailsView;
 class UEdGraphNode;
 class UChaosClothComponent;
+class SChaosClothAssetEditorRestSpaceViewport;
+class SChaosClothAssetEditor3DViewport;
 
 namespace Dataflow
 {
@@ -107,7 +109,8 @@ private:
 	TSharedPtr<FChaosClothAssetEditor3DViewportClient> ClothPreviewViewportClient;
 	TSharedPtr<FAssetEditorModeManager> ClothPreviewEditorModeManager;
 
-	TWeakPtr<SEditorViewport> RestSpaceViewport;
+	TSharedPtr<SChaosClothAssetEditorRestSpaceViewport> RestSpaceViewportWidget;
+	TSharedPtr<SChaosClothAssetEditor3DViewport> PreviewViewportWidget;
 
 	TSharedPtr<SDockTab> PreviewSceneDockTab;
 	TSharedPtr<SWidget> AdvancedPreviewSettingsWidget;
diff --git a/Engine/Plugins/Experimental/ChaosClothAssetEditor/Source/ChaosClothAssetEditor/Private/ChaosClothAsset/SClothEditor3DViewport.h b/Engine/Plugins/Experimental/ChaosClothAssetEditor/Source/ChaosClothAssetEditor/Private/ChaosClothAsset/SClothEditor3DViewport.h
--- a/Engine/Plugins/Experimental/ChaosClothAssetEditor/Source/ChaosClothAssetEditor/Private/ChaosClothAsset/SClothEditor3DViewport.h
+++ b/Engine/Plugins/Experimental/ChaosClothAssetEditor/Source/ChaosClothAssetEditor/Private/ChaosClothAsset/SClothEditor3DViewport.h
@@ -5,11 +5,12 @@
 #include "CoreMinimal.h"
 #include "SAssetEditorViewport.h"
 #include "SCommonEditorViewportToolbarBase.h"
+#include "SBaseCharacterFXEditorViewport.h"
 
 /**
  * Viewport used for 3D preview in cloth editor. Has a custom toolbar overlay at the top.
  */
-class CHAOSCLOTHASSETEDITOR_API SChaosClothAssetEditor3DViewport : public SAssetEditorViewport, public ICommonEditorViewportToolbarInfoProvider
+class CHAOSCLOTHASSETEDITOR_API SChaosClothAssetEditor3DViewport : public SBaseCharacterFXEditorViewport, public ICommonEditorViewportToolbarInfoProvider
 {
 
 public:
diff --git a/Engine/Plugins/Experimental/ChaosClothAssetEditor/Source/ChaosClothAssetEditorTools/Private/ChaosClothAsset/ClothTransferSkinWeightsTool.cpp b/Engine/Plugins/Experimental/ChaosClothAssetEditor/Source/ChaosClothAssetEditorTools/Private/ChaosClothAsset/ClothTransferSkinWeightsTool.cpp
--- a/Engine/Plugins/Experimental/ChaosClothAssetEditor/Source/ChaosClothAssetEditorTools/Private/ChaosClothAsset/ClothTransferSkinWeightsTool.cpp
+++ b/Engine/Plugins/Experimental/ChaosClothAssetEditor/Source/ChaosClothAssetEditorTools/Private/ChaosClothAsset/ClothTransferSkinWeightsTool.cpp
@@ -30,6 +30,7 @@
 #include "InteractiveToolManager.h"
 
 #include "Engine/SkeletalMesh.h"
+#include "TargetInterfaces/SkeletalMeshBackedTarget.h"
 
 #define LOCTEXT_NAMESPACE "ClothSkinWeightRetargetingTool"
 
@@ -57,7 +58,12 @@ const FToolTargetTypeRequirements& UClothTransferSkinWeightsToolBuilder::GetTarg
 
 bool UClothTransferSkinWeightsToolBuilder::CanBuildTool(const FToolBuilderState& SceneState) const
 {
-	return (SceneState.TargetManager->CountSelectedAndTargetable(SceneState, GetTargetRequirements()) == 1);
+	const bool ClothComponentSelected = (SceneState.TargetManager->CountSelectedAndTargetable(SceneState, GetTargetRequirements()) == 1);
+	
+	const static FToolTargetTypeRequirements SourceMeshRequirements(USkeletalMeshBackedTarget::StaticClass());
+	const bool SkeletalMeshComponentSelected = (SceneState.TargetManager->CountSelectedAndTargetable(SceneState, SourceMeshRequirements) == 1);
+
+	return ClothComponentSelected && SkeletalMeshComponentSelected;
 }
 
 UInteractiveTool* UClothTransferSkinWeightsToolBuilder::BuildTool(const FToolBuilderState& SceneState) const
