Cloth transfer skin weights tool:
- automatically populate the source asset based on selected skeletal mesh component
- enable hiding the source component
- visualize transferred skin weights on the target mesh

diff --git a/Engine/Plugins/Experimental/ChaosClothAssetEditor/Source/ChaosClothAssetEditorTools/Private/ChaosClothAsset/ClothTransferSkinWeightsTool.cpp b/Engine/Plugins/Experimental/ChaosClothAssetEditor/Source/ChaosClothAssetEditorTools/Private/ChaosClothAsset/ClothTransferSkinWeightsTool.cpp
--- a/Engine/Plugins/Experimental/ChaosClothAssetEditor/Source/ChaosClothAssetEditorTools/Private/ChaosClothAsset/ClothTransferSkinWeightsTool.cpp
+++ b/Engine/Plugins/Experimental/ChaosClothAssetEditor/Source/ChaosClothAssetEditorTools/Private/ChaosClothAsset/ClothTransferSkinWeightsTool.cpp
@@ -30,10 +30,102 @@
 #include "InteractiveToolManager.h"
 
 #include "Engine/SkeletalMesh.h"
+#include "PreviewMesh.h"
+#include "DynamicMeshEditor.h"
+#include "Util/ColorConstants.h"
+#include "ToolSetupUtil.h"
+#include "ModelingToolTargetUtil.h"
 #include "TargetInterfaces/SkeletalMeshBackedTarget.h"
+#include "Components/SkeletalMeshComponent.h"
 
 #define LOCTEXT_NAMESPACE "ClothSkinWeightRetargetingTool"
 
+
+namespace ClothTransferSkinWeightsToolHelpers
+{
+	/** 
+	 * Compute mappings between indices and bone names.
+	 * 
+     * @note We assume that each mesh inherits its reference skeleton from the same USkeleton asset. However, their  
+     * internal indexing can be different and hence when transfering weights we need to make sure we reference the correct  
+     * bones via their names instead of indices. 
+     */
+	void GetBoneMaps(const USkinnedAsset* SourceSkinnedAsset, 
+						  const USkinnedAsset* TargetSkinnedAsset, 
+						  TMap<FBoneIndexType, FName>& SourceIndexToBone,
+						  TMap<FName, FBoneIndexType>& TargetBoneToIndex)
+	{
+		TargetBoneToIndex.Reset();
+		SourceIndexToBone.Reset();
+		const FReferenceSkeleton& SourceRefSkeleton = SourceSkinnedAsset->GetRefSkeleton();
+		for (int32 Index = 0; Index < SourceRefSkeleton.GetRawBoneNum(); ++Index)
+		{
+			SourceIndexToBone.Add(Index, SourceRefSkeleton.GetRawRefBoneInfo()[Index].Name);
+		}
+
+		const FReferenceSkeleton& TargetRefSkeleton = TargetSkinnedAsset->GetRefSkeleton();
+		for (int32 Index = 0; Index < TargetRefSkeleton.GetRawBoneNum(); ++Index)
+		{
+			TargetBoneToIndex.Add(TargetRefSkeleton.GetRawRefBoneInfo()[Index].Name, Index);
+		}
+	};
+
+	void SkeletalMeshToDynamicMesh(USkeletalMesh* FromSkeletalMeshAsset, int32 SourceLODIdx, FDynamicMesh3& ToDynamicMesh)
+	{
+		FMeshDescription SourceMesh;
+
+		// Check first if we have bulk data available and non-empty.
+		if (FromSkeletalMeshAsset->IsLODImportedDataBuildAvailable(SourceLODIdx) && !FromSkeletalMeshAsset->IsLODImportedDataEmpty(SourceLODIdx))
+		{
+			FSkeletalMeshImportData SkeletalMeshImportData;
+			FromSkeletalMeshAsset->LoadLODImportedData(SourceLODIdx, SkeletalMeshImportData);
+			SkeletalMeshImportData.GetMeshDescription(SourceMesh);
+		}
+		else
+		{
+			// Fall back on the LOD model directly if no bulk data exists. When we commit
+			// the mesh description, we override using the bulk data. This can happen for older
+			// skeletal meshes, from UE 4.24 and earlier.
+			const FSkeletalMeshModel* SkeletalMeshModel = FromSkeletalMeshAsset->GetImportedModel();
+			if (SkeletalMeshModel && SkeletalMeshModel->LODModels.IsValidIndex(SourceLODIdx))
+			{
+				SkeletalMeshModel->LODModels[SourceLODIdx].GetMeshDescription(SourceMesh, FromSkeletalMeshAsset);
+			}
+		}
+
+		FMeshDescriptionToDynamicMesh Converter;
+		Converter.Convert(&SourceMesh, ToDynamicMesh);
+	};
+
+
+	void ClothComponentToDynamicMesh(const UChaosClothComponent* ClothComponent, UE::Geometry::FDynamicMesh3& MeshOut)
+	{
+		const UChaosClothAsset* ChaosClothAsset = ClothComponent->GetClothAsset();
+		if (!ChaosClothAsset)
+		{
+			return;
+		}
+
+		const UE::Chaos::ClothAsset::FClothConstAdapter ClothAdapter(ChaosClothAsset->GetClothCollection());
+		constexpr int32 LodIndex = 0;
+		const UE::Chaos::ClothAsset::FClothLodConstAdapter ClothLodAdapter = ClothAdapter.GetLod(LodIndex);
+		constexpr bool bGet2DPattern = false;
+
+		UE::Geometry::FDynamicMeshEditor MeshEditor(&MeshOut);
+		FClothPatternToDynamicMesh Converter;
+
+		for (int32 PatternIndex = 0; PatternIndex < ClothLodAdapter.GetNumPatterns(); ++PatternIndex)
+		{
+			UE::Geometry::FDynamicMesh3 PatternMesh;
+			Converter.Convert(ChaosClothAsset, LodIndex, PatternIndex, bGet2DPattern, PatternMesh);
+
+			UE::Geometry::FMeshIndexMappings IndexMaps;
+			MeshEditor.AppendMesh(&PatternMesh, IndexMaps);
+		}
+	};
+}
+
+
 // ------------------- Properties -------------------
 
 void UClothTransferSkinWeightsToolActionProperties::PostAction(EClothTransferSkinWeightsToolActions Action)
@@ -68,19 +160,37 @@
 
 UInteractiveTool* UClothTransferSkinWeightsToolBuilder::BuildTool(const FToolBuilderState& SceneState) const
 {
-	UClothTransferSkinWeightsTool* NewTool = NewObject<UClothTransferSkinWeightsTool>();
+	UClothTransferSkinWeightsTool* NewTool = NewObject<UClothTransferSkinWeightsTool>(SceneState.ToolManager);
 	
 	UToolTarget* Target = SceneState.TargetManager->BuildFirstSelectedTargetable(SceneState, GetTargetRequirements());
 	NewTool->SetTarget(Target);
-
+	NewTool->SetWorld(SceneState.World);
 	return NewTool;
 }
 
+
+void UClothTransferSkinWeightsToolBuilder::PostSetupTool(UInteractiveTool* Tool, const FToolBuilderState& SceneState) const
+{
+	if (UClothTransferSkinWeightsTool* NewTool = Cast<UClothTransferSkinWeightsTool>(Tool))
+	{
+		for (UActorComponent* SelectedComponent : SceneState.SelectedComponents)
+		{
+			if (USkeletalMeshComponent* SkeletalMeshComponent = Cast<USkeletalMeshComponent>(SelectedComponent))
+			{
+				NewTool->ToolProperties->SourceMesh = SkeletalMeshComponent->GetSkeletalMeshAsset();
+				NewTool->ToolProperties->SourceMeshTransform = SkeletalMeshComponent->GetComponentTransform();
+				NewTool->SourceComponent = SkeletalMeshComponent;
+				break;
+			}
+		}
+	}
+}
+
 // ------------------- Tool -------------------
 
 void UClothTransferSkinWeightsTool::Setup()
 {
-	UInteractiveTool::Setup();
+	USingleSelectionMeshEditingTool::Setup();
 
 	UClothComponentToolTarget* ClothComponentToolTarget = Cast<UClothComponentToolTarget>(Target);
 	ClothComponent = ClothComponentToolTarget->GetClothComponent();
@@ -91,8 +201,178 @@
 	ActionProperties = NewObject<UClothTransferSkinWeightsToolActionProperties>(this);
 	ActionProperties->ParentTool = this;
 	AddToolPropertySource(ActionProperties);
+
+
+	PreviewMesh = NewObject<UPreviewMesh>(this);
+	if (PreviewMesh == nullptr)
+	{
+		return;
+	}
+	PreviewMesh->CreateInWorld(GetTargetWorld(), FTransform::Identity);
+	ToolSetupUtil::ApplyRenderingConfigurationToPreview(PreviewMesh, nullptr);
+
+
+	PreviewMesh->SetTransform(ClothComponentToolTarget->GetWorldTransform());
+
+	ToolProperties->WatchProperty(ToolProperties->SourceMesh, [this](TObjectPtr<USkeletalMesh>) { UpdatePreviewMesh(); });
+	//ToolProperties->WatchProperty(ToolProperties->SourceMeshTransform, [this](const FTransform&) { UpdatePreviewMesh(); });
+	ToolProperties->WatchProperty(ToolProperties->BoneName, [this](const FName&) {  UpdatePreviewMeshColor(); });
+	ToolProperties->WatchProperty(ToolProperties->bHideSourceMesh, [this](bool) { UpdateSourceMeshRender(); });
 }
 
+void UClothTransferSkinWeightsTool::Shutdown(EToolShutdownType ShutdownType)
+{
+	USingleSelectionMeshEditingTool::Shutdown(ShutdownType);
+
+	if (PreviewMesh)
+	{
+		PreviewMesh->Disconnect();
+	}
+
+	UE::ToolTarget::ShowSourceObject(Target);
+	SourceComponent->SetVisibility(true);
+}
+
+void UClothTransferSkinWeightsTool::UpdatePreviewMeshColor()
+{
+	using namespace UE::AnimationCore;
+	using namespace UE::Geometry;
+
+	PreviewMesh->SetTriangleColorFunction([this](const FDynamicMesh3* Mesh, int TriangleID) -> FColor
+	{
+		const FName& CurrentBoneName = ToolProperties->BoneName;
+		if (!TargetMeshBoneNameToIndex.Contains(CurrentBoneName))
+		{
+			return FColor::Black;
+		}
+
+		const FBoneIndexType CurrentBoneIndex = TargetMeshBoneNameToIndex[CurrentBoneName];
+		const FIndex3i Tri = Mesh->GetTriangle(TriangleID);
+
+		const FName ProfileName = FSkeletalMeshAttributes::DefaultSkinWeightProfileName; // always use default profile for now, later this will be set by the user
+		FDynamicMeshVertexSkinWeightsAttribute* Attribute = PreviewMesh->GetPreviewDynamicMesh()->Attributes()->GetSkinWeightsAttribute(ProfileName);
+		if (Attribute == nullptr)
+		{
+			FLinearColor Lin(1.0f, 0.3f, 0.3f, 1.0f);
+			return Lin.ToFColor(/*bSRGB = */ true);
+		}
+
+		float AvgWeight = 0.0f;
+		for (int32 VID = 0; VID < 3; ++VID)
+		{
+			int32 VertexID = Tri[VID];
+			FBoneWeights Data;
+			Attribute->GetValue(VertexID, Data);
+			for (FBoneWeight Wt : Data)
+			{
+				if (Wt.GetBoneIndex() == CurrentBoneIndex)
+				{
+					AvgWeight += Wt.GetWeight();
+				}
+			}
+		}
+
+		AvgWeight /= 3.0f;
+		FLinearColor Lin(AvgWeight, AvgWeight, AvgWeight, 1.0f);
+		return Lin.ToFColor(/*bSRGB = */ true);
+	},
+	UPreviewMesh::ERenderUpdateMode::FullUpdate);
+}
+
+void UClothTransferSkinWeightsTool::UpdatePreviewMesh()
+{
+	using namespace UE::AnimationCore;
+	using namespace UE::Geometry;
+	
+	//TODO: for now, assume we are always transfering from LOD 0, but make this a parameter in the future...
+	constexpr int32 SourceLODIdx = 0; 
+
+	// User hasn't specified the source mesh in the UI
+	if (ToolProperties->SourceMesh == nullptr)
+	{
+		//TODO: Display error message
+		return;
+	}
+
+	// Convert source Skeletal Mesh to Dynamic Mesh
+	FDynamicMesh3 SourceDynamicMesh;
+	USkeletalMesh* FromSkeletalMeshAsset = ToolProperties->SourceMesh;
+	ClothTransferSkinWeightsToolHelpers::SkeletalMeshToDynamicMesh(FromSkeletalMeshAsset, SourceLODIdx, SourceDynamicMesh);
+	MeshTransforms::ApplyTransform(SourceDynamicMesh, ToolProperties->SourceMeshTransform, true);
+
+	// Convert target ClothComponent to Dynamic Mesh
+	UE::Geometry::FDynamicMesh3 TargetDynamicMesh;
+	TargetDynamicMesh.EnableAttributes();
+	TargetDynamicMesh.Attributes()->AttachSkinWeightsAttribute(FSkeletalMeshAttributes::DefaultSkinWeightProfileName, new FDynamicMeshVertexSkinWeightsAttribute(&TargetDynamicMesh));
+	ClothTransferSkinWeightsToolHelpers::ClothComponentToDynamicMesh(ClothComponent, TargetDynamicMesh);
+
+
+	FTransferBoneWeights TransferBoneWeights(&SourceDynamicMesh, FSkeletalMeshAttributes::DefaultSkinWeightProfileName);
+
+	// Compute bone index mappings
+	TMap<FBoneIndexType, FName> SourceIndexToBone;
+	ClothTransferSkinWeightsToolHelpers::GetBoneMaps(FromSkeletalMeshAsset, ClothComponent->GetClothAsset(), SourceIndexToBone, this->TargetMeshBoneNameToIndex);
+	TransferBoneWeights.SourceIndexToBone = &SourceIndexToBone;
+	TransferBoneWeights.TargetBoneToIndex = &this->TargetMeshBoneNameToIndex;
+
+	// Do the actual transfer
+	const FTransformSRT3d TargetToWorld = ClothComponent->GetComponentTransform();
+	if (TransferBoneWeights.Validate() == EOperationValidationResult::Ok)
+	{
+		TransferBoneWeights.Compute(TargetDynamicMesh, TargetToWorld, FSkeletalMeshAttributes::DefaultSkinWeightProfileName);
+	}
+
+	// Get set of bone indices used in the target mesh
+	TMap<FName, FBoneIndexType> UsedBoneNames;
+	const FReferenceSkeleton& TargetRefSkeleton = ClothComponent->GetClothAsset()->GetRefSkeleton();
+
+	const TMap<FName, TUniquePtr<FDynamicMeshVertexSkinWeightsAttribute>>& WeightLayers = TargetDynamicMesh.Attributes()->GetSkinWeightsAttributes();
+	for (const TPair<FName, TUniquePtr<FDynamicMeshVertexSkinWeightsAttribute>>& LayerPair : WeightLayers)
+	{
+		const FDynamicMeshVertexSkinWeightsAttribute* Layer = LayerPair.Value.Get();
+
+		for (int VertexID = 0; VertexID < TargetDynamicMesh.MaxVertexID(); ++VertexID)
+		{
+			if (TargetDynamicMesh.IsVertex(VertexID))
+			{
+				FBoneWeights Data;
+				Layer->GetValue(VertexID, Data);
+				for (FBoneWeight Wt : Data)
+				{
+					const FBoneIndexType BoneIndex = Wt.GetBoneIndex();
+					const FName& BoneName = TargetRefSkeleton.GetRawRefBoneInfo()[BoneIndex].Name;
+					UsedBoneNames.Add(BoneName, BoneIndex);
+				}
+			}
+		}
+	}
+
+	// Update list of bone names in the Properties panel
+	UsedBoneNames.ValueSort([](int16 A, int16 B) { return A < B; });
+	UsedBoneNames.GetKeys(ToolProperties->BoneNameList);
+
+	// Update the preview mesh
+	PreviewMesh->UpdatePreview(&TargetDynamicMesh);
+	PreviewMesh->SetMaterial(ToolSetupUtil::GetDefaultSculptMaterial(GetToolManager()));
+	PreviewMesh->SetOverrideRenderMaterial(ToolSetupUtil::GetSelectionMaterial(GetToolManager()));
+
+	UpdatePreviewMeshColor();
+
+	PreviewMesh->SetTransform(TargetToWorld);
+	PreviewMesh->SetVisible(true);
+
+	UE::ToolTarget::HideSourceObject(Target);
+}
+
+
+void UClothTransferSkinWeightsTool::UpdateSourceMeshRender()
+{
+	if (ToolProperties && SourceComponent)
+	{
+		SourceComponent->SetVisibility(!ToolProperties->bHideSourceMesh);
+	}
+}
+
 void UClothTransferSkinWeightsTool::TransferWeights()
 {
 	using namespace UE::AnimationCore;
@@ -100,62 +380,7 @@
 	
 	//TODO: for now, assume we are always transfering from LOD 0, but make this a parameter in the future...
 	constexpr int32 SourceLODIdx = 0; 
-
-	/** 
-	 * Compute mappings between indices and bone names.
-	 * 
-     * @note We assume that each mesh inherits its reference skeleton from the same USkeleton asset. However, their  
-     * internal indexing can be different and hence when transfering weights we need to make sure we reference the correct  
-     * bones via their names instead of indices. 
-     */
-	auto GetBoneMaps = [](const USkinnedAsset* SourceSkinnedAsset, 
-						  const USkinnedAsset* TargetSkinnedAsset, 
-						  TMap<FBoneIndexType, FName>& IndexToBone,
-						  TMap<FName, FBoneIndexType>& BoneToIndex)
-	{
-		BoneToIndex.Reset();
-		IndexToBone.Reset();
-		const FReferenceSkeleton& SourceRefSkeleton = SourceSkinnedAsset->GetRefSkeleton();
-		for (int32 Index = 0; Index < SourceRefSkeleton.GetRawBoneNum(); ++Index)
-		{
-			IndexToBone.Add(Index, SourceRefSkeleton.GetRawRefBoneInfo()[Index].Name);
-		}
-
-		const FReferenceSkeleton& TargetRefSkeleton = TargetSkinnedAsset->GetRefSkeleton();
-		for (int32 Index = 0; Index < TargetRefSkeleton.GetRawBoneNum(); ++Index)
-		{
-			BoneToIndex.Add(TargetRefSkeleton.GetRawRefBoneInfo()[Index].Name, Index);
-		}
-	};
-
-	auto SkeletalMeshToDynamicMesh = [SourceLODIdx](USkeletalMesh* FromSkeletalMeshAsset,
-									    		    FDynamicMesh3& ToDynamicMesh)
-	{
-		FMeshDescription SourceMesh;
-
-		// Check first if we have bulk data available and non-empty.
-		if (FromSkeletalMeshAsset->IsLODImportedDataBuildAvailable(SourceLODIdx) && !FromSkeletalMeshAsset->IsLODImportedDataEmpty(SourceLODIdx))
-		{
-			FSkeletalMeshImportData SkeletalMeshImportData;
-			FromSkeletalMeshAsset->LoadLODImportedData(SourceLODIdx, SkeletalMeshImportData);
-			SkeletalMeshImportData.GetMeshDescription(SourceMesh);
-		}
-		else
-		{
-			// Fall back on the LOD model directly if no bulk data exists. When we commit
-			// the mesh description, we override using the bulk data. This can happen for older
-			// skeletal meshes, from UE 4.24 and earlier.
-			const FSkeletalMeshModel* SkeletalMeshModel = FromSkeletalMeshAsset->GetImportedModel();
-			if (SkeletalMeshModel && SkeletalMeshModel->LODModels.IsValidIndex(SourceLODIdx))
-			{
-				SkeletalMeshModel->LODModels[SourceLODIdx].GetMeshDescription(SourceMesh, FromSkeletalMeshAsset);
-			}
-		}
-
-		FMeshDescriptionToDynamicMesh Converter;
-		Converter.Convert(&SourceMesh, ToDynamicMesh);
-	};
-
+	
 	// User hasn't specified the source mesh in the UI
 	if (ToolProperties->SourceMesh == nullptr) 
 	{
@@ -165,8 +390,8 @@
 
 	// Convert source Skeletal Mesh to Dynamic Mesh
 	FDynamicMesh3 SourceDynamicMesh;
-	USkeletalMesh* FromSkeletalMeshAsset = ToolProperties->SourceMesh.Get();
-	SkeletalMeshToDynamicMesh(FromSkeletalMeshAsset, SourceDynamicMesh);
+	USkeletalMesh* FromSkeletalMeshAsset = ToolProperties->SourceMesh;
+	ClothTransferSkinWeightsToolHelpers::SkeletalMeshToDynamicMesh(FromSkeletalMeshAsset, SourceLODIdx, SourceDynamicMesh);
 	FTransformSRT3d SourceToWorld; //TODO: Allows the user to set this value or infer it from an editor
 	MeshTransforms::ApplyTransform(SourceDynamicMesh, SourceToWorld, true);
 
@@ -175,7 +400,7 @@
 	// Compute bone index mappings
 	TMap<FBoneIndexType, FName> SourceIndexToBone;
 	TMap<FName, FBoneIndexType> TargetBoneToIndex;
-	GetBoneMaps(static_cast<USkinnedAsset*>(FromSkeletalMeshAsset), static_cast<USkinnedAsset*>(TargetClothAsset), SourceIndexToBone, TargetBoneToIndex);
+	ClothTransferSkinWeightsToolHelpers::GetBoneMaps(static_cast<USkinnedAsset*>(FromSkeletalMeshAsset), static_cast<USkinnedAsset*>(TargetClothAsset), SourceIndexToBone, TargetBoneToIndex);
 	
 	// Setup bone weight transfer operator
 	FTransferBoneWeights TransferBoneWeights(&SourceDynamicMesh, FSkeletalMeshAttributes::DefaultSkinWeightProfileName); 
diff --git a/Engine/Plugins/Experimental/ChaosClothAssetEditor/Source/ChaosClothAssetEditorTools/Public/ChaosClothAsset/ClothTransferSkinWeightsTool.h b/Engine/Plugins/Experimental/ChaosClothAssetEditor/Source/ChaosClothAssetEditorTools/Public/ChaosClothAsset/ClothTransferSkinWeightsTool.h
--- a/Engine/Plugins/Experimental/ChaosClothAssetEditor/Source/ChaosClothAssetEditorTools/Public/ChaosClothAsset/ClothTransferSkinWeightsTool.h
+++ b/Engine/Plugins/Experimental/ChaosClothAssetEditor/Source/ChaosClothAssetEditorTools/Public/ChaosClothAsset/ClothTransferSkinWeightsTool.h
@@ -5,21 +5,46 @@
 #include "UObject/Object.h"
 #include "SingleSelectionTool.h"
 #include "InteractiveToolBuilder.h"
+#include "BaseTools/SingleSelectionMeshEditingTool.h"
+#include "BoneIndices.h"
+#include "TransformTypes.h"
 
 #include "ClothTransferSkinWeightsTool.generated.h"
 
 class UChaosClothComponent;
 class UClothTransferSkinWeightsTool;
 class USkeletalMesh;
+class UPreviewMesh;
+class USkeletalMeshComponent;
 
 UCLASS()
-class CHAOSCLOTHASSETEDITORTOOLS_API UClothTransferSkinWeightsToolProperties : public UObject
+class CHAOSCLOTHASSETEDITORTOOLS_API UClothTransferSkinWeightsToolProperties : public UInteractiveToolPropertySet
 {
 	GENERATED_BODY()
 public:
 
-	UPROPERTY(EditAnywhere, Category = Files)
+	UPROPERTY(EditAnywhere, Category = Source)
 	TObjectPtr<USkeletalMesh> SourceMesh;
+
+	UPROPERTY(EditAnywhere, Category = Source)
+	FTransform SourceMeshTransform;
+
+	UPROPERTY(EditAnywhere, Category = Source)
+	bool bHideSourceMesh = false;
+
+	UPROPERTY(EditAnywhere, Category = Visualization, meta = (GetOptions = GetBoneNameList))
+	FName BoneName;
+
+	// Get the list of valid bone names
+	UFUNCTION()
+	TArray<FName> GetBoneNameList()
+	{
+		return BoneNameList;
+	}
+
+	UPROPERTY(meta = (TransientToolProperty))
+	TArray<FName> BoneNameList;
+
 };
 
 UENUM()
@@ -51,7 +76,7 @@ public:
 
 
 UCLASS()
-class CHAOSCLOTHASSETEDITORTOOLS_API UClothTransferSkinWeightsToolBuilder : public UInteractiveToolWithToolTargetsBuilder
+class CHAOSCLOTHASSETEDITORTOOLS_API UClothTransferSkinWeightsToolBuilder : public USingleSelectionMeshEditingToolBuilder
 {
 	GENERATED_BODY()
 
@@ -64,22 +89,25 @@ public:
 	virtual bool CanBuildTool(const FToolBuilderState& SceneState) const override;
 
 	virtual UInteractiveTool* BuildTool(const FToolBuilderState& SceneState) const override;
+	virtual void PostSetupTool(UInteractiveTool* Tool, const FToolBuilderState& SceneState) const override;
 };
 
 
 UCLASS()
-class CHAOSCLOTHASSETEDITORTOOLS_API UClothTransferSkinWeightsTool : public USingleSelectionTool
+class CHAOSCLOTHASSETEDITORTOOLS_API UClothTransferSkinWeightsTool : public USingleSelectionMeshEditingTool
 {
 	GENERATED_BODY()
 public:
 
 	// UInteractiveTool
 	virtual void Setup() override;
+	virtual void Shutdown(EToolShutdownType ShutdownType) override;
 	virtual void OnTick(float DeltaTime) override;
 
 private:
 
 	friend class UClothTransferSkinWeightsToolActionProperties;
+	friend class UClothTransferSkinWeightsToolBuilder;
 
 	UPROPERTY()
 	TObjectPtr<UClothTransferSkinWeightsToolProperties> ToolProperties;
@@ -92,10 +120,23 @@ private:
 	UPROPERTY()
 	TObjectPtr<const UChaosClothComponent> ClothComponent;
 
+	UPROPERTY()
+	TObjectPtr<USkeletalMeshComponent> SourceComponent;
+
+
 	void RequestAction(EClothTransferSkinWeightsToolActions ActionType);
 
 	void TransferWeights();
 
+	UPROPERTY()
+	TObjectPtr<UPreviewMesh> PreviewMesh;
+
+	TMap<FName, FBoneIndexType> TargetMeshBoneNameToIndex;
+
+	void UpdatePreviewMeshColor();
+	void UpdatePreviewMesh();
+
+	void UpdateSourceMeshRender();
 };
 
 
