TransferSkinWeightsTool: fix missing function override in tool builder class

CIS 254907

diff --git a/Engine/Plugins/Experimental/ChaosClothAssetEditor/Source/ChaosClothAssetEditorTools/Private/ChaosClothAsset/ClothTransferSkinWeightsTool.cpp b/Engine/Plugins/Experimental/ChaosClothAssetEditor/Source/ChaosClothAssetEditorTools/Private/ChaosClothAsset/ClothTransferSkinWeightsTool.cpp
--- a/Engine/Plugins/Experimental/ChaosClothAssetEditor/Source/ChaosClothAssetEditorTools/Private/ChaosClothAsset/ClothTransferSkinWeightsTool.cpp
+++ b/Engine/Plugins/Experimental/ChaosClothAssetEditor/Source/ChaosClothAssetEditorTools/Private/ChaosClothAsset/ClothTransferSkinWeightsTool.cpp
@@ -158,17 +158,15 @@ bool UClothTransferSkinWeightsToolBuilder::CanBuildTool(const FToolBuilderState&
 	return ClothComponentSelected && SkeletalMeshComponentSelected;
 }
 
-UInteractiveTool* UClothTransferSkinWeightsToolBuilder::BuildTool(const FToolBuilderState& SceneState) const
+USingleSelectionMeshEditingTool* UClothTransferSkinWeightsToolBuilder::CreateNewTool(const FToolBuilderState& SceneState) const
 {
 	UClothTransferSkinWeightsTool* NewTool = NewObject<UClothTransferSkinWeightsTool>(SceneState.ToolManager);
-	
-	UToolTarget* Target = SceneState.TargetManager->BuildFirstSelectedTargetable(SceneState, GetTargetRequirements());
-	NewTool->SetTarget(Target);
-	NewTool->SetWorld(SceneState.World);
+
+	// Setting Target and World on the new tool is handled in USingleSelectionMeshEditingToolBuilder::InitializeNewTool
+
 	return NewTool;
 }
 
-
 void UClothTransferSkinWeightsToolBuilder::PostSetupTool(UInteractiveTool* Tool, const FToolBuilderState& SceneState) const
 {
 	if (UClothTransferSkinWeightsTool* NewTool = Cast<UClothTransferSkinWeightsTool>(Tool))
diff --git a/Engine/Plugins/Experimental/ChaosClothAssetEditor/Source/ChaosClothAssetEditorTools/Public/ChaosClothAsset/ClothTransferSkinWeightsTool.h b/Engine/Plugins/Experimental/ChaosClothAssetEditor/Source/ChaosClothAssetEditorTools/Public/ChaosClothAsset/ClothTransferSkinWeightsTool.h
--- a/Engine/Plugins/Experimental/ChaosClothAssetEditor/Source/ChaosClothAssetEditorTools/Public/ChaosClothAsset/ClothTransferSkinWeightsTool.h
+++ b/Engine/Plugins/Experimental/ChaosClothAssetEditor/Source/ChaosClothAssetEditorTools/Public/ChaosClothAsset/ClothTransferSkinWeightsTool.h
@@ -87,8 +87,9 @@ protected:
 public:
 
 	virtual bool CanBuildTool(const FToolBuilderState& SceneState) const override;
+	
+	virtual USingleSelectionMeshEditingTool* CreateNewTool(const FToolBuilderState& SceneState) const override;	
 
-	virtual UInteractiveTool* BuildTool(const FToolBuilderState& SceneState) const override;
 	virtual void PostSetupTool(UInteractiveTool* Tool, const FToolBuilderState& SceneState) const override;
 };
 
