ClothEditor: fix for null graph editor widget when adding a new Dataflow asset to the ClothAsset

diff --git a/Engine/Plugins/Experimental/ChaosClothAssetEditor/Source/ChaosClothAssetEditor/Private/ChaosClothAsset/ClothEditorToolkit.cpp b/Engine/Plugins/Experimental/ChaosClothAssetEditor/Source/ChaosClothAssetEditor/Private/ChaosClothAsset/ClothEditorToolkit.cpp
--- a/Engine/Plugins/Experimental/ChaosClothAssetEditor/Source/ChaosClothAssetEditor/Private/ChaosClothAsset/ClothEditorToolkit.cpp
+++ b/Engine/Plugins/Experimental/ChaosClothAssetEditor/Source/ChaosClothAssetEditor/Private/ChaosClothAsset/ClothEditorToolkit.cpp
@@ -798,6 +798,11 @@ void FChaosClothAssetEditorToolkit::ReinitializeGraphEditorWidget()
 	Args._EvaluateGraph = EvalLambda;
 
 	UChaosClothAsset* const ClothAsset = GetAsset();
+
+	if (!GraphEditor)
+	{
+		GraphEditor = CreateGraphEditorWidget();
+	}
 	GraphEditor->Construct(Args, ClothAsset);
 
 	GraphEditor->OnSelectionChangedMulticast.Clear();
