Cloth Editor: providing a default asset name when the user is asked to save a new Dataflow asset

diff --git a/Engine/Plugins/Experimental/ChaosClothAssetEditor/Source/ChaosClothAssetEditor/Private/ChaosClothAsset/AssetDefinition_ClothAsset.cpp b/Engine/Plugins/Experimental/ChaosClothAssetEditor/Source/ChaosClothAssetEditor/Private/ChaosClothAsset/AssetDefinition_ClothAsset.cpp
--- a/Engine/Plugins/Experimental/ChaosClothAssetEditor/Source/ChaosClothAssetEditor/Private/ChaosClothAsset/AssetDefinition_ClothAsset.cpp
+++ b/Engine/Plugins/Experimental/ChaosClothAssetEditor/Source/ChaosClothAssetEditor/Private/ChaosClothAsset/AssetDefinition_ClothAsset.cpp
@@ -28,6 +28,8 @@ namespace ClothAssetDefinitionHelpers
 			{
 				const FString PackageName = ClothAsset->GetOutermost()->GetName();
 				NewDataflowAssetDialogConfig.DefaultPath = FPackageName::GetLongPackagePath(PackageName);
+				const FString ClothName = ClothAsset->GetName();
+				NewDataflowAssetDialogConfig.DefaultAssetName = ClothName + "_Dataflow";
 				NewDataflowAssetDialogConfig.AssetClassNames.Add(DataflowClass->GetClassPathName());
 				NewDataflowAssetDialogConfig.ExistingAssetPolicy = ESaveAssetDialogExistingAssetPolicy::AllowButWarn;
 				NewDataflowAssetDialogConfig.DialogTitleOverride = LOCTEXT("NewDataflowAssetDialogTitle", "Save Dataflow Asset As");
