Cloth Editor: clear the GraphEditor widget when there is no Dataflow  associated with a ClothAsset, and create a new GraphEditor widget when a Dataflow is added

diff --git a/Engine/Plugins/Experimental/ChaosClothAssetEditor/Source/ChaosClothAssetEditor/Private/ChaosClothAsset/ClothEditorToolkit.cpp b/Engine/Plugins/Experimental/ChaosClothAssetEditor/Source/ChaosClothAssetEditor/Private/ChaosClothAsset/ClothEditorToolkit.cpp
--- a/Engine/Plugins/Experimental/ChaosClothAssetEditor/Source/ChaosClothAssetEditor/Private/ChaosClothAsset/ClothEditorToolkit.cpp
+++ b/Engine/Plugins/Experimental/ChaosClothAssetEditor/Source/ChaosClothAssetEditor/Private/ChaosClothAsset/ClothEditorToolkit.cpp
@@ -33,6 +33,7 @@
 #include "PropertyEditorModule.h"
 #include "Algo/RemoveIf.h"
 #include "AdvancedPreviewSceneModule.h"
+#include "Widgets/Layout/SSpacer.h"
 
 #define LOCTEXT_NAMESPACE "ChaosClothAssetEditorToolkit"
 
@@ -631,42 +632,30 @@ TSharedRef<SDockTab> FChaosClothAssetEditorToolkit::SpawnTab_GraphCanvas(const F
 {
 	check(Args.GetTabId() == GraphCanvasTabId);
 
+	SAssignNew(GraphEditorTab, SDockTab)
+		.Label(LOCTEXT("DataflowEditor_Dataflow_TabTitle", "Graph"));
+
 	if (GraphEditor)
 	{
-		TSharedRef<SDockTab> SpawnedTab = SNew(SDockTab)
-			.Label(LOCTEXT("DataflowEditor_Dataflow_TabTitle", "Graph"))
-			[
-				GraphEditor.ToSharedRef()
-			];
-
-		return SpawnedTab;
+		GraphEditorTab.Get()->SetContent(GraphEditor.ToSharedRef());
 	}
-	else
-	{
-		TSharedRef<SDockTab> SpawnedTab = SNew(SDockTab)
-			.Label(LOCTEXT("DataflowEditor_Dataflow_TabTitle", "Graph"));
 
-		return SpawnedTab;
-	}
+	return GraphEditorTab.ToSharedRef();
 }
 
 TSharedRef<SDockTab> FChaosClothAssetEditorToolkit::SpawnTab_NodeDetails(const FSpawnTabArgs& Args)
 {
 	check(Args.GetTabId() == NodeDetailsTabId);
 
+	SAssignNew(NodeDetailsTab, SDockTab)
+		.Label(LOCTEXT("DataflowEditor_NodeDetails_TabTitle", "Node Details"));
+
 	if (NodeDetailsEditor)
 	{
-		return SNew(SDockTab)
-			.Label(LOCTEXT("DataflowEditor_NodeDetails_TabTitle", "Node Details"))
-			[
-				NodeDetailsEditor->GetWidget()->AsShared()
-			];
-	}
-	else
-	{
-		return SNew(SDockTab)
-			.Label(LOCTEXT("DataflowEditor_NodeDetails_TabTitle", "Node Details"));
+		NodeDetailsTab.Get()->SetContent(NodeDetailsEditor->GetWidget()->AsShared());
 	}
+
+	return NodeDetailsTab.ToSharedRef();
 }
 
 void FChaosClothAssetEditorToolkit::InitDetailsViewPanel()
@@ -733,6 +722,14 @@ void FChaosClothAssetEditorToolkit::OnFinishedChangingAssetProperties(const FPro
 				Dataflow->Schema = UDataflowSchema::StaticClass();
 				ReinitializeGraphEditorWidget();
 			}
+			else
+			{
+				// Clear the GraphEditor area
+				// (Can't have a SDataflowGraphEditor with a null UDataflow, so just put down Spacers if we have no Dataflow)
+				GraphEditor.Reset();
+				GraphEditorTab.Get()->SetContent(SNew(SSpacer));
+				NodeDetailsTab.Get()->SetContent(SNew(SSpacer));
+			}
 		}
 	}
 }
@@ -792,18 +789,31 @@ void FChaosClothAssetEditorToolkit::ReinitializeGraphEditorWidget()
 	InEvents.OnVerifyTextCommit = FOnNodeVerifyTextCommit::CreateSP(this, &FChaosClothAssetEditorToolkit::OnNodeVerifyTitleCommit);
 	InEvents.OnTextCommitted = FOnNodeTextCommitted::CreateSP(this, &FChaosClothAssetEditorToolkit::OnNodeTitleCommitted);
 
-	SDataflowGraphEditor::FArguments Args;
-	Args._GraphToEdit = Dataflow;
-	Args._GraphEvents = InEvents;
-	Args._DetailsView = NodeDetailsEditor;
-	Args._EvaluateGraph = EvalLambda;
-
 	UChaosClothAsset* const ClothAsset = GetAsset();
 
 	if (!GraphEditor)
 	{
+		DataflowTerminalPath = ClothEditorToolkitHelpers::GetDataflowTerminalFrom(ClothAsset);
+
+		NodeDetailsEditor = CreateNodeDetailsEditorWidget(ClothAsset);
+		if (NodeDetailsTab.IsValid())
+		{
+			NodeDetailsTab.Get()->SetContent(NodeDetailsEditor->GetWidget().ToSharedRef());
+		}
+
 		GraphEditor = CreateGraphEditorWidget();
+		if (GraphEditorTab.IsValid())
+		{
+			GraphEditorTab.Get()->SetContent(GraphEditor.ToSharedRef());
+		}
 	}
+
+	SDataflowGraphEditor::FArguments Args;
+	Args._GraphToEdit = Dataflow;
+	Args._GraphEvents = InEvents;
+	Args._DetailsView = NodeDetailsEditor;
+	Args._EvaluateGraph = EvalLambda;
+
 	GraphEditor->Construct(Args, ClothAsset);
 
 	GraphEditor->OnSelectionChangedMulticast.Clear();
diff --git a/Engine/Plugins/Experimental/ChaosClothAssetEditor/Source/ChaosClothAssetEditor/Private/ChaosClothAsset/ClothEditorToolkit.h b/Engine/Plugins/Experimental/ChaosClothAssetEditor/Source/ChaosClothAssetEditor/Private/ChaosClothAsset/ClothEditorToolkit.h
--- a/Engine/Plugins/Experimental/ChaosClothAssetEditor/Source/ChaosClothAssetEditor/Private/ChaosClothAsset/ClothEditorToolkit.h
+++ b/Engine/Plugins/Experimental/ChaosClothAssetEditor/Source/ChaosClothAssetEditor/Private/ChaosClothAsset/ClothEditorToolkit.h
@@ -127,11 +127,13 @@ private:
 	void EvaluateNode(FDataflowNode* Node, FDataflowOutput* Out);
 
 	static const FName GraphCanvasTabId;
+	TSharedPtr<SDockTab> GraphEditorTab;
 	TSharedPtr<SDataflowGraphEditor> GraphEditor;
 	TSharedRef<SDataflowGraphEditor> CreateGraphEditorWidget();
 	void ReinitializeGraphEditorWidget();
 
 	static const FName NodeDetailsTabId;
+	TSharedPtr<SDockTab> NodeDetailsTab;
 	TSharedPtr<IStructureDetailsView> NodeDetailsEditor;
 	TSharedPtr<IStructureDetailsView> CreateNodeDetailsEditorWidget(UObject* ObjectToEdit);
 
