Cloth Editor: when loading a cloth asset without a Dataflow graph, allow the user to either create a new one or attach an existing one. Also disallow overwriting an existing Dataflow asset with a new one.

diff --git a/Engine/Plugins/Experimental/ChaosClothAssetEditor/Source/ChaosClothAssetEditor/Private/ChaosClothAsset/AssetDefinition_ClothAsset.cpp b/Engine/Plugins/Experimental/ChaosClothAssetEditor/Source/ChaosClothAssetEditor/Private/ChaosClothAsset/AssetDefinition_ClothAsset.cpp
--- a/Engine/Plugins/Experimental/ChaosClothAssetEditor/Source/ChaosClothAssetEditor/Private/ChaosClothAsset/AssetDefinition_ClothAsset.cpp
+++ b/Engine/Plugins/Experimental/ChaosClothAssetEditor/Source/ChaosClothAssetEditor/Private/ChaosClothAsset/AssetDefinition_ClothAsset.cpp
@@ -12,63 +12,135 @@
 #include "FileHelpers.h"
 #include "ThumbnailRendering/SceneThumbnailInfo.h"
 #include "Misc/WarnIfAssetsLoadedInScope.h"
+#include "Dialog/SMessageDialog.h"
 
 #define LOCTEXT_NAMESPACE "AssetDefinition_ClothAsset"
 
 namespace ClothAssetDefinitionHelpers
 {
-	// Create a new UDataflow if one doesn't already exist for the Cloth Asset
-	UObject* CreateNewDataflowAsset(const UChaosClothAsset* ClothAsset)
+	// Return true if we should proceed, false if we should re-open the dialog
+	bool CreateNewDataflowAsset(const UChaosClothAsset* ClothAsset, UObject*& OutDataflowAsset)
 	{
-		if (FMessageDialog::Open(EAppMsgType::YesNo, LOCTEXT("MissingDataflow", "This Cloth asset currently has no Dataflow graph. Would you like to create a new one?")) == EAppReturnType::Yes)
+		const UClass* const DataflowClass = UDataflow::StaticClass();
+
+		FSaveAssetDialogConfig NewDataflowAssetDialogConfig;
 		{
-			const UClass* const DataflowClass = UDataflow::StaticClass();
+			const FString PackageName = ClothAsset->GetOutermost()->GetName();
+			NewDataflowAssetDialogConfig.DefaultPath = FPackageName::GetLongPackagePath(PackageName);
+			const FString ClothName = ClothAsset->GetName();
+			NewDataflowAssetDialogConfig.DefaultAssetName = ClothName + "_Dataflow";
+			NewDataflowAssetDialogConfig.AssetClassNames.Add(DataflowClass->GetClassPathName());
+			NewDataflowAssetDialogConfig.ExistingAssetPolicy = ESaveAssetDialogExistingAssetPolicy::Disallow;
+			NewDataflowAssetDialogConfig.DialogTitleOverride = LOCTEXT("NewDataflowAssetDialogTitle", "Save Dataflow Asset As");
+		}
+
+		FContentBrowserModule& ContentBrowserModule = FModuleManager::LoadModuleChecked<FContentBrowserModule>("ContentBrowser");
 
-			FSaveAssetDialogConfig NewDataflowAssetDialogConfig;
+		FString NewPackageName;
+		FText OutError;
+		for (bool bFilenameValid = false; !bFilenameValid; bFilenameValid = FFileHelper::IsFilenameValidForSaving(NewPackageName, OutError))
+		{
+			const FString AssetSavePath = ContentBrowserModule.Get().CreateModalSaveAssetDialog(NewDataflowAssetDialogConfig);
+			if (AssetSavePath.IsEmpty())
 			{
-				const FString PackageName = ClothAsset->GetOutermost()->GetName();
-				NewDataflowAssetDialogConfig.DefaultPath = FPackageName::GetLongPackagePath(PackageName);
-				const FString ClothName = ClothAsset->GetName();
-				NewDataflowAssetDialogConfig.DefaultAssetName = ClothName + "_Dataflow";
-				NewDataflowAssetDialogConfig.AssetClassNames.Add(DataflowClass->GetClassPathName());
-				NewDataflowAssetDialogConfig.ExistingAssetPolicy = ESaveAssetDialogExistingAssetPolicy::AllowButWarn;
-				NewDataflowAssetDialogConfig.DialogTitleOverride = LOCTEXT("NewDataflowAssetDialogTitle", "Save Dataflow Asset As");
+				OutDataflowAsset = nullptr;
+				return false;
 			}
+			NewPackageName = FPackageName::ObjectPathToPackageName(AssetSavePath);
+		}
 
-			FContentBrowserModule& ContentBrowserModule = FModuleManager::LoadModuleChecked<FContentBrowserModule>("ContentBrowser");
+		const FName NewAssetName(FPackageName::GetLongPackageAssetName(NewPackageName));
+		UPackage* const NewPackage = CreatePackage(*NewPackageName);
+		UObject* const NewAsset = NewObject<UObject>(NewPackage, DataflowClass, NewAssetName, RF_Public | RF_Standalone | RF_Transactional);
+
+		NewAsset->MarkPackageDirty();
+
+		// Notify the asset registry
+		FAssetRegistryModule::AssetCreated(NewAsset);
+
+		// Save the package
+		TArray<UPackage*> PackagesToSave;
+		PackagesToSave.Add(NewAsset->GetOutermost());
+		constexpr bool bCheckDirty = false;
+		constexpr bool bPromptToSave = false;
+		FEditorFileUtils::EPromptReturnCode ReturnCode = FEditorFileUtils::PromptForCheckoutAndSave(PackagesToSave, bCheckDirty, bPromptToSave);
+
+		if (ReturnCode != FEditorFileUtils::EPromptReturnCode::PR_Success)
+		{
+			OutDataflowAsset = nullptr;
+			return false;
+		}
+
+		OutDataflowAsset = NewAsset;
+		return true;
+	}
 
-			FString NewPackageName;
-			FText OutError;
-			for (bool bFilenameValid = false; !bFilenameValid; bFilenameValid = FFileHelper::IsFilenameValidForSaving(NewPackageName, OutError))
-			{
-				const FString AssetSavePath = ContentBrowserModule.Get().CreateModalSaveAssetDialog(NewDataflowAssetDialogConfig);
-				if (AssetSavePath.IsEmpty())
-				{
-					return nullptr;
-				}
-				NewPackageName = FPackageName::ObjectPathToPackageName(AssetSavePath);
-			}
 
-			const FName NewAssetName(FPackageName::GetLongPackageAssetName(NewPackageName));
-			UPackage* const NewPackage = CreatePackage(*NewPackageName);
-			UObject* const NewAsset = NewObject<UObject>(NewPackage, DataflowClass, NewAssetName, RF_Public | RF_Standalone | RF_Transactional);
+	// Return true if we should proceed, false if we should re-open the dialog
+	bool OpenDataflowAsset(const UChaosClothAsset* ClothAsset, UObject*& OutDataflowAsset)
+	{
+		const UClass* const DataflowClass = UDataflow::StaticClass();
+
+		FOpenAssetDialogConfig NewDataflowAssetDialogConfig;
+		{
+			const FString PackageName = ClothAsset->GetOutermost()->GetName();
+			NewDataflowAssetDialogConfig.DefaultPath = FPackageName::GetLongPackagePath(PackageName);
+			NewDataflowAssetDialogConfig.AssetClassNames.Add(DataflowClass->GetClassPathName());
+			NewDataflowAssetDialogConfig.bAllowMultipleSelection = false;
+			NewDataflowAssetDialogConfig.DialogTitleOverride = LOCTEXT("OpenDataflowAssetDialogTitle", "Open Dataflow Asset");
+		}
+
+		FContentBrowserModule& ContentBrowserModule = FModuleManager::LoadModuleChecked<FContentBrowserModule>("ContentBrowser");
+		TArray<FAssetData> AssetData = ContentBrowserModule.Get().CreateModalOpenAssetDialog(NewDataflowAssetDialogConfig);
+
+		if (AssetData.Num() == 1)
+		{
+			OutDataflowAsset = AssetData[0].GetAsset();
+			return true;
+		}
 
-			NewAsset->MarkPackageDirty();
+		return false;
+	}
 
-			// Notify the asset registry
-			FAssetRegistryModule::AssetCreated(NewAsset);
+	// Return true if we should proceed, false if we should re-open the dialog
+	bool NewOrOpenDialog(const UChaosClothAsset* ClothAsset, UObject*& OutDataflowAsset)
+	{
+		TSharedRef<SMessageDialog> ConfirmDialog = SNew(SMessageDialog)
+			.Title(FText(LOCTEXT("ClothDataflow_WindowTitle", "Create or Open Dataflow graph?")))
+			.Message(LOCTEXT("ClothDataflow_WindowText", "This Cloth Asset currently has no Dataflow graph"))
+			.Buttons({
+				SMessageDialog::FButton(LOCTEXT("ClothDataflow_NewText", "Create new Dataflow")),
+				SMessageDialog::FButton(LOCTEXT("ClothDataflow_OpenText", "Open existing Dataflow")),
+				SMessageDialog::FButton(LOCTEXT("ClothDataflow_ContinueText", "Continue without Dataflow")),
+			});
+
+		const int32 ResultButtonIdx = ConfirmDialog->ShowModal();
+		switch(ResultButtonIdx)
+		{
+		case 0:
+			return CreateNewDataflowAsset(ClothAsset, OutDataflowAsset);
+			break;
+		case 1:
+			return OpenDataflowAsset(ClothAsset, OutDataflowAsset);
+			break;
+		default:
+			break;
+		}
 
-			// Save the package
-			TArray<UPackage*> PackagesToSave;
-			PackagesToSave.Add(NewAsset->GetOutermost());
-			constexpr bool bCheckDirty = false;
-			constexpr bool bPromptToSave = false;
-			FEditorFileUtils::PromptForCheckoutAndSave(PackagesToSave, bCheckDirty, bPromptToSave);
+		return true;
+	}
 
-			return NewAsset;
+	// Create a new UDataflow if one doesn't already exist for the Cloth Asset
+	UObject* NewOrOpenDataflowAsset(const UChaosClothAsset* ClothAsset)
+	{
+		UObject* DataflowAsset = nullptr;
+		bool bDialogDone = false;
+		while (!bDialogDone)
+		{
+			bDialogDone = NewOrOpenDialog(ClothAsset, DataflowAsset);
 		}
 
-		return nullptr;
+		return DataflowAsset;
 	}
 }
 
@@ -116,9 +188,9 @@ EAssetCommandResult UAssetDefinition_ClothAsset::OpenAssets(const FAssetOpenArgs
 
 		if (!ClothAsset->DataflowAsset)
 		{
-			if (UObject* const NewDataflowAsset = ClothAssetDefinitionHelpers::CreateNewDataflowAsset(ClothAsset))
+			if (UDataflow* const NewDataflowAsset = Cast<UDataflow>(ClothAssetDefinitionHelpers::NewOrOpenDataflowAsset(ClothAsset)))
 			{
-				ClothAsset->DataflowAsset = CastChecked<UDataflow>(NewDataflowAsset);
+				ClothAsset->DataflowAsset = NewDataflowAsset;
 			}
 		}
 
