Cloth Editor: fix crash at Editor shutdown due to SharedPtr cycle

diff --git a/Engine/Plugins/Experimental/ChaosClothAssetEditor/Source/ChaosClothAssetEditor/Private/ChaosClothAsset/ClothEditor3DViewportClient.cpp b/Engine/Plugins/Experimental/ChaosClothAssetEditor/Source/ChaosClothAssetEditor/Private/ChaosClothAsset/ClothEditor3DViewportClient.cpp
--- a/Engine/Plugins/Experimental/ChaosClothAssetEditor/Source/ChaosClothAssetEditor/Private/ChaosClothAsset/ClothEditor3DViewportClient.cpp
+++ b/Engine/Plugins/Experimental/ChaosClothAssetEditor/Source/ChaosClothAssetEditor/Private/ChaosClothAsset/ClothEditor3DViewportClient.cpp
@@ -68,11 +68,11 @@ void FChaosClothAssetEditor3DViewportClient::Tick(float DeltaSeconds)
 		return Dataflow::FTimestamp::Invalid;
 	};
 
-	if (ClothToolkit.IsValid())
+	if (TSharedPtr<const FChaosClothAssetEditorToolkit> PinnedClothToolkit = ClothToolkit.Pin())
 	{
-		if (const TSharedPtr<Dataflow::FContext> Context = ClothToolkit->GetDataflowContext())
+		if (const TSharedPtr<Dataflow::FContext> Context = PinnedClothToolkit->GetDataflowContext())
 		{
-			if (const UDataflow* const Dataflow = ClothToolkit->GetDataflow())
+			if (const UDataflow* const Dataflow = PinnedClothToolkit->GetDataflow())
 			{
 				if (UDataflowComponent* const DataflowComponent = ClothEdMode->GetDataflowComponent())
 				{
@@ -133,7 +133,7 @@ void FChaosClothAssetEditor3DViewportClient::SetClothEdMode(TObjectPtr<UChaosClo
 	ClothEdMode = InClothEdMode;
 }
 
-void FChaosClothAssetEditor3DViewportClient::SetClothEditorToolkit(TSharedPtr<const FChaosClothAssetEditorToolkit> InClothToolkit)
+void FChaosClothAssetEditor3DViewportClient::SetClothEditorToolkit(TWeakPtr<const FChaosClothAssetEditorToolkit> InClothToolkit)
 {
 	ClothToolkit = InClothToolkit;
 }
diff --git a/Engine/Plugins/Experimental/ChaosClothAssetEditor/Source/ChaosClothAssetEditor/Private/ChaosClothAsset/ClothEditor3DViewportClient.h b/Engine/Plugins/Experimental/ChaosClothAssetEditor/Source/ChaosClothAssetEditor/Private/ChaosClothAsset/ClothEditor3DViewportClient.h
--- a/Engine/Plugins/Experimental/ChaosClothAssetEditor/Source/ChaosClothAssetEditor/Private/ChaosClothAsset/ClothEditor3DViewportClient.h
+++ b/Engine/Plugins/Experimental/ChaosClothAssetEditor/Source/ChaosClothAssetEditor/Private/ChaosClothAsset/ClothEditor3DViewportClient.h
@@ -33,7 +33,7 @@ public:
 
 	void SetClothComponent(TObjectPtr<UChaosClothComponent> ClothComponent);
 	void SetClothEdMode(TObjectPtr<UChaosClothAssetEditorMode> ClothEdMode);
-	void SetClothEditorToolkit(TSharedPtr<const FChaosClothAssetEditorToolkit> ClothToolkit);
+	void SetClothEditorToolkit(TWeakPtr<const FChaosClothAssetEditorToolkit> ClothToolkit);
 
 	void SoftResetSimulation();
 	void HardResetSimulation();
@@ -60,7 +60,7 @@ private:
 
 	TObjectPtr<UChaosClothAssetEditorMode> ClothEdMode;
 
-	TSharedPtr<const FChaosClothAssetEditorToolkit> ClothToolkit;
+	TWeakPtr<const FChaosClothAssetEditorToolkit> ClothToolkit;
 	
 	bool bSimMeshWireframe = true;
 	bool bRenderMeshWireframe = false;
