Cloth Editor: adding a Dataflow node to read a Datasmith file and extract ClothAsset

diff --git a/Engine/Plugins/Experimental/ChaosClothAsset/Source/ChaosClothAsset/Public/ChaosClothAsset/ClothCollection.h b/Engine/Plugins/Experimental/ChaosClothAsset/Source/ChaosClothAsset/Public/ChaosClothAsset/ClothCollection.h
--- a/Engine/Plugins/Experimental/ChaosClothAsset/Source/ChaosClothAsset/Public/ChaosClothAsset/ClothCollection.h
+++ b/Engine/Plugins/Experimental/ChaosClothAsset/Source/ChaosClothAsset/Public/ChaosClothAsset/ClothCollection.h
@@ -27,6 +27,8 @@ namespace UE::Chaos::ClothAsset
 		using Super::Serialize;
 		void Serialize(FArchive& Ar);
 
+		virtual void Reset() override { Super::Reset(); Construct(); }
+
 		/** Set the number of elements to one of the groups that have start/end indices while maintaining the correct order of the data. */
 		int32 SetNumElements(int32 InNumElements, const FName& GroupName, TManagedArray<int32>& StartArray, TManagedArray<int32>& EndArray, int32 StartEndIndex);
 
diff --git a/Engine/Plugins/Experimental/ChaosClothAsset/Source/ChaosClothAssetEngine/Private/ChaosClothAsset/ClothAsset.cpp b/Engine/Plugins/Experimental/ChaosClothAsset/Source/ChaosClothAssetEngine/Private/ChaosClothAsset/ClothAsset.cpp
--- a/Engine/Plugins/Experimental/ChaosClothAsset/Source/ChaosClothAssetEngine/Private/ChaosClothAsset/ClothAsset.cpp
+++ b/Engine/Plugins/Experimental/ChaosClothAsset/Source/ChaosClothAssetEngine/Private/ChaosClothAsset/ClothAsset.cpp
@@ -327,6 +327,13 @@ void UChaosClothAsset::Build()
 	// Release render resources
 	ReleaseResources();
 
+	// Flush the resource release commands to the rendering thread to ensure that the build doesn't occur while a resource is still
+	// allocated, and potentially accessing the UClothAsset
+	ReleaseResourcesFence.Wait();
+
+	// Clear SkeletalMeshRenderData so it can be rebuilt
+	SetResourceForRendering(nullptr);
+
 	// Set a new Guid to invalidate the DDC
 	AssetGuid = FGuid::NewGuid();
 
diff --git a/Engine/Plugins/Experimental/ChaosClothAssetEditor/ChaosClothAssetEditor.uplugin b/Engine/Plugins/Experimental/ChaosClothAssetEditor/ChaosClothAssetEditor.uplugin
--- a/Engine/Plugins/Experimental/ChaosClothAssetEditor/ChaosClothAssetEditor.uplugin
+++ b/Engine/Plugins/Experimental/ChaosClothAssetEditor/ChaosClothAssetEditor.uplugin
@@ -59,7 +59,10 @@
 		{
 			"Name": "Dataflow",
 			"Enabled": true
+		},
+		{
+			"Name": "DatasmithImporter",
+			"Enabled": true
 		}
-
 	]
 }
\ No newline at end of file
diff --git a/Engine/Plugins/Experimental/ChaosClothAssetEditor/Source/ChaosClothAssetEditor/ChaosClothAssetEditor.Build.cs b/Engine/Plugins/Experimental/ChaosClothAssetEditor/Source/ChaosClothAssetEditor/ChaosClothAssetEditor.Build.cs
--- a/Engine/Plugins/Experimental/ChaosClothAssetEditor/Source/ChaosClothAssetEditor/ChaosClothAssetEditor.Build.cs
+++ b/Engine/Plugins/Experimental/ChaosClothAssetEditor/Source/ChaosClothAssetEditor/ChaosClothAssetEditor.Build.cs
@@ -63,6 +63,8 @@ public class ChaosClothAssetEditor : ModuleRules
 				"TypedElementFramework",
 				"TypedElementRuntime",
 				"AssetDefinition",
+				"DatasmithImporter",
+				"ExternalSource"
 				// ... add private dependencies that you statically link with here ...	
 			}
 			);
diff --git a/Engine/Plugins/Experimental/ChaosClothAssetEditor/Source/ChaosClothAssetEditor/Private/ChaosClothAsset/ClothDataflowNodes.cpp b/Engine/Plugins/Experimental/ChaosClothAssetEditor/Source/ChaosClothAssetEditor/Private/ChaosClothAsset/ClothDataflowNodes.cpp
--- a/Engine/Plugins/Experimental/ChaosClothAssetEditor/Source/ChaosClothAssetEditor/Private/ChaosClothAsset/ClothDataflowNodes.cpp
+++ b/Engine/Plugins/Experimental/ChaosClothAssetEditor/Source/ChaosClothAssetEditor/Private/ChaosClothAsset/ClothDataflowNodes.cpp
@@ -2,6 +2,14 @@
 
 #include "ChaosClothAsset/ClothDataflowNodes.h"
 #include "Logging/LogMacros.h"
+#include "GeometryCollection/ManagedArrayCollection.h"
+#include "SourceUri.h"
+#include "ExternalSource.h"
+#include "ExternalSourceModule.h"
+#include "DatasmithImportContext.h"
+#include "DatasmithImportFactory.h"
+#include "ChaosClothAsset/ClothAsset.h"
+#include "ChaosClothAsset/ClothCollection.h"
 
 #include UE_INLINE_GENERATED_CPP_BY_NAME(ClothDataflowNodes)
 
@@ -12,6 +20,135 @@ namespace Dataflow
 	void RegisterClothDataflowNodes()
 	{
 		DATAFLOW_NODE_REGISTER_CREATION_FACTORY(FClothAssetTerminalDataflowNode);
+		DATAFLOW_NODE_REGISTER_CREATION_FACTORY(FClothAssetDatasmithImportNode);
+	}
+}
+
+FClothAssetTerminalDataflowNode::FClothAssetTerminalDataflowNode(const Dataflow::FNodeParameters& InParam, FGuid InGuid) :
+	FDataflowTerminalNode(InParam, InGuid)
+{
+	RegisterInputConnection(&Collection);
+	RegisterOutputConnection(&Collection, &Collection);
+}
+
+
+void FClothAssetTerminalDataflowNode::SetAssetValue(TObjectPtr<UObject> Asset, Dataflow::FContext& Context) const
+{
+	if (UChaosClothAsset* const ClothAsset = Cast<UChaosClothAsset>(Asset.Get()))
+	{
+		const FManagedArrayCollection& InCollection = GetValue<FManagedArrayCollection>(Context, &Collection);
+		if (InCollection.HasGroup(UE::Chaos::ClothAsset::FClothCollection::LodsGroup))  // TODO: SkeletalMeshRenderData crashes with empty collection
+		{
+			ClothAsset->GetClothCollection()->Reset();
+			
+			InCollection.CopyTo(ClothAsset->GetClothCollection().Get());
+
+			// Set the render mesh to duplicate the sim mesh
+			constexpr int32 MaterialId = 0;	
+			
+			ClothAsset->CopySimMeshToRenderMesh(MaterialId);  // TODO: Make this a node (needs to make cloth collection an adapter)
+
+			// Rebuild the asset static data
+			ClothAsset->Build();
+		}
+	}
+}
+
+void FClothAssetTerminalDataflowNode::Evaluate(Dataflow::FContext& Context, const FDataflowOutput* Out) const
+{
+	const FManagedArrayCollection& InCollection = GetValue<FManagedArrayCollection>(Context, &Collection);
+	SetValue<FManagedArrayCollection>(Context, InCollection, &Collection);
+}
+
+FClothAssetDatasmithImportNode::FClothAssetDatasmithImportNode(const Dataflow::FNodeParameters& InParam, FGuid InGuid)
+	: FDataflowNode(InParam, InGuid),
+	DestPackageName("/Game/ClothAsset")
+{
+	RegisterInputConnection(&DatasmithFile);
+	RegisterInputConnection(&DestPackageName);
+	RegisterOutputConnection(&Collection);
+}
+
+
+bool FClothAssetDatasmithImportNode::EvaluateImpl(Dataflow::FContext& Context, FManagedArrayCollection& OutCollection) const
+{
+	using namespace UE::DatasmithImporter;
+
+	const FFilePath& InFilePath = GetValue<FFilePath>(Context, &DatasmithFile);
+	const FString& InDestPackageName = GetValue<FString>(Context, &DestPackageName);
+
+	const FSourceUri SourceUri = FSourceUri::FromFilePath(InFilePath.FilePath);
+	const TSharedPtr<FExternalSource> ExternalSource = IExternalSourceModule::GetOrCreateExternalSource(SourceUri);
+	if (!ExternalSource)
+	{
+		return false;
+	}
+
+	constexpr bool bLoadConfig = false; // TODO: Not sure what this does
+	const FName LoggerName(TEXT("ImportDatasmithClothNode"));
+	const FText LoggerLabel(NSLOCTEXT("ImportDatasmithClothNode", "LoggerLabel", "ImportDatasmithClothNode"));
+	FDatasmithImportContext DatasmithImportContext(ExternalSource, bLoadConfig, LoggerName, LoggerLabel);
+
+	const TStrongObjectPtr<const UPackage> DestinationPackage(CreatePackage(*InDestPackageName));		// TODO: would be nice to not have to create the package but it seems necessary for now
+	if (!ensure(DestinationPackage))
+	{
+		// Failed to create the package to hold this asset for some reason
+		return false;
+	}
+	
+	// Don't create the Actors in the level, just read the Assets
+	DatasmithImportContext.Options->BaseOptions.SceneHandling = EDatasmithImportScene::AssetsOnly;
+
+	constexpr EObjectFlags NewObjectFlags = RF_Public | RF_Standalone | RF_Transactional;
+	const TSharedPtr<FJsonObject> ImportSettingsJson;
+	constexpr bool bIsSilent = true;
+	const FString DestinationPath = DestinationPackage->GetName();
+	if (!DatasmithImportContext.Init(DestinationPath, NewObjectFlags, GWarn, ImportSettingsJson, bIsSilent))
+	{
+		return false;
+	}
+
+	if (const TSharedPtr<IDatasmithScene> LoadedScene = ExternalSource->TryLoad())
+	{
+		DatasmithImportContext.InitScene(LoadedScene.ToSharedRef());
+	}
+	else
+	{
+		return false;
+	}
+
+	bool bUserCancelled = false;
+	bool bImportSucceed = DatasmithImportFactoryImpl::ImportDatasmithScene(DatasmithImportContext, bUserCancelled);
+	bImportSucceed &= !bUserCancelled;
+
+	if (bImportSucceed && DatasmithImportContext.ImportedClothes.Num() > 0)
+	{
+		const UObject* const ClothObject = DatasmithImportContext.ImportedClothes.CreateIterator()->Value;
+		const UChaosClothAsset* const DatasmithClothAsset = Cast<UChaosClothAsset>(ClothObject);
+		ensure(DatasmithClothAsset);
+
+		DatasmithClothAsset->GetClothCollection()->CopyTo(&OutCollection);
+
+		return true;
+	}
+
+	return false;
+}
+
+void FClothAssetDatasmithImportNode::Evaluate(Dataflow::FContext& Context, const FDataflowOutput* Out) const
+{
+	FManagedArrayCollection OutCollection;
+	const bool bSuccess = EvaluateImpl(Context, OutCollection);
+
+	if (bSuccess)
+	{
+		SetValue<FManagedArrayCollection>(Context, OutCollection, &Collection);
+	}
+	else
+	{
+		// Init with an empty cloth collection
+		UE::Chaos::ClothAsset::FClothCollection ClothCollection;
+		SetValue<FManagedArrayCollection>(Context, ClothCollection, &Collection);
 	}
 }
 
diff --git a/Engine/Plugins/Experimental/ChaosClothAssetEditor/Source/ChaosClothAssetEditor/Private/ChaosClothAsset/ClothDataflowNodes.h b/Engine/Plugins/Experimental/ChaosClothAssetEditor/Source/ChaosClothAssetEditor/Private/ChaosClothAsset/ClothDataflowNodes.h
--- a/Engine/Plugins/Experimental/ChaosClothAssetEditor/Source/ChaosClothAssetEditor/Private/ChaosClothAsset/ClothDataflowNodes.h
+++ b/Engine/Plugins/Experimental/ChaosClothAssetEditor/Source/ChaosClothAssetEditor/Private/ChaosClothAsset/ClothDataflowNodes.h
@@ -3,45 +3,62 @@
 #pragma once 
 
 #include "Dataflow/DataflowCore.h"
-#include "ChaosClothAsset/ClothAsset.h"
-#include "ChaosClothAsset/ClothCollection.h"
-#include "GeometryCollection/GeometryCollection.h"
+#include "GeometryCollection/ManagedArrayCollection.h"
 #include "ClothDataflowNodes.generated.h"
 
-USTRUCT()
+
+USTRUCT(meta = (DataflowCloth))
 struct FClothAssetTerminalDataflowNode : public FDataflowTerminalNode
 {
 	GENERATED_USTRUCT_BODY()
-	DATAFLOW_NODE_DEFINE_INTERNAL(FClothAssetTerminalDataflowNode, "ClothAssetTerminal", "Cloth", "")
-	DATAFLOW_NODE_RENDER_TYPE(FGeometryCollection::StaticType(), "Collection")
+	DATAFLOW_NODE_DEFINE_INTERNAL(FClothAssetTerminalDataflowNode, "ClothAssetTerminal", "Cloth", "Cloth Terminal")
+	//DATAFLOW_NODE_RENDER_TYPE(FGeometryCollection::StaticType(), "Collection")
 
 public:
 
 	UPROPERTY(meta = (DataflowInput, DataflowOutput, DisplayName = "Collection", DataflowPassthrough = "Collection"))
-	FManagedArrayCollection Collection;		// TODO: Replace with FClothCollection
-
-	FClothAssetTerminalDataflowNode(const Dataflow::FNodeParameters& InParam, FGuid InGuid = FGuid::NewGuid())
-		: FDataflowTerminalNode(InParam, InGuid)
-	{
-		RegisterInputConnection(&Collection);
-		RegisterOutputConnection(&Collection, &Collection);
-	}
-
-	virtual void SetAssetValue(TObjectPtr<UObject> Asset, Dataflow::FContext& Context) const override
-	{
-		if (UChaosClothAsset* ClothAsset = Cast<UChaosClothAsset>(Asset.Get()))
-		{
-			// TODO: Modify ClothAsset directly
-		}
-	}
-
-	virtual void Evaluate(Dataflow::FContext& Context, const FDataflowOutput* Out) const override
-	{
-		const FManagedArrayCollection& InCollection = GetValue<FManagedArrayCollection>(Context, &Collection);
-		SetValue<FManagedArrayCollection>(Context, InCollection, &Collection);
-	}
+	FManagedArrayCollection Collection;
+
+	FClothAssetTerminalDataflowNode(const Dataflow::FNodeParameters& InParam, FGuid InGuid = FGuid::NewGuid());
+
+	virtual void SetAssetValue(TObjectPtr<UObject> OutAsset, Dataflow::FContext& Context) const override;
+
+	virtual void Evaluate(Dataflow::FContext& Context, const FDataflowOutput* Out) const override;
 };
 
+
+USTRUCT(meta = (DataflowFlesh))
+struct FClothAssetDatasmithImportNode : public FDataflowNode
+{
+	GENERATED_USTRUCT_BODY()
+	DATAFLOW_NODE_DEFINE_INTERNAL(FClothAssetDatasmithImportNode, "Import", "Cloth", "Cloth Datasmith Import")
+	//DATAFLOW_NODE_RENDER_TYPE(FGeometryCollection::StaticType(), "Collection")
+
+public:
+	
+	// Datasmith file to read from
+	UPROPERTY(EditAnywhere, Category = "Datasmith Importer")
+	FFilePath DatasmithFile;
+
+	// Package to import into
+	UPROPERTY(EditAnywhere, Category = "Datasmith Importer")
+	FString DestPackageName;
+
+	// ManagedArrayCollection for the first ClothAsset found in the input Datasmith file
+	UPROPERTY(meta = (DataflowOutput, DisplayName = "Collection"))
+	FManagedArrayCollection Collection;
+
+	FClothAssetDatasmithImportNode(const Dataflow::FNodeParameters& InParam, FGuid InGuid = FGuid::NewGuid());
+
+	virtual void Evaluate(Dataflow::FContext& Context, const FDataflowOutput* Out) const override;
+
+private:
+
+	bool EvaluateImpl(Dataflow::FContext& Context, FManagedArrayCollection& OutCollection) const;
+
+};
+
+
 namespace Dataflow
 {
 	void RegisterClothDataflowNodes();
diff --git a/Engine/Plugins/Experimental/ChaosClothAssetEditor/Source/ChaosClothAssetEditor/Private/ChaosClothAsset/ClothEditorToolkit.cpp b/Engine/Plugins/Experimental/ChaosClothAssetEditor/Source/ChaosClothAssetEditor/Private/ChaosClothAsset/ClothEditorToolkit.cpp
--- a/Engine/Plugins/Experimental/ChaosClothAssetEditor/Source/ChaosClothAssetEditor/Private/ChaosClothAsset/ClothEditorToolkit.cpp
+++ b/Engine/Plugins/Experimental/ChaosClothAssetEditor/Source/ChaosClothAssetEditor/Private/ChaosClothAsset/ClothEditorToolkit.cpp
@@ -686,6 +686,11 @@ void FChaosClothAssetEditorToolkit::InitDetailsViewPanel()
 		PreviewSceneDockTab->SetContent(AdvancedPreviewSettingsWidget.ToSharedRef());
 	}
 
+	PopulateOutliner();
+}
+
+void FChaosClothAssetEditorToolkit::PopulateOutliner()
+{
 	TSharedPtr<UE::Chaos::ClothAsset::FClothCollection> ClothCollection;
 	if (const TObjectPtr<UChaosClothComponent> ClothComponent = ClothPreviewScene->ClothComponent)
 	{
@@ -746,7 +751,9 @@ void FChaosClothAssetEditorToolkit::EvaluateNode(FDataflowNode* Node, FDataflowO
 		}
 		LastDataflowNodeTimestamp = Dataflow::FTimestamp::Invalid;
 
-		FDataflowEditorCommands::EvaluateTerminalNode(*DataflowContext.Get(), LastDataflowNodeTimestamp, Dataflow, Node, Out, Asset, DataflowTerminalPath);
+		FDataflowEditorCommands::EvaluateTerminalNode(*DataflowContext.Get(), LastDataflowNodeTimestamp, Dataflow, Node, nullptr, Asset, DataflowTerminalPath);
+
+		OnClothAssetChanged();
 	}
 };
 
@@ -893,4 +900,24 @@ void FChaosClothAssetEditorToolkit::OnNodeSelectionChanged(const TSet<UObject*>&
 
 //~ Ends DataflowEditorActions
 
+void FChaosClothAssetEditorToolkit::OnClothAssetChanged()
+{
+	TArray<TObjectPtr<UObject>> ObjectsToEdit;
+	OwningAssetEditor->GetObjectsToEdit(ObjectsToEdit);
+
+	UChaosClothAssetEditorMode* const ClothMode = CastChecked<UChaosClothAssetEditorMode>(EditorModeManager->GetActiveScriptableMode(UChaosClothAssetEditorMode::EM_ChaosClothAssetEditorModeId));
+	ClothMode->InitializeTargets(ObjectsToEdit);
+
+	if (UChaosClothAsset* const ClothAsset = Cast<UChaosClothAsset>(ObjectsToEdit[0]))
+	{
+		ClothPreviewScene->CreateClothActor(ClothAsset);
+
+		ensure(ClothAsset->HasAnyFlags(RF_Transactional));		// Ensure all objects are transactable for undo/redo in the details panel
+		SetEditingObject(ClothAsset);
+
+		PopulateOutliner();
+	}
+}
+
+
 #undef LOCTEXT_NAMESPACE
diff --git a/Engine/Plugins/Experimental/ChaosClothAssetEditor/Source/ChaosClothAssetEditor/Private/ChaosClothAsset/ClothEditorToolkit.h b/Engine/Plugins/Experimental/ChaosClothAssetEditor/Source/ChaosClothAssetEditor/Private/ChaosClothAsset/ClothEditorToolkit.h
--- a/Engine/Plugins/Experimental/ChaosClothAssetEditor/Source/ChaosClothAssetEditor/Private/ChaosClothAsset/ClothEditorToolkit.h
+++ b/Engine/Plugins/Experimental/ChaosClothAssetEditor/Source/ChaosClothAssetEditor/Private/ChaosClothAsset/ClothEditorToolkit.h
@@ -101,6 +101,10 @@ private:
 	void InitDetailsViewPanel();
 	void OnFinishedChangingAssetProperties(const FPropertyChangedEvent&);
 
+	void PopulateOutliner();
+	void OnClothAssetChanged();
+
+
 	/** Scene in which the 3D sim space preview meshes live. Ownership shared with AdvancedPreviewSettingsWidget*/
 	TSharedPtr<FChaosClothPreviewScene> ClothPreviewScene;
 
@@ -120,6 +124,7 @@ private:
 	TSharedPtr<SComboBox<FName>> SelectedGroupNameComboBox;
 	TArray<FName> ClothCollectionGroupNames;		// Data source for SelectedGroupNameComboBox
 
+
 	//~ Begin Dataflow support
 	 
 	UDataflow* Dataflow = nullptr;
@@ -148,6 +153,5 @@ private:
 	void OnNodeSelectionChanged(const TSet<UObject*>& NewSelection) const;
 
 	//~ End Dataflow support
-
 };
 
