Fix CIS Static Analyzer issue 259793

diff --git a/Engine/Plugins/Experimental/ChaosClothAssetEditor/Source/ChaosClothAssetEditor/Private/ChaosClothAsset/ClothDataflowNodes.cpp b/Engine/Plugins/Experimental/ChaosClothAssetEditor/Source/ChaosClothAssetEditor/Private/ChaosClothAsset/ClothDataflowNodes.cpp
--- a/Engine/Plugins/Experimental/ChaosClothAssetEditor/Source/ChaosClothAssetEditor/Private/ChaosClothAsset/ClothDataflowNodes.cpp
+++ b/Engine/Plugins/Experimental/ChaosClothAssetEditor/Source/ChaosClothAssetEditor/Private/ChaosClothAsset/ClothDataflowNodes.cpp
@@ -125,11 +125,11 @@ bool FClothAssetDatasmithImportNode::EvaluateImpl(Dataflow::FContext& Context, F
 	{
 		const UObject* const ClothObject = DatasmithImportContext.ImportedClothes.CreateIterator()->Value;
 		const UChaosClothAsset* const DatasmithClothAsset = Cast<UChaosClothAsset>(ClothObject);
-		ensure(DatasmithClothAsset);
-
-		DatasmithClothAsset->GetClothCollection()->CopyTo(&OutCollection);
-
-		return true;
+		if (ensure(DatasmithClothAsset))
+		{
+			DatasmithClothAsset->GetClothCollection()->CopyTo(&OutCollection);
+			return true;
+		}
 	}
 
 	return false;
