Cloth Editor: fix ensure when shutting down the editor. Make sure to remove the 3D viewport gizmo from the Gizmo Manager before the Gizmo Manager is destroyed.

diff --git a/Engine/Plugins/Experimental/ChaosClothAssetEditor/Source/ChaosClothAssetEditor/Private/ChaosClothAsset/ClothEditor3DViewportClient.cpp b/Engine/Plugins/Experimental/ChaosClothAssetEditor/Source/ChaosClothAssetEditor/Private/ChaosClothAsset/ClothEditor3DViewportClient.cpp
--- a/Engine/Plugins/Experimental/ChaosClothAssetEditor/Source/ChaosClothAssetEditor/Private/ChaosClothAsset/ClothEditor3DViewportClient.cpp
+++ b/Engine/Plugins/Experimental/ChaosClothAssetEditor/Source/ChaosClothAssetEditor/Private/ChaosClothAsset/ClothEditor3DViewportClient.cpp
@@ -39,14 +39,28 @@ FChaosClothAssetEditor3DViewportClient::FChaosClothAssetEditor3DViewportClient(F
 	const FString GizmoIdentifier = TEXT("ChaosClothAssetEditor3DViewportClientGizmoIdentifier");
 	Gizmo = GizmoManager->Create3AxisTransformGizmo(this, GizmoIdentifier);
 
-	UInteractiveToolManager* const ToolManager = InteractiveToolsContext->ToolManager;
-	Gizmo->SetActiveTarget(TransformProxy, ToolManager);
+	Gizmo->SetActiveTarget(TransformProxy);
 	Gizmo->SetVisibility(true);
 	Gizmo->bUseContextGizmoMode = false;
 	Gizmo->bUseContextCoordinateSystem = false;
 	Gizmo->ActiveGizmoMode = EToolContextTransformGizmoMode::Combined;
 }
 
+FChaosClothAssetEditor3DViewportClient::~FChaosClothAssetEditor3DViewportClient()
+{
+	DeleteViewportGizmo();
+}
+
+void FChaosClothAssetEditor3DViewportClient::DeleteViewportGizmo()
+{
+	if (Gizmo && ModeTools && ModeTools->GetInteractiveToolsContext() && ModeTools->GetInteractiveToolsContext()->GizmoManager)
+	{
+		ModeTools->GetInteractiveToolsContext()->GizmoManager->DestroyGizmo(Gizmo);
+	}
+	Gizmo = nullptr;
+	TransformProxy = nullptr;
+}
+
 void FChaosClothAssetEditor3DViewportClient::AddReferencedObjects(FReferenceCollector& Collector)
 {
 	FEditorViewportClient::AddReferencedObjects(Collector);
diff --git a/Engine/Plugins/Experimental/ChaosClothAssetEditor/Source/ChaosClothAssetEditor/Private/ChaosClothAsset/ClothEditor3DViewportClient.h b/Engine/Plugins/Experimental/ChaosClothAssetEditor/Source/ChaosClothAssetEditor/Private/ChaosClothAsset/ClothEditor3DViewportClient.h
--- a/Engine/Plugins/Experimental/ChaosClothAssetEditor/Source/ChaosClothAssetEditor/Private/ChaosClothAsset/ClothEditor3DViewportClient.h
+++ b/Engine/Plugins/Experimental/ChaosClothAssetEditor/Source/ChaosClothAssetEditor/Private/ChaosClothAsset/ClothEditor3DViewportClient.h
@@ -23,7 +23,10 @@ public:
 	FChaosClothAssetEditor3DViewportClient(FEditorModeTools* InModeTools, FPreviewScene* InPreviewScene = nullptr,
 		const TWeakPtr<SEditorViewport>& InEditorViewportWidget = nullptr);
 
-	virtual ~FChaosClothAssetEditor3DViewportClient() = default;
+	virtual ~FChaosClothAssetEditor3DViewportClient();
+
+	// Delete the viewport gizmo and transform proxy
+	void DeleteViewportGizmo();
 
 	void EnableSimMeshWireframe(bool bEnable ) { bSimMeshWireframe = bEnable; }
 	bool SimMeshWireframeEnabled() const { return bSimMeshWireframe; }
diff --git a/Engine/Plugins/Experimental/ChaosClothAssetEditor/Source/ChaosClothAssetEditor/Private/ChaosClothAsset/ClothEditorToolkit.cpp b/Engine/Plugins/Experimental/ChaosClothAssetEditor/Source/ChaosClothAssetEditor/Private/ChaosClothAsset/ClothEditorToolkit.cpp
--- a/Engine/Plugins/Experimental/ChaosClothAssetEditor/Source/ChaosClothAssetEditor/Private/ChaosClothAsset/ClothEditorToolkit.cpp
+++ b/Engine/Plugins/Experimental/ChaosClothAssetEditor/Source/ChaosClothAssetEditor/Private/ChaosClothAsset/ClothEditorToolkit.cpp
@@ -193,6 +193,13 @@ FChaosClothAssetEditorToolkit::FChaosClothAssetEditorToolkit(UAssetEditor* InOwn
 
 FChaosClothAssetEditorToolkit::~FChaosClothAssetEditorToolkit()
 {
+	if (ClothPreviewViewportClient)
+	{
+		// Delete the gizmo in the viewport before deleting the EditorModeManager. The Gizmo Manager can get tripped up if it gets deleted
+		// while it still has active gizmos.
+		ClothPreviewViewportClient->DeleteViewportGizmo();
+	}
+
 	// We need to force the cloth editor mode deletion now because otherwise the preview and rest-space worlds
 	// will end up getting destroyed before the mode's Exit() function gets to run, and we'll get some
 	// warnings when we destroy any mode actors.
