Cloth Editor: fix for "Post Processing Enabled" check box having no effect

diff --git a/Engine/Plugins/Experimental/ChaosClothAssetEditor/Source/ChaosClothAssetEditor/Private/ChaosClothAsset/ClothEditor3DViewportClient.cpp b/Engine/Plugins/Experimental/ChaosClothAssetEditor/Source/ChaosClothAssetEditor/Private/ChaosClothAsset/ClothEditor3DViewportClient.cpp
--- a/Engine/Plugins/Experimental/ChaosClothAssetEditor/Source/ChaosClothAssetEditor/Private/ChaosClothAsset/ClothEditor3DViewportClient.cpp
+++ b/Engine/Plugins/Experimental/ChaosClothAssetEditor/Source/ChaosClothAssetEditor/Private/ChaosClothAsset/ClothEditor3DViewportClient.cpp
@@ -16,11 +16,14 @@
 #include "Engine/Selection.h"
 #include "EngineUtils.h"
 #include "Animation/SkeletalMeshActor.h"
+#include "AssetViewerSettings.h"
+#include "Editor/EditorPerProjectUserSettings.h"
 
 FChaosClothAssetEditor3DViewportClient::FChaosClothAssetEditor3DViewportClient(FEditorModeTools* InModeTools,
-	FPreviewScene* InPreviewScene, 
+	TSharedPtr<FChaosClothPreviewScene> InPreviewScene,
 	const TWeakPtr<SEditorViewport>& InEditorViewportWidget)
-	: FEditorViewportClient(InModeTools, InPreviewScene, InEditorViewportWidget)
+	: FEditorViewportClient(InModeTools, InPreviewScene.Get(), InEditorViewportWidget),
+	  ClothPreviewScene(InPreviewScene)
 {
 	// We want our near clip plane to be quite close so that we can zoom in further.
 	OverrideNearClipPlane(KINDA_SMALL_NUMBER);
@@ -44,11 +47,23 @@ FChaosClothAssetEditor3DViewportClient::FChaosClothAssetEditor3DViewportClient(F
 	Gizmo->bUseContextGizmoMode = false;
 	Gizmo->bUseContextCoordinateSystem = false;
 	Gizmo->ActiveGizmoMode = EToolContextTransformGizmoMode::Combined;
+
+	// Set correct flags according to current profile settings
+	SetAdvancedShowFlagsForScene(UAssetViewerSettings::Get()->Profiles[GetMutableDefault<UEditorPerProjectUserSettings>()->AssetViewerProfileIndex].bPostProcessingEnabled);
+}
+
+void FChaosClothAssetEditor3DViewportClient::RegisterSettingsChangedDelegate()
+{
+	// Remove any existing delegate in case this function is called twice
+	UAssetViewerSettings::Get()->OnAssetViewerSettingsChanged().RemoveAll(this);
+	UAssetViewerSettings::Get()->OnAssetViewerSettingsChanged().AddSP(this, &FChaosClothAssetEditor3DViewportClient::OnAssetViewerSettingsChanged);
 }
 
 FChaosClothAssetEditor3DViewportClient::~FChaosClothAssetEditor3DViewportClient()
 {
 	DeleteViewportGizmo();
+
+	UAssetViewerSettings::Get()->OnAssetViewerSettingsChanged().RemoveAll(this);
 }
 
 void FChaosClothAssetEditor3DViewportClient::DeleteViewportGizmo()
@@ -299,3 +314,33 @@ void FChaosClothAssetEditor3DViewportClient::ProcessClick(FSceneView& View, HHit
 	}
 
 }
+
+
+void FChaosClothAssetEditor3DViewportClient::OnAssetViewerSettingsChanged(const FName& InPropertyName)
+{
+	if (InPropertyName == GET_MEMBER_NAME_CHECKED(FPreviewSceneProfile, bPostProcessingEnabled) || InPropertyName == NAME_None)
+	{
+		const UAssetViewerSettings* const Settings = UAssetViewerSettings::Get();
+		const TSharedPtr<const FChaosClothPreviewScene> PinnedClothPreviewScene = ClothPreviewScene.Pin();
+		if (Settings && PinnedClothPreviewScene)
+		{
+			const int32 ProfileIndex = PinnedClothPreviewScene->GetCurrentProfileIndex();
+			if (Settings->Profiles.IsValidIndex(ProfileIndex))
+			{
+				SetAdvancedShowFlagsForScene(Settings->Profiles[ProfileIndex].bPostProcessingEnabled);
+			}
+		}
+	}
+}
+
+void FChaosClothAssetEditor3DViewportClient::SetAdvancedShowFlagsForScene(const bool bAdvancedShowFlags)
+{
+	if (bAdvancedShowFlags)
+	{
+		EngineShowFlags.EnableAdvancedFeatures();
+	}
+	else
+	{
+		EngineShowFlags.DisableAdvancedFeatures();
+	}
+}
\ No newline at end of file
diff --git a/Engine/Plugins/Experimental/ChaosClothAssetEditor/Source/ChaosClothAssetEditor/Private/ChaosClothAsset/ClothEditor3DViewportClient.h b/Engine/Plugins/Experimental/ChaosClothAssetEditor/Source/ChaosClothAssetEditor/Private/ChaosClothAsset/ClothEditor3DViewportClient.h
--- a/Engine/Plugins/Experimental/ChaosClothAssetEditor/Source/ChaosClothAssetEditor/Private/ChaosClothAsset/ClothEditor3DViewportClient.h
+++ b/Engine/Plugins/Experimental/ChaosClothAssetEditor/Source/ChaosClothAssetEditor/Private/ChaosClothAsset/ClothEditor3DViewportClient.h
@@ -11,18 +11,22 @@ class UChaosClothAssetEditorMode;
 class FChaosClothAssetEditorToolkit;
 class UTransformProxy;
 class UCombinedTransformGizmo;
+class FChaosClothPreviewScene;
 
 /**
  * Viewport client for the 3d sim preview in the cloth editor. Currently same as editor viewport
  * client but doesn't allow editor gizmos/widgets.
  */
-class CHAOSCLOTHASSETEDITOR_API FChaosClothAssetEditor3DViewportClient : public FEditorViewportClient
+class CHAOSCLOTHASSETEDITOR_API FChaosClothAssetEditor3DViewportClient : public FEditorViewportClient, public TSharedFromThis<FChaosClothAssetEditor3DViewportClient>
 {
 public:
 
-	FChaosClothAssetEditor3DViewportClient(FEditorModeTools* InModeTools, FPreviewScene* InPreviewScene = nullptr,
+	FChaosClothAssetEditor3DViewportClient(FEditorModeTools* InModeTools, TSharedPtr<FChaosClothPreviewScene> InPreviewScene,
 		const TWeakPtr<SEditorViewport>& InEditorViewportWidget = nullptr);
 
+	// Call this after construction to initialize callbacks when settings change
+	void RegisterSettingsChangedDelegate();
+
 	virtual ~FChaosClothAssetEditor3DViewportClient();
 
 	// Delete the viewport gizmo and transform proxy
@@ -59,6 +63,11 @@ private:
 	virtual void ProcessClick(FSceneView& View, HHitProxy* HitProxy, FKey Key, EInputEvent Event, uint32 HitX, uint32 HitY) override;
 	virtual void Draw(const FSceneView* View, FPrimitiveDrawInterface* PDI) override;
 
+	void OnAssetViewerSettingsChanged(const FName& InPropertyName);
+	void SetAdvancedShowFlagsForScene(const bool bAdvancedShowFlags);
+
+	TWeakPtr<FChaosClothPreviewScene> ClothPreviewScene;
+
 	TObjectPtr<UChaosClothComponent> ClothComponent;
 
 	TObjectPtr<UChaosClothAssetEditorMode> ClothEdMode;
diff --git a/Engine/Plugins/Experimental/ChaosClothAssetEditor/Source/ChaosClothAssetEditor/Private/ChaosClothAsset/ClothEditorToolkit.cpp b/Engine/Plugins/Experimental/ChaosClothAssetEditor/Source/ChaosClothAssetEditor/Private/ChaosClothAsset/ClothEditorToolkit.cpp
--- a/Engine/Plugins/Experimental/ChaosClothAssetEditor/Source/ChaosClothAssetEditor/Private/ChaosClothAsset/ClothEditorToolkit.cpp
+++ b/Engine/Plugins/Experimental/ChaosClothAssetEditor/Source/ChaosClothAssetEditor/Private/ChaosClothAsset/ClothEditorToolkit.cpp
@@ -178,7 +178,8 @@ FChaosClothAssetEditorToolkit::FChaosClothAssetEditorToolkit(UAssetEditor* InOwn
 
 	//ClothPreviewInputRouter = ClothPreviewEditorModeManager->GetInteractiveToolsContext()->InputRouter;
 	ClothPreviewTabContent = MakeShareable(new FEditorViewportTabContent());
-	ClothPreviewViewportClient = MakeShared<FChaosClothAssetEditor3DViewportClient>(ClothPreviewEditorModeManager.Get(), ClothPreviewScene.Get());
+	ClothPreviewViewportClient = MakeShared<FChaosClothAssetEditor3DViewportClient>(ClothPreviewEditorModeManager.Get(), ClothPreviewScene);
+	ClothPreviewViewportClient->RegisterSettingsChangedDelegate();
 
 	ClothPreviewViewportDelegate = [this](FAssetEditorViewportConstructionArgs InArgs)
 	{
