Chaos Cloth - Add basic Dataflow nodes to the node editor, including the cloth simulation default config node.

#rnx

diff --git a/Engine/Plugins/Experimental/ChaosClothAsset/Source/ChaosClothAsset/Private/ChaosClothAsset/ClothCollection.cpp b/Engine/Plugins/Experimental/ChaosClothAsset/Source/ChaosClothAsset/Private/ChaosClothAsset/ClothCollection.cpp
--- a/Engine/Plugins/Experimental/ChaosClothAsset/Source/ChaosClothAsset/Private/ChaosClothAsset/ClothCollection.cpp
+++ b/Engine/Plugins/Experimental/ChaosClothAsset/Source/ChaosClothAsset/Private/ChaosClothAsset/ClothCollection.cpp
@@ -14,6 +14,9 @@ namespace UE::Chaos::ClothAsset
 	const FName UE::Chaos::ClothAsset::FClothCollection::TethersGroup("Tethers");
 	const FName UE::Chaos::ClothAsset::FClothCollection::TetherBatchesGroup("TetherBatches");
 	const FName UE::Chaos::ClothAsset::FClothCollection::LodsGroup("Lods");
+	const FName UE::Chaos::ClothAsset::FClothCollection::MaterialsGroup("Materials");
+	const FName UE::Chaos::ClothAsset::FClothCollection::CollisionsGroup("CollisionsGroup");
+	const FName UE::Chaos::ClothAsset::FClothCollection::SkeletonsGroup("SkeletonsGroup");
 
 	FClothCollection::FClothCollection()
 	{
@@ -103,6 +106,15 @@ namespace UE::Chaos::ClothAsset
 		AddExternalAttribute<int32>("TetherBatchStart", LodsGroup, TetherBatchStart, TetherBatchesDependency);
 		AddExternalAttribute<int32>("TetherBatchEnd", LodsGroup, TetherBatchEnd, TetherBatchesDependency);
 		AddExternalAttribute<int32>("LodBiasDepth", LodsGroup, LodBiasDepth);
+
+		// Materials Group
+		AddExternalAttribute<FString>("MaterialPathName", MaterialsGroup, MaterialPathName);
+
+		// Collisions Group
+		AddExternalAttribute<FString>("PhysicsAssetPathName", CollisionsGroup, PhysicsAssetPathName);
+		
+		// Skeletons Group
+		AddExternalAttribute<FString>("SkeletonAssetPathName", SkeletonsGroup, SkeletonAssetPathName);
 	}
 
 	int32 FClothCollection::SetNumElements(int32 InNumElements, const FName& GroupName, TManagedArray<int32>& StartArray, TManagedArray<int32>& EndArray, int32 StartEndIndex)
diff --git a/Engine/Plugins/Experimental/ChaosClothAsset/Source/ChaosClothAsset/Private/ChaosClothAsset/ClothGeometryTools.cpp b/Engine/Plugins/Experimental/ChaosClothAsset/Source/ChaosClothAsset/Private/ChaosClothAsset/ClothGeometryTools.cpp
--- a/Engine/Plugins/Experimental/ChaosClothAsset/Source/ChaosClothAsset/Private/ChaosClothAsset/ClothGeometryTools.cpp
+++ b/Engine/Plugins/Experimental/ChaosClothAsset/Source/ChaosClothAsset/Private/ChaosClothAsset/ClothGeometryTools.cpp
@@ -6,6 +6,25 @@
 
 namespace UE::Chaos::ClothAsset
 {
+	void FClothGeometryTools::DeleteRenderMesh(const TSharedPtr<FClothCollection>& ClothCollection)
+	{
+		FClothAdapter ClothAdapter(ClothCollection);
+
+		for (int32 LodIndex = 0; LodIndex < ClothAdapter.GetNumLods(); ++LodIndex)
+		{
+			FClothLodAdapter ClothLodAdapter = ClothAdapter.GetLod(LodIndex);
+
+			for (int32 PatternIndex = ClothLodAdapter.GetNumPatterns() - 1; PatternIndex >= 0 ; --PatternIndex)  // Use a reverse order to avoid having to move previous elements
+			{
+				FClothPatternAdapter ClothPatternAdapter = ClothLodAdapter.GetPattern(PatternIndex);
+				ClothPatternAdapter.SetNumRenderVertices(0);
+				ClothPatternAdapter.SetNumRenderFaces(0);
+			}
+		}
+		// TODO: Add Materials functions to cloth adapter
+		ClothCollection->EmptyGroup(FClothCollection::MaterialsGroup);
+	}
+
 	void FClothGeometryTools::CopySimMeshToRenderMesh(const TSharedPtr<FClothCollection>& ClothCollection, int32 MaterialIndex)
 	{
 		FClothAdapter ClothAdapter(ClothCollection);
@@ -100,7 +119,7 @@ namespace UE::Chaos::ClothAsset
 				for (int32 VertexIndex = 0; VertexIndex < NumVertices; ++VertexIndex)
 				{
 					RenderPosition[VertexIndex] = SimRestPosition[VertexIndex];
-					RenderNormal[VertexIndex] = SimRestNormal[VertexIndex];
+					RenderNormal[VertexIndex] = -SimRestNormal[VertexIndex];  // Simulation normals use reverse normals
 					RenderUVs[VertexIndex] = { SimPosition[VertexIndex] * UVInvScale };
 					RenderColor[VertexIndex] = FLinearColor::White;
 					RenderTangentU[VertexIndex].Normalize();
@@ -109,4 +128,66 @@ namespace UE::Chaos::ClothAsset
 			}
 		}
 	}
-}  // End namespace UE::Chaos::ClothAsset
\ No newline at end of file
+
+	void FClothGeometryTools::ReverseNormals(
+		const TSharedPtr<FClothCollection>& ClothCollection,
+		bool bReverseSimMeshNormals,
+		bool bReverseRenderMeshNormals,
+		const TArray<int32>& PatternSelection)
+	{
+		auto ReverseSimNormals = [](const TArrayView<FVector3f>& SimRestNormal)
+			{
+				for (int32 VertexIndex = 0; VertexIndex < SimRestNormal.Num(); ++VertexIndex)
+				{
+					SimRestNormal[VertexIndex] = -SimRestNormal[VertexIndex];
+				}
+			};
+		auto ReverseRenderNormals = [](const TArrayView<FVector3f>& RenderNormal, const TArrayView<FVector3f>& RenderTangentU)
+		{
+			check(RenderNormal.Num() == RenderTangentU.Num())
+			for (int32 VertexIndex = 0; VertexIndex < RenderNormal.Num(); ++VertexIndex)
+			{
+				RenderNormal[VertexIndex] = -RenderNormal[VertexIndex];      // Equivalent of rotating the normal basis
+				RenderTangentU[VertexIndex] = -RenderTangentU[VertexIndex];  // around tangent V
+			}
+		};
+
+		FClothAdapter ClothAdapter(ClothCollection);
+
+		for (int32 LodIndex = 0; LodIndex < ClothAdapter.GetNumLods(); ++LodIndex)
+		{
+			FClothLodAdapter ClothLodAdapter = ClothAdapter.GetLod(LodIndex);
+
+			if (PatternSelection.IsEmpty())
+			{
+				if (bReverseSimMeshNormals)
+				{
+					ReverseSimNormals(ClothLodAdapter.GetPatternsSimRestNormal());
+				}
+				if (bReverseRenderMeshNormals)
+				{
+					ReverseRenderNormals(ClothLodAdapter.GetPatternsRenderNormal(), ClothLodAdapter.GetPatternsRenderTangentU());
+				}
+			}
+			else
+			{
+				for (int32 PatternIndex = 0; PatternIndex < ClothLodAdapter.GetNumPatterns(); ++PatternIndex)
+				{
+					if (PatternSelection.Find(PatternIndex) != INDEX_NONE)
+					{
+						FClothPatternAdapter ClothPatternAdapter = ClothLodAdapter.GetPattern(PatternIndex);
+
+						if (bReverseSimMeshNormals)
+						{
+							ReverseSimNormals(ClothPatternAdapter.GetSimRestNormal());
+						}
+						if (bReverseRenderMeshNormals)
+						{
+							ReverseRenderNormals(ClothPatternAdapter.GetRenderNormal(), ClothPatternAdapter.GetRenderTangentU());
+						}
+					}
+				}
+			}
+		}
+	}
+}  // End namespace UE::Chaos::ClothAsset
diff --git a/Engine/Plugins/Experimental/ChaosClothAsset/Source/ChaosClothAsset/Public/ChaosClothAsset/ClothCollection.h b/Engine/Plugins/Experimental/ChaosClothAsset/Source/ChaosClothAsset/Public/ChaosClothAsset/ClothCollection.h
--- a/Engine/Plugins/Experimental/ChaosClothAsset/Source/ChaosClothAsset/Public/ChaosClothAsset/ClothCollection.h
+++ b/Engine/Plugins/Experimental/ChaosClothAsset/Source/ChaosClothAsset/Public/ChaosClothAsset/ClothCollection.h
@@ -62,6 +62,9 @@ namespace UE::Chaos::ClothAsset
 		static const FName TethersGroup;  // Tethers information
 		static const FName TetherBatchesGroup;  // Tethers parallel batch processing information
 		static const FName LodsGroup;  // Lod split information
+		static const FName MaterialsGroup;  // Material information
+		static const FName CollisionsGroup;  // Collision information
+		static const FName SkeletonsGroup;  // Skeleton information
 
 		// Sim Vertices Group
 		TManagedArray<FVector2f> SimPosition;
@@ -126,6 +129,15 @@ namespace UE::Chaos::ClothAsset
 		TManagedArray<int32> TetherBatchEnd;
 		TManagedArray<int32> LodBiasDepth;  // The number of LODs covered by each Sim LOD (for the wrap deformer)
 
+		// Materials Group
+		TManagedArray<FString> MaterialPathName;
+
+		// Collisions Group
+		TManagedArray<FString> PhysicsAssetPathName;
+
+		// Skeletons Group
+		TManagedArray<FString> SkeletonAssetPathName;
+
 	protected:
 		void Construct();
 	};
diff --git a/Engine/Plugins/Experimental/ChaosClothAsset/Source/ChaosClothAsset/Public/ChaosClothAsset/ClothGeometryTools.h b/Engine/Plugins/Experimental/ChaosClothAsset/Source/ChaosClothAsset/Public/ChaosClothAsset/ClothGeometryTools.h
--- a/Engine/Plugins/Experimental/ChaosClothAsset/Source/ChaosClothAsset/Public/ChaosClothAsset/ClothGeometryTools.h
+++ b/Engine/Plugins/Experimental/ChaosClothAsset/Source/ChaosClothAsset/Public/ChaosClothAsset/ClothGeometryTools.h
@@ -11,7 +11,17 @@ namespace UE::Chaos::ClothAsset
 	 */
 	struct CHAOSCLOTHASSET_API FClothGeometryTools
 	{
+		/** Delete the render mesh data. */
+		static void DeleteRenderMesh(const TSharedPtr<FClothCollection>& ClothCollection);
+
 		/** Turn the sim mesh portion of this ClothCollection into a render mesh. */
 		static void CopySimMeshToRenderMesh(const TSharedPtr<FClothCollection>& ClothCollection, int32 MaterialIndex);
+
+		/** Reverse the mesh normals. Will reverse all normals if pattern selection is empty. */
+		static void ReverseNormals(
+			const TSharedPtr<FClothCollection>& ClothCollection,
+			bool bReverseSimMeshNormals,
+			bool bReverseRenderMeshNormals,
+			const TArray<int32>& PatternSelection);
 	};
 }  // End namespace UE::Chaos::ClothAsset
\ No newline at end of file
diff --git a/Engine/Plugins/Experimental/ChaosClothAssetEditor/Source/ChaosClothAssetEditor/ChaosClothAssetEditor.Build.cs b/Engine/Plugins/Experimental/ChaosClothAssetEditor/Source/ChaosClothAssetEditor/ChaosClothAssetEditor.Build.cs
--- a/Engine/Plugins/Experimental/ChaosClothAssetEditor/Source/ChaosClothAssetEditor/ChaosClothAssetEditor.Build.cs
+++ b/Engine/Plugins/Experimental/ChaosClothAssetEditor/Source/ChaosClothAssetEditor/ChaosClothAssetEditor.Build.cs
@@ -50,6 +50,7 @@ public class ChaosClothAssetEditor : ModuleRules
 				"ClothingSystemEditorInterface",
 				"ClothingSystemRuntimeInterface",
 				"Chaos",
+				"ChaosCloth",
 				"ChaosClothAsset",
 				"ChaosClothAssetEngine",
 				"ClothingSystemEditor",
diff --git a/Engine/Plugins/Experimental/ChaosClothAssetEditor/Source/ChaosClothAssetEditor/Private/ChaosClothAsset/ChaosClothAssetEditorModule.cpp b/Engine/Plugins/Experimental/ChaosClothAssetEditor/Source/ChaosClothAssetEditor/Private/ChaosClothAsset/ChaosClothAssetEditorModule.cpp
--- a/Engine/Plugins/Experimental/ChaosClothAssetEditor/Source/ChaosClothAssetEditor/Private/ChaosClothAsset/ChaosClothAssetEditorModule.cpp
+++ b/Engine/Plugins/Experimental/ChaosClothAssetEditor/Source/ChaosClothAssetEditor/Private/ChaosClothAsset/ChaosClothAssetEditorModule.cpp
@@ -4,7 +4,7 @@
 #include "ChaosClothAsset/ClothEditorMode.h"
 #include "ChaosClothAsset/ClothEditorCommands.h"
 #include "ChaosClothAsset/ClothEditorStyle.h"
-#include "ChaosClothAsset/ClothDataflowNodes.h"
+#include "ChaosClothAsset/DataflowNodes/DataflowNodes.h"
 #include "ContentBrowserMenuContexts.h"
 #include "EditorModeRegistry.h"
 #include "Selection.h"
@@ -15,10 +15,10 @@
 
 #define LOCTEXT_NAMESPACE "FChaosClothAssetEditorModule"
 
-using namespace UE::Chaos::ClothAsset;
-
 void FChaosClothAssetEditorModule::StartupModule()
 {
+	using namespace UE::Chaos::ClothAsset;
+
 	FChaosClothAssetEditorStyle::Get(); // Causes the constructor to be called
 
 	FChaosClothAssetEditorCommands::Register();
@@ -27,7 +27,7 @@ void FChaosClothAssetEditorModule::StartupModule()
 	FAssetToolsModule& AssetToolsModule = FAssetToolsModule::GetModule();
 	IAssetTools& AssetTools = AssetToolsModule.Get();
 
-	Dataflow::RegisterClothDataflowNodes();
+	DataflowNodes::Register();
 
 	// TODO: Register details view customizations
 }
diff --git a/Engine/Plugins/Experimental/ChaosClothAssetEditor/Source/ChaosClothAssetEditor/Private/ChaosClothAsset/ClothDataflowNodes.cpp b/Engine/Plugins/Experimental/ChaosClothAssetEditor/Source/ChaosClothAssetEditor/Private/ChaosClothAsset/ClothDataflowNodes.cpp
deleted file mode 100644
diff --git a/Engine/Plugins/Experimental/ChaosClothAssetEditor/Source/ChaosClothAssetEditor/Private/ChaosClothAsset/ClothDataflowNodes.h b/Engine/Plugins/Experimental/ChaosClothAssetEditor/Source/ChaosClothAssetEditor/Private/ChaosClothAsset/ClothDataflowNodes.h
deleted file mode 100644
diff --git a/Engine/Plugins/Experimental/ChaosClothAssetEditor/Source/ChaosClothAssetEditor/Private/ChaosClothAsset/DataflowNodes/CopySimulationToRenderMeshNode.cpp b/Engine/Plugins/Experimental/ChaosClothAssetEditor/Source/ChaosClothAssetEditor/Private/ChaosClothAsset/DataflowNodes/CopySimulationToRenderMeshNode.cpp
new file mode 100644
--- /dev/null
+++ b/Engine/Plugins/Experimental/ChaosClothAssetEditor/Source/ChaosClothAssetEditor/Private/ChaosClothAsset/DataflowNodes/CopySimulationToRenderMeshNode.cpp
@@ -0,0 +1,63 @@
+// Copyright Epic Games, Inc. All Rights Reserved.
+
+#include "ChaosClothAsset/DataflowNodes/CopySimulationToRenderMeshNode.h"
+#include "ChaosClothAsset/DataflowNodes/DataflowNodes.h"
+#include "ChaosClothAsset/ClothGeometryTools.h"
+#include "Materials/Material.h"
+
+#include UE_INLINE_GENERATED_CPP_BY_NAME(CopySimulationToRenderMeshNode)
+
+#define LOCTEXT_NAMESPACE "ChaosClothAssetCopySimulationToRenderMeshNode"
+
+FChaosClothAssetCopySimulationToRenderMeshNode::FChaosClothAssetCopySimulationToRenderMeshNode(const Dataflow::FNodeParameters& InParam, FGuid InGuid)
+	: FDataflowNode(InParam, InGuid)
+{
+	RegisterInputConnection(&Collection);
+	RegisterInputConnection(&Patterns);
+	RegisterOutputConnection(&Collection, &Collection);
+}
+
+void FChaosClothAssetCopySimulationToRenderMeshNode::Evaluate(Dataflow::FContext& Context, const FDataflowOutput* Out) const
+{
+	if (Out->IsA<FManagedArrayCollection>(&Collection))
+	{
+		using namespace UE::Chaos::ClothAsset;
+
+		// Evaluate in collection
+		const FManagedArrayCollection& InCollection = GetValue<FManagedArrayCollection>(Context, &Collection);
+
+		// TODO: Needs to make cloth collection a facade to avoid the copy of the cloth collection to a managed array
+		//       and remove the above const reference to operated on the InCollection instead.
+		TSharedPtr<FClothCollection> ClothCollection = MakeShared<FClothCollection>();
+		InCollection.CopyTo(ClothCollection.Get());
+
+		// Empty mesh and materials
+		FClothGeometryTools::DeleteRenderMesh(ClothCollection);
+
+		// Find the material asset path
+		const FString MaterialPathName = Material ? 
+			Material->GetPathName() :
+			FString(TEXT("/Engine/EditorMaterials/Cloth/CameraLitDoubleSided.CameraLitDoubleSided"));
+
+		// Find or add the material attribute
+		// TODO: Add to ClothAdapter functionality
+		TManagedArray<FString>* const FindAttributeResult = ClothCollection->FindAttribute<FString>("MaterialPathName", FClothCollection::MaterialsGroup);
+		TManagedArray<FString>& MaterialPathNameArray = FindAttributeResult ? *FindAttributeResult :
+			ClothCollection->AddAttribute<FString>("MaterialPathName", FClothCollection::MaterialsGroup);
+
+		// Work out the material index and set the path name
+		int32 MaterialIndex = MaterialPathNameArray.Find(MaterialPathName);
+	
+		if (MaterialIndex == INDEX_NONE)
+		{
+			MaterialIndex = ClothCollection->AddElements(1, FClothCollection::MaterialsGroup);
+			MaterialPathNameArray[MaterialIndex] = MaterialPathName;
+		}
+
+		FClothGeometryTools::CopySimMeshToRenderMesh(ClothCollection, MaterialIndex);
+
+		SetValue<FManagedArrayCollection>(Context, *ClothCollection, &Collection);
+	}
+}
+
+#undef LOCTEXT_NAMESPACE
diff --git a/Engine/Plugins/Experimental/ChaosClothAssetEditor/Source/ChaosClothAssetEditor/Private/ChaosClothAsset/DataflowNodes/CopySimulationToRenderMeshNode.h b/Engine/Plugins/Experimental/ChaosClothAssetEditor/Source/ChaosClothAssetEditor/Private/ChaosClothAsset/DataflowNodes/CopySimulationToRenderMeshNode.h
new file mode 100644
--- /dev/null
+++ b/Engine/Plugins/Experimental/ChaosClothAssetEditor/Source/ChaosClothAssetEditor/Private/ChaosClothAsset/DataflowNodes/CopySimulationToRenderMeshNode.h
@@ -0,0 +1,34 @@
+// Copyright Epic Games, Inc. All Rights Reserved.
+
+#pragma once 
+
+#include "Dataflow/DataflowNode.h"
+#include "GeometryCollection/ManagedArrayCollection.h"
+#include "CopySimulationToRenderMeshNode.generated.h"
+
+class UMaterialInterface;
+
+USTRUCT(Meta = (DataflowCloth))
+struct FChaosClothAssetCopySimulationToRenderMeshNode : public FDataflowNode
+{
+	GENERATED_USTRUCT_BODY()
+	DATAFLOW_NODE_DEFINE_INTERNAL(FChaosClothAssetCopySimulationToRenderMeshNode, "CopySimulationToRenderMesh", "Cloth", "Cloth Simulation Render Mesh")
+	//DATAFLOW_NODE_RENDER_TYPE(FGeometryCollection::StaticType(), "Collection")  // TODO: Leave out the render type until there is something to render
+
+public:
+
+	UPROPERTY(Meta = (Dataflowinput, DataflowOutput, DisplayName = "Collection", DataflowPassthrough = "Collection"))
+	FManagedArrayCollection Collection;
+
+	/** List of patterns to apply the operation on. All patterns will be used if left empty. */
+	UPROPERTY(EditAnywhere, Category = "Pattern Selection")
+	TArray<int32> Patterns;
+
+	/** New material for the render mesh. */
+	UPROPERTY(EditAnywhere, Category = "Render Mesh")
+	TObjectPtr<const UMaterialInterface> Material;
+
+	FChaosClothAssetCopySimulationToRenderMeshNode(const Dataflow::FNodeParameters& InParam, FGuid InGuid = FGuid::NewGuid());
+
+	virtual void Evaluate(Dataflow::FContext& Context, const FDataflowOutput* Out) const override;
+};
diff --git a/Engine/Plugins/Experimental/ChaosClothAssetEditor/Source/ChaosClothAssetEditor/Private/ChaosClothAsset/DataflowNodes/DataflowNodes.cpp b/Engine/Plugins/Experimental/ChaosClothAssetEditor/Source/ChaosClothAssetEditor/Private/ChaosClothAsset/DataflowNodes/DataflowNodes.cpp
new file mode 100644
--- /dev/null
+++ b/Engine/Plugins/Experimental/ChaosClothAssetEditor/Source/ChaosClothAssetEditor/Private/ChaosClothAsset/DataflowNodes/DataflowNodes.cpp
@@ -0,0 +1,40 @@
+// Copyright Epic Games, Inc. All Rights Reserved.
+
+#include "ChaosClothAsset/DataflowNodes/DataflowNodes.h"
+#include "ChaosClothAsset/DataflowNodes/TerminalNode.h"
+#include "ChaosClothAsset/DataflowNodes/DatasmithImportNode.h"
+#include "ChaosClothAsset/DataflowNodes/ImportNode.h"
+#include "ChaosClothAsset/DataflowNodes/PatternSelectionNode.h"
+#include "ChaosClothAsset/DataflowNodes/DeleteRenderMeshNode.h"
+#include "ChaosClothAsset/DataflowNodes/CopySimulationToRenderMeshNode.h"
+#include "ChaosClothAsset/DataflowNodes/ReverseNormalsNode.h"
+#include "ChaosClothAsset/DataflowNodes/SimulationDefaultConfigNode.h"
+#include "Dataflow/DataflowNodeFactory.h"
+#include "Framework/Notifications/NotificationManager.h"
+#include "Widgets/Notifications/SNotificationList.h"
+
+DEFINE_LOG_CATEGORY(LogChaosClothAssetDataflowNodes);
+
+namespace UE::Chaos::ClothAsset::DataflowNodes
+{
+	void Register()
+	{
+		DATAFLOW_NODE_REGISTER_CREATION_FACTORY(FChaosClothAssetTerminalNode);
+		DATAFLOW_NODE_REGISTER_CREATION_FACTORY(FChaosClothAssetDatasmithImportNode);
+		DATAFLOW_NODE_REGISTER_CREATION_FACTORY(FChaosClothAssetImportNode);
+		DATAFLOW_NODE_REGISTER_CREATION_FACTORY(FChaosClothAssetPatternSelectionNode);
+		DATAFLOW_NODE_REGISTER_CREATION_FACTORY(FChaosClothAssetDeleteRenderMeshNode);
+		DATAFLOW_NODE_REGISTER_CREATION_FACTORY(FChaosClothAssetCopySimulationToRenderMeshNode);
+		DATAFLOW_NODE_REGISTER_CREATION_FACTORY(FChaosClothAssetReverseNormalsNode);
+		DATAFLOW_NODE_REGISTER_CREATION_FACTORY(FChaosClothAssetSimulationDefaultConfigNode);
+	}
+
+	void LogAndToastWarning(const FText& Error)
+	{
+		FNotificationInfo Info(Error);
+		Info.ExpireDuration = 5.0f;
+		FSlateNotificationManager::Get().AddNotification(Info);
+
+		UE_LOG(LogChaosClothAssetDataflowNodes, Warning, TEXT("%s"), *Error.ToString());
+	}
+}
diff --git a/Engine/Plugins/Experimental/ChaosClothAssetEditor/Source/ChaosClothAssetEditor/Private/ChaosClothAsset/DataflowNodes/DataflowNodes.h b/Engine/Plugins/Experimental/ChaosClothAssetEditor/Source/ChaosClothAssetEditor/Private/ChaosClothAsset/DataflowNodes/DataflowNodes.h
new file mode 100644
--- /dev/null
+++ b/Engine/Plugins/Experimental/ChaosClothAssetEditor/Source/ChaosClothAssetEditor/Private/ChaosClothAsset/DataflowNodes/DataflowNodes.h
@@ -0,0 +1,16 @@
+// Copyright Epic Games, Inc. All Rights Reserved.
+
+#pragma once
+
+#include "Dataflow/DataflowNodeParameters.h"
+#include "Dataflow/DataflowInputOutput.h"
+#include "Dataflow/DataflowObjectInterface.h"
+#include "Logging/LogMacros.h"
+
+namespace UE::Chaos::ClothAsset::DataflowNodes
+{
+	void Register();
+	void LogAndToastWarning(const FText& Error);
+}
+
+DECLARE_LOG_CATEGORY_EXTERN(LogChaosClothAssetDataflowNodes, Log, All);
diff --git a/Engine/Plugins/Experimental/ChaosClothAssetEditor/Source/ChaosClothAssetEditor/Private/ChaosClothAsset/DataflowNodes/DatasmithImportNode.cpp b/Engine/Plugins/Experimental/ChaosClothAssetEditor/Source/ChaosClothAssetEditor/Private/ChaosClothAsset/DataflowNodes/DatasmithImportNode.cpp
new file mode 100644
--- /dev/null
+++ b/Engine/Plugins/Experimental/ChaosClothAssetEditor/Source/ChaosClothAssetEditor/Private/ChaosClothAsset/DataflowNodes/DatasmithImportNode.cpp
@@ -0,0 +1,108 @@
+// Copyright Epic Games, Inc. All Rights Reserved.
+
+#include "ChaosClothAsset/DataflowNodes/DatasmithImportNode.h"
+#include "ChaosClothAsset/DataflowNodes/DataflowNodes.h"
+#include "ChaosClothAsset/ClothCollection.h"
+#include "ChaosClothAsset/ClothAsset.h"
+#include "DatasmithImportContext.h"
+#include "DatasmithImportFactory.h"
+#include "ExternalSource.h"
+#include "ExternalSourceModule.h"
+#include "SourceUri.h"
+
+#include UE_INLINE_GENERATED_CPP_BY_NAME(DatasmithImportNode)
+
+#define LOCTEXT_NAMESPACE "ChaosClothAssetDatasmithImportNode"
+
+FChaosClothAssetDatasmithImportNode::FChaosClothAssetDatasmithImportNode(const Dataflow::FNodeParameters& InParam, FGuid InGuid)
+	: FDataflowNode(InParam, InGuid)
+	, DestPackageName("/Game/ClothAsset")
+{
+	RegisterInputConnection(&DatasmithFile);
+	RegisterInputConnection(&DestPackageName);
+	RegisterOutputConnection(&Collection);
+}
+
+bool FChaosClothAssetDatasmithImportNode::EvaluateImpl(Dataflow::FContext& Context, FManagedArrayCollection& OutCollection) const
+{
+	using namespace UE::DatasmithImporter;
+
+	const FFilePath& InFilePath = GetValue<FFilePath>(Context, &DatasmithFile);
+	const FString& InDestPackageName = GetValue<FString>(Context, &DestPackageName);
+
+	const FSourceUri SourceUri = FSourceUri::FromFilePath(InFilePath.FilePath);
+	const TSharedPtr<FExternalSource> ExternalSource = IExternalSourceModule::GetOrCreateExternalSource(SourceUri);
+	if (!ExternalSource)
+	{
+		return false;
+	}
+
+	constexpr bool bLoadConfig = false; // TODO: Not sure what this does
+	const FName LoggerName(TEXT("ImportDatasmithClothNode"));
+	const FText LoggerLabel(NSLOCTEXT("ImportDatasmithClothNode", "LoggerLabel", "ImportDatasmithClothNode"));
+	FDatasmithImportContext DatasmithImportContext(ExternalSource, bLoadConfig, LoggerName, LoggerLabel);
+
+	const TStrongObjectPtr<const UPackage> DestinationPackage(CreatePackage(*InDestPackageName));		// TODO: would be nice to not have to create the package but it seems necessary for now
+	if (!ensure(DestinationPackage))
+	{
+		// Failed to create the package to hold this asset for some reason
+		return false;
+	}
+	
+	// Don't create the Actors in the level, just read the Assets
+	DatasmithImportContext.Options->BaseOptions.SceneHandling = EDatasmithImportScene::AssetsOnly;
+
+	constexpr EObjectFlags NewObjectFlags = RF_Public | RF_Standalone | RF_Transactional;
+	const TSharedPtr<FJsonObject> ImportSettingsJson;
+	constexpr bool bIsSilent = true;
+	const FString DestinationPath = DestinationPackage->GetName();
+	if (!DatasmithImportContext.Init(DestinationPath, NewObjectFlags, GWarn, ImportSettingsJson, bIsSilent))
+	{
+		return false;
+	}
+
+	if (const TSharedPtr<IDatasmithScene> LoadedScene = ExternalSource->TryLoad())
+	{
+		DatasmithImportContext.InitScene(LoadedScene.ToSharedRef());
+	}
+	else
+	{
+		return false;
+	}
+
+	bool bUserCancelled = false;
+	bool bImportSucceed = DatasmithImportFactoryImpl::ImportDatasmithScene(DatasmithImportContext, bUserCancelled);
+	bImportSucceed &= !bUserCancelled;
+
+	if (bImportSucceed && DatasmithImportContext.ImportedClothes.Num() > 0)
+	{
+		const UObject* const ClothObject = DatasmithImportContext.ImportedClothes.CreateIterator()->Value;
+		const UChaosClothAsset* const DatasmithClothAsset = Cast<UChaosClothAsset>(ClothObject);
+		if (ensure(DatasmithClothAsset))
+		{
+			DatasmithClothAsset->GetClothCollection()->CopyTo(&OutCollection);
+			return true;
+		}
+	}
+
+	return false;
+}
+
+void FChaosClothAssetDatasmithImportNode::Evaluate(Dataflow::FContext& Context, const FDataflowOutput* Out) const
+{
+	FManagedArrayCollection OutCollection;
+	const bool bSuccess = EvaluateImpl(Context, OutCollection);
+
+	if (bSuccess)
+	{
+		SetValue<FManagedArrayCollection>(Context, OutCollection, &Collection);
+	}
+	else
+	{
+		// Init with an empty cloth collection
+		UE::Chaos::ClothAsset::FClothCollection ClothCollection;
+		SetValue<FManagedArrayCollection>(Context, ClothCollection, &Collection);
+	}
+}
+
+#undef LOCTEXT_NAMESPACE
diff --git a/Engine/Plugins/Experimental/ChaosClothAssetEditor/Source/ChaosClothAssetEditor/Private/ChaosClothAsset/DataflowNodes/DatasmithImportNode.h b/Engine/Plugins/Experimental/ChaosClothAssetEditor/Source/ChaosClothAssetEditor/Private/ChaosClothAsset/DataflowNodes/DatasmithImportNode.h
new file mode 100644
--- /dev/null
+++ b/Engine/Plugins/Experimental/ChaosClothAssetEditor/Source/ChaosClothAssetEditor/Private/ChaosClothAsset/DataflowNodes/DatasmithImportNode.h
@@ -0,0 +1,36 @@
+// Copyright Epic Games, Inc. All Rights Reserved.
+
+#pragma once 
+
+#include "Dataflow/DataflowNode.h"
+#include "GeometryCollection/ManagedArrayCollection.h"
+#include "DatasmithImportNode.generated.h"
+
+USTRUCT(meta = (DataflowCloth))
+struct FChaosClothAssetDatasmithImportNode : public FDataflowNode
+{
+	GENERATED_USTRUCT_BODY()
+	DATAFLOW_NODE_DEFINE_INTERNAL(FChaosClothAssetDatasmithImportNode, "DatasmithImport", "Cloth", "Cloth Datasmith Import")
+	//DATAFLOW_NODE_RENDER_TYPE(FGeometryCollection::StaticType(), "Collection")
+
+public:
+	
+	/** Datasmith file to read from. */
+	UPROPERTY(EditAnywhere, Category = "Datasmith Importer")
+	FFilePath DatasmithFile;
+
+	/** Package to import into. */
+	UPROPERTY(EditAnywhere, Category = "Datasmith Importer")
+	FString DestPackageName;
+
+	UPROPERTY(meta = (DataflowOutput, DisplayName = "Collection"))
+	FManagedArrayCollection Collection;
+
+	FChaosClothAssetDatasmithImportNode(const Dataflow::FNodeParameters& InParam, FGuid InGuid = FGuid::NewGuid());
+
+	virtual void Evaluate(Dataflow::FContext& Context, const FDataflowOutput* Out) const override;
+
+private:
+
+	bool EvaluateImpl(Dataflow::FContext& Context, FManagedArrayCollection& OutCollection) const;
+};
diff --git a/Engine/Plugins/Experimental/ChaosClothAssetEditor/Source/ChaosClothAssetEditor/Private/ChaosClothAsset/DataflowNodes/DeleteRenderMeshNode.cpp b/Engine/Plugins/Experimental/ChaosClothAssetEditor/Source/ChaosClothAssetEditor/Private/ChaosClothAsset/DataflowNodes/DeleteRenderMeshNode.cpp
new file mode 100644
--- /dev/null
+++ b/Engine/Plugins/Experimental/ChaosClothAssetEditor/Source/ChaosClothAssetEditor/Private/ChaosClothAsset/DataflowNodes/DeleteRenderMeshNode.cpp
@@ -0,0 +1,46 @@
+// Copyright Epic Games, Inc. All Rights Reserved.
+
+#include "ChaosClothAsset/DataflowNodes/DeleteRenderMeshNode.h"
+#include "ChaosClothAsset/DataflowNodes/DataflowNodes.h"
+#include "ChaosClothAsset/ClothGeometryTools.h"
+
+#include UE_INLINE_GENERATED_CPP_BY_NAME(DeleteRenderMeshNode)
+
+#define LOCTEXT_NAMESPACE "ChaosClothAssetDeleteRenderMeshNode"
+
+FChaosClothAssetDeleteRenderMeshNode::FChaosClothAssetDeleteRenderMeshNode(const Dataflow::FNodeParameters& InParam, FGuid InGuid)
+	: FDataflowNode(InParam, InGuid)
+{
+	RegisterInputConnection(&Collection);
+	RegisterInputConnection(&Patterns);
+	RegisterOutputConnection(&Collection, &Collection);
+}
+
+void FChaosClothAssetDeleteRenderMeshNode::Evaluate(Dataflow::FContext& Context, const FDataflowOutput* Out) const
+{
+	if (Out->IsA<FManagedArrayCollection>(&Collection))
+	{
+		using namespace UE::Chaos::ClothAsset;
+
+		// Evaluate in collection
+		const FManagedArrayCollection& InCollection = GetValue<FManagedArrayCollection>(Context, &Collection);
+
+		// TODO: Needs to make cloth collection a facade to avoid the copy of the cloth collection to a managed array
+		//       and remove the above const reference to operated on the InCollection instead.
+		TSharedPtr<FClothCollection> ClothCollection = MakeShared<FClothCollection>();
+		InCollection.CopyTo(ClothCollection.Get());
+
+		if (Patterns.IsEmpty())
+		{
+			FClothGeometryTools::DeleteRenderMesh(ClothCollection);
+		}
+		else
+		{
+			DataflowNodes::LogAndToastWarning(LOCTEXT("DeletePerPatternNotImplemented", "FClothAssetDeleteRenderMeshNode: Delete per patterns not implemented."));
+		}
+
+		SetValue<FManagedArrayCollection>(Context, *ClothCollection, &Collection);
+	}
+}
+
+#undef LOCTEXT_NAMESPACE
diff --git a/Engine/Plugins/Experimental/ChaosClothAssetEditor/Source/ChaosClothAssetEditor/Private/ChaosClothAsset/DataflowNodes/DeleteRenderMeshNode.h b/Engine/Plugins/Experimental/ChaosClothAssetEditor/Source/ChaosClothAssetEditor/Private/ChaosClothAsset/DataflowNodes/DeleteRenderMeshNode.h
new file mode 100644
--- /dev/null
+++ b/Engine/Plugins/Experimental/ChaosClothAssetEditor/Source/ChaosClothAssetEditor/Private/ChaosClothAsset/DataflowNodes/DeleteRenderMeshNode.h
@@ -0,0 +1,28 @@
+// Copyright Epic Games, Inc. All Rights Reserved.
+
+#pragma once 
+
+#include "Dataflow/DataflowNode.h"
+#include "GeometryCollection/ManagedArrayCollection.h"
+#include "DeleteRenderMeshNode.generated.h"
+
+USTRUCT(Meta = (DataflowCloth))
+struct FChaosClothAssetDeleteRenderMeshNode : public FDataflowNode
+{
+	GENERATED_USTRUCT_BODY()
+	DATAFLOW_NODE_DEFINE_INTERNAL(FChaosClothAssetDeleteRenderMeshNode, "DeleteRenderMesh", "Cloth", "Cloth Delete Render Mesh")
+	//DATAFLOW_NODE_RENDER_TYPE(FGeometryCollection::StaticType(), "Collection")  // TODO: Leave out the render type until there is something to render
+
+public:
+
+	UPROPERTY(Meta = (Dataflowinput, DataflowOutput, DisplayName = "Collection", DataflowPassthrough = "Collection"))
+	FManagedArrayCollection Collection;
+
+	/** List of patterns to apply the operation on. All patterns will be used if left empty. */
+	UPROPERTY(EditAnywhere, Category = "Pattern Selection")
+	TArray<int32> Patterns;
+
+	FChaosClothAssetDeleteRenderMeshNode(const Dataflow::FNodeParameters& InParam, FGuid InGuid = FGuid::NewGuid());
+
+	virtual void Evaluate(Dataflow::FContext& Context, const FDataflowOutput* Out) const override;
+};
diff --git a/Engine/Plugins/Experimental/ChaosClothAssetEditor/Source/ChaosClothAssetEditor/Private/ChaosClothAsset/DataflowNodes/ImportNode.cpp b/Engine/Plugins/Experimental/ChaosClothAssetEditor/Source/ChaosClothAssetEditor/Private/ChaosClothAsset/DataflowNodes/ImportNode.cpp
new file mode 100644
--- /dev/null
+++ b/Engine/Plugins/Experimental/ChaosClothAssetEditor/Source/ChaosClothAssetEditor/Private/ChaosClothAsset/DataflowNodes/ImportNode.cpp
@@ -0,0 +1,103 @@
+// Copyright Epic Games, Inc. All Rights Reserved.
+
+#include "ChaosClothAsset/DataflowNodes/ImportNode.h"
+#include "ChaosClothAsset/DataflowNodes/DataflowNodes.h"
+#include "ChaosClothAsset/ClothAsset.h"
+#include "ChaosClothAsset/ClothCollection.h"
+#include "Engine/SkinnedAssetCommon.h"
+#include "Materials/Material.h"
+#include "PhysicsEngine/PhysicsAsset.h"
+
+#include UE_INLINE_GENERATED_CPP_BY_NAME(ImportNode)
+
+#define LOCTEXT_NAMESPACE "ChaosClothAssetImportNode"
+
+FChaosClothAssetImportNode::FChaosClothAssetImportNode(const Dataflow::FNodeParameters& InParam, FGuid InGuid)
+	: FDataflowNode(InParam, InGuid)
+{
+	RegisterOutputConnection(&Collection);
+}
+
+void FChaosClothAssetImportNode::Evaluate(Dataflow::FContext& Context, const FDataflowOutput* Out) const
+{
+	using namespace UE::Chaos::ClothAsset;
+	using namespace UE::Chaos::ClothAsset::DataflowNodes;
+
+	if (Out->IsA<FManagedArrayCollection>(&Collection))
+	{
+		bool bSetValue = false;
+		if (ClothAsset)
+		{
+			// Copy the main cloth asset details to this dataflow's owner if any
+			if (const Dataflow::FEngineContext* EngineContext = Context.AsType<Dataflow::FEngineContext>())
+			{
+				if (const UChaosClothAsset* const OwnerClothAsset = Cast<UChaosClothAsset>(EngineContext->Owner))
+				{
+					if (OwnerClothAsset == ClothAsset)
+					{
+						// Can't create a loop
+						bSetValue = false;
+						LogAndToastWarning(LOCTEXT("RecursiveAssetLoop", "FClothAssetNode: The source asset cannot be the terminal asset."));
+					}
+					else
+					{
+						bSetValue = true;
+					}
+				}
+			}
+			else
+			{
+				// No terminal asset, this is a stray dataflow and it is safe to set the value without fear of a loop
+				bSetValue = true;
+			}
+		}
+
+		// Copy the cloth asset to this node's output collection
+		if (bSetValue)
+		{
+			// TODO: Needs to make cloth collection a facade to avoid the copy of the cloth collection to a managed array
+			TSharedPtr<FClothCollection> ClothCollection = MakeShared<FClothCollection>();
+
+			if (ClothAsset->GetClothCollection())
+			{
+				const TArray<FName> GroupsToSkip =
+				{
+					FClothCollection::MaterialsGroup,   // Don't copy materials, and collisions but rebuild the
+					FClothCollection::CollisionsGroup   // group from the asset for backward compatibility instead
+				};
+				ClothAsset->GetClothCollection()->CopyTo(ClothCollection.Get(), GroupsToSkip);
+			}
+
+			// Rebuild material list
+			const TArray<FSkeletalMaterial>& Materials = ClothAsset->GetMaterials();
+			const int32 NumMaterials = Materials.Num();
+
+			ClothCollection->AddElements(NumMaterials, FClothCollection::MaterialsGroup);
+
+			for (int32 MaterialIndex = 0; MaterialIndex < NumMaterials; ++MaterialIndex)
+			{
+				const FSkeletalMaterial& Material = Materials[MaterialIndex];
+				ClothCollection->MaterialPathName[MaterialIndex] = Material.MaterialInterface ?
+					Material.MaterialInterface->GetPathName() :
+					FString();
+			}
+
+			// Overwrite physics asset
+			if (const UPhysicsAsset* const PhysicsAsset = ClothAsset->GetPhysicsAsset())
+			{
+				ClothCollection->AddElements(1, FClothCollection::CollisionsGroup);
+				ClothCollection->PhysicsAssetPathName[0] = PhysicsAsset->GetPathName();
+			}
+
+			SetValue<FManagedArrayCollection>(Context, *ClothCollection, &Collection);
+		}
+		else
+		{
+			// Init with an empty cloth collection
+			UE::Chaos::ClothAsset::FClothCollection ClothCollection;
+			SetValue<FManagedArrayCollection>(Context, ClothCollection, &Collection);
+		}
+	}
+}
+
+#undef LOCTEXT_NAMESPACE
diff --git a/Engine/Plugins/Experimental/ChaosClothAssetEditor/Source/ChaosClothAssetEditor/Private/ChaosClothAsset/DataflowNodes/ImportNode.h b/Engine/Plugins/Experimental/ChaosClothAssetEditor/Source/ChaosClothAssetEditor/Private/ChaosClothAsset/DataflowNodes/ImportNode.h
new file mode 100644
--- /dev/null
+++ b/Engine/Plugins/Experimental/ChaosClothAssetEditor/Source/ChaosClothAssetEditor/Private/ChaosClothAsset/DataflowNodes/ImportNode.h
@@ -0,0 +1,29 @@
+// Copyright Epic Games, Inc. All Rights Reserved.
+
+#pragma once 
+
+#include "Dataflow/DataflowNode.h"
+#include "GeometryCollection/ManagedArrayCollection.h"
+#include "ImportNode.generated.h"
+
+class UChaosClothAsset;
+
+USTRUCT(Meta = (DataflowCloth))
+struct FChaosClothAssetImportNode : public FDataflowNode
+{
+	GENERATED_USTRUCT_BODY()
+	DATAFLOW_NODE_DEFINE_INTERNAL(FChaosClothAssetImportNode, "ClothAssetImport", "Cloth", "Cloth Asset Import")
+	//DATAFLOW_NODE_RENDER_TYPE(FGeometryCollection::StaticType(), "Collection")  // TODO: Leave out the render type until there is something to render
+
+public:
+
+	UPROPERTY(Meta = (DataflowOutput, DisplayName = "Collection"))
+	FManagedArrayCollection Collection;
+
+	UPROPERTY(EditAnywhere, Category = "Dataflow")
+	TObjectPtr<const UChaosClothAsset> ClothAsset;
+
+	FChaosClothAssetImportNode(const Dataflow::FNodeParameters& InParam, FGuid InGuid = FGuid::NewGuid());
+
+	virtual void Evaluate(Dataflow::FContext& Context, const FDataflowOutput* Out) const override;
+};
diff --git a/Engine/Plugins/Experimental/ChaosClothAssetEditor/Source/ChaosClothAssetEditor/Private/ChaosClothAsset/DataflowNodes/PatternSelectionNode.cpp b/Engine/Plugins/Experimental/ChaosClothAssetEditor/Source/ChaosClothAssetEditor/Private/ChaosClothAsset/DataflowNodes/PatternSelectionNode.cpp
new file mode 100644
--- /dev/null
+++ b/Engine/Plugins/Experimental/ChaosClothAssetEditor/Source/ChaosClothAssetEditor/Private/ChaosClothAsset/DataflowNodes/PatternSelectionNode.cpp
@@ -0,0 +1,26 @@
+// Copyright Epic Games, Inc. All Rights Reserved.
+
+#include "ChaosClothAsset/DataflowNodes/PatternSelectionNode.h"
+#include "ChaosClothAsset/DataflowNodes/DataflowNodes.h"
+
+#include UE_INLINE_GENERATED_CPP_BY_NAME(PatternSelectionNode)
+
+#define LOCTEXT_NAMESPACE "ChaosClothAssetPatternSelectionNode"
+
+FChaosClothAssetPatternSelectionNode::FChaosClothAssetPatternSelectionNode(const Dataflow::FNodeParameters& InParam, FGuid InGuid)
+	: FDataflowNode(InParam, InGuid)
+{
+	RegisterOutputConnection(&Patterns);
+}
+
+void FChaosClothAssetPatternSelectionNode::Evaluate(Dataflow::FContext& Context, const FDataflowOutput* Out) const
+{
+	if (Out->IsA<TArray<int32>>(&Patterns))
+	{
+		// Init with empty pattern indices
+		const TArray<int32> Selection;
+		SetValue<TArray<int32>>(Context, Selection, &Patterns);
+	}
+}
+
+#undef LOCTEXT_NAMESPACE
diff --git a/Engine/Plugins/Experimental/ChaosClothAssetEditor/Source/ChaosClothAssetEditor/Private/ChaosClothAsset/DataflowNodes/PatternSelectionNode.h b/Engine/Plugins/Experimental/ChaosClothAssetEditor/Source/ChaosClothAssetEditor/Private/ChaosClothAsset/DataflowNodes/PatternSelectionNode.h
new file mode 100644
--- /dev/null
+++ b/Engine/Plugins/Experimental/ChaosClothAssetEditor/Source/ChaosClothAssetEditor/Private/ChaosClothAsset/DataflowNodes/PatternSelectionNode.h
@@ -0,0 +1,23 @@
+// Copyright Epic Games, Inc. All Rights Reserved.
+
+#pragma once 
+
+#include "Dataflow/DataflowNode.h"
+#include "PatternSelectionNode.generated.h"
+
+USTRUCT(Meta = (DataflowCloth))
+struct FChaosClothAssetPatternSelectionNode : public FDataflowNode
+{
+	GENERATED_USTRUCT_BODY()
+	DATAFLOW_NODE_DEFINE_INTERNAL(FChaosClothAssetPatternSelectionNode, "PatternSelection", "Cloth", "Cloth Pattern Selection")
+	//DATAFLOW_NODE_RENDER_TYPE(FGeometryCollection::StaticType(), "Collection")  // TODO: Leave out the render type until there is something to render
+
+public:
+
+	UPROPERTY(EditAnywhere, Category = "Pattern", Meta = (DataflowOutput, DisplayName = "Patterns"))
+	TArray<int32> Patterns;
+
+	FChaosClothAssetPatternSelectionNode(const Dataflow::FNodeParameters& InParam, FGuid InGuid = FGuid::NewGuid());
+
+	virtual void Evaluate(Dataflow::FContext& Context, const FDataflowOutput* Out) const override;
+};
diff --git a/Engine/Plugins/Experimental/ChaosClothAssetEditor/Source/ChaosClothAssetEditor/Private/ChaosClothAsset/DataflowNodes/ReverseNormalsNode.cpp b/Engine/Plugins/Experimental/ChaosClothAssetEditor/Source/ChaosClothAssetEditor/Private/ChaosClothAsset/DataflowNodes/ReverseNormalsNode.cpp
new file mode 100644
--- /dev/null
+++ b/Engine/Plugins/Experimental/ChaosClothAssetEditor/Source/ChaosClothAssetEditor/Private/ChaosClothAsset/DataflowNodes/ReverseNormalsNode.cpp
@@ -0,0 +1,39 @@
+// Copyright Epic Games, Inc. All Rights Reserved.
+
+#include "ChaosClothAsset/DataflowNodes/ReverseNormalsNode.h"
+#include "ChaosClothAsset/DataflowNodes/DataflowNodes.h"
+#include "ChaosClothAsset/ClothGeometryTools.h"
+
+#include UE_INLINE_GENERATED_CPP_BY_NAME(ReverseNormalsNode)
+
+#define LOCTEXT_NAMESPACE "ChaosClothAssetReverseNormalsNode"
+
+FChaosClothAssetReverseNormalsNode::FChaosClothAssetReverseNormalsNode(const Dataflow::FNodeParameters& InParam, FGuid InGuid)
+	: FDataflowNode(InParam, InGuid)
+{
+	RegisterInputConnection(&Collection);
+	RegisterInputConnection(&Patterns);
+	RegisterOutputConnection(&Collection, &Collection);
+}
+
+void FChaosClothAssetReverseNormalsNode::Evaluate(Dataflow::FContext& Context, const FDataflowOutput* Out) const
+{
+	if (Out->IsA<FManagedArrayCollection>(&Collection))
+	{
+		using namespace UE::Chaos::ClothAsset;
+
+		// Evaluate in collection
+		const FManagedArrayCollection& InCollection = GetValue<FManagedArrayCollection>(Context, &Collection);
+
+		// TODO: Needs to make cloth collection a facade to avoid the copy of the cloth collection to a managed array
+		//       and remove the above const reference to operated on the InCollection instead.
+		TSharedPtr<FClothCollection> ClothCollection = MakeShared<FClothCollection>();
+		InCollection.CopyTo(ClothCollection.Get());
+
+		FClothGeometryTools::ReverseNormals(ClothCollection, bReverseSimMeshNormals, bReverseRenderMeshNormals, Patterns);
+
+		SetValue<FManagedArrayCollection>(Context, *ClothCollection, &Collection);
+	}
+}
+
+#undef LOCTEXT_NAMESPACE
diff --git a/Engine/Plugins/Experimental/ChaosClothAssetEditor/Source/ChaosClothAssetEditor/Private/ChaosClothAsset/DataflowNodes/ReverseNormalsNode.h b/Engine/Plugins/Experimental/ChaosClothAssetEditor/Source/ChaosClothAssetEditor/Private/ChaosClothAsset/DataflowNodes/ReverseNormalsNode.h
new file mode 100644
--- /dev/null
+++ b/Engine/Plugins/Experimental/ChaosClothAssetEditor/Source/ChaosClothAssetEditor/Private/ChaosClothAsset/DataflowNodes/ReverseNormalsNode.h
@@ -0,0 +1,36 @@
+// Copyright Epic Games, Inc. All Rights Reserved.
+
+#pragma once 
+
+#include "Dataflow/DataflowNode.h"
+#include "GeometryCollection/ManagedArrayCollection.h"
+#include "ReverseNormalsNode.generated.h"
+
+USTRUCT(Meta = (DataflowCloth))
+struct FChaosClothAssetReverseNormalsNode : public FDataflowNode
+{
+	GENERATED_USTRUCT_BODY()
+	DATAFLOW_NODE_DEFINE_INTERNAL(FChaosClothAssetReverseNormalsNode, "ReverseNormals", "Cloth", "Cloth Reverse Simulation Render Mesh Normals")
+	//DATAFLOW_NODE_RENDER_TYPE(FGeometryCollection::StaticType(), "Collection")  // TODO: Leave out the render type until there is something to render
+
+public:
+
+	UPROPERTY(Meta = (Dataflowinput, DataflowOutput, DisplayName = "Collection", DataflowPassthrough = "Collection"))
+	FManagedArrayCollection Collection;
+
+	/** List of patterns to apply the operation on. All patterns will be used if left empty. */
+	UPROPERTY(EditAnywhere, Category = "Pattern Selection")
+	TArray<int32> Patterns;
+
+	/** Whether to reverse the simulation mesh normals. */
+	UPROPERTY(EditAnywhere, Category = "Reverse Normals")
+	bool bReverseSimMeshNormals = true;
+
+	/** Whether to reverse the render mesh normals. */
+	UPROPERTY(EditAnywhere, Category = "Reverse Normals")
+	bool bReverseRenderMeshNormals = true;
+
+	FChaosClothAssetReverseNormalsNode(const Dataflow::FNodeParameters& InParam, FGuid InGuid = FGuid::NewGuid());
+
+	virtual void Evaluate(Dataflow::FContext& Context, const FDataflowOutput* Out) const override;
+};
diff --git a/Engine/Plugins/Experimental/ChaosClothAssetEditor/Source/ChaosClothAssetEditor/Private/ChaosClothAsset/DataflowNodes/SimulationDefaultConfigNode.cpp b/Engine/Plugins/Experimental/ChaosClothAssetEditor/Source/ChaosClothAssetEditor/Private/ChaosClothAsset/DataflowNodes/SimulationDefaultConfigNode.cpp
new file mode 100644
--- /dev/null
+++ b/Engine/Plugins/Experimental/ChaosClothAssetEditor/Source/ChaosClothAssetEditor/Private/ChaosClothAsset/DataflowNodes/SimulationDefaultConfigNode.cpp
@@ -0,0 +1,66 @@
+// Copyright Epic Games, Inc. All Rights Reserved.
+
+#include "ChaosClothAsset/DataflowNodes/SimulationDefaultConfigNode.h"
+#include "ChaosClothAsset/DataflowNodes/DataflowNodes.h"
+#include "ChaosCloth/ChaosClothConfig.h"
+#include "ChaosCloth/ChaosClothingSimulationConfig.h"
+#include "Chaos/CollectionPropertyFacade.h"
+
+#include UE_INLINE_GENERATED_CPP_BY_NAME(SimulationDefaultConfigNode)
+
+#define LOCTEXT_NAMESPACE "ChaosClothAssetSimulationDefaultConfigNode"
+
+FChaosClothAssetSimulationDefaultConfigNode::FChaosClothAssetSimulationDefaultConfigNode(const Dataflow::FNodeParameters& InParam, FGuid InGuid)
+	: FDataflowNode(InParam, InGuid)
+	, SimulationConfig(NewObject<UChaosClothConfig>())
+	, SharedSimulationConfig(NewObject<UChaosClothSharedSimConfig>())
+{
+	RegisterInputConnection(&Collection);
+	RegisterOutputConnection(&Collection, &Collection);
+}
+
+void FChaosClothAssetSimulationDefaultConfigNode::Serialize(FArchive& Ar)
+{
+	if (!SimulationConfig)
+	{
+		SimulationConfig = NewObject<UChaosClothConfig>();
+	}
+	SimulationConfig->Serialize(Ar);
+
+	if (!SharedSimulationConfig)
+	{
+		SharedSimulationConfig = NewObject<UChaosClothSharedSimConfig>();
+	}
+	SharedSimulationConfig->Serialize(Ar);
+}
+
+void FChaosClothAssetSimulationDefaultConfigNode::Evaluate(Dataflow::FContext& Context, const FDataflowOutput* Out) const
+{
+	using namespace Chaos;
+	using namespace Chaos::Softs;
+
+	if (Out->IsA<FManagedArrayCollection>(&Collection))
+	{
+		FManagedArrayCollection InCollection = GetValue<FManagedArrayCollection>(Context, &Collection);
+
+		FClothingSimulationConfig ClothingSimulationConfig;
+		ClothingSimulationConfig.Initialize(SimulationConfig.Get(), SharedSimulationConfig.Get());
+
+		ClothingSimulationConfig.GetPropertyCollection()->CopyTo(&InCollection);
+
+		SetValue<FManagedArrayCollection>(Context, InCollection, &Collection);
+	}
+}
+
+void FChaosClothAssetSimulationDefaultConfigNode::AddReferencedObjects(FReferenceCollector& Collector)
+{
+	Collector.AddReferencedObject(SimulationConfig);
+	Collector.AddReferencedObject(SharedSimulationConfig);
+}
+
+FString FChaosClothAssetSimulationDefaultConfigNode::GetReferencerName() const
+{
+	return TEXT("FChaosClothAssetSimulationDefaultConfigNode");
+}
+
+#undef LOCTEXT_NAMESPACE
diff --git a/Engine/Plugins/Experimental/ChaosClothAssetEditor/Source/ChaosClothAssetEditor/Private/ChaosClothAsset/DataflowNodes/SimulationDefaultConfigNode.h b/Engine/Plugins/Experimental/ChaosClothAssetEditor/Source/ChaosClothAssetEditor/Private/ChaosClothAsset/DataflowNodes/SimulationDefaultConfigNode.h
new file mode 100644
--- /dev/null
+++ b/Engine/Plugins/Experimental/ChaosClothAssetEditor/Source/ChaosClothAssetEditor/Private/ChaosClothAsset/DataflowNodes/SimulationDefaultConfigNode.h
@@ -0,0 +1,43 @@
+// Copyright Epic Games, Inc. All Rights Reserved.
+
+#pragma once 
+
+#include "Dataflow/DataflowNode.h"
+#include "GeometryCollection/ManagedArrayCollection.h"
+#include "SimulationDefaultConfigNode.generated.h"
+
+class UChaosClothConfig;
+class UChaosClothSharedSimConfig;
+class UMaterialInterface;
+
+USTRUCT(Meta = (DataflowCloth))
+struct FChaosClothAssetSimulationDefaultConfigNode : public FDataflowNode, public FGCObject
+{
+	GENERATED_USTRUCT_BODY()
+	DATAFLOW_NODE_DEFINE_INTERNAL(FChaosClothAssetSimulationDefaultConfigNode, "SimulationDefaultConfig", "Cloth", "Cloth Simulation Default Config")
+	//DATAFLOW_NODE_RENDER_TYPE(FManagedArrayCollection::StaticType(), "Collection")  // TODO: Leave out the render type until there is something to render
+
+public:
+
+	UPROPERTY(Meta = (Dataflowinput, DataflowOutput, DisplayName = "Collection", DataflowPassthrough = "Collection"))
+	FManagedArrayCollection Collection;
+
+	/** Cloth Simulation Properties. */
+	UPROPERTY(EditAnywhere, Instanced, NoClear, Category = "Dataflow")
+	TObjectPtr<UChaosClothConfig> SimulationConfig;
+
+	/** Cloth Shared Simulation Properties. */
+	UPROPERTY(EditAnywhere, Instanced, NoClear, Category = "Dataflow")
+	TObjectPtr<UChaosClothSharedSimConfig> SharedSimulationConfig;
+
+	FChaosClothAssetSimulationDefaultConfigNode(const Dataflow::FNodeParameters& InParam, FGuid InGuid = FGuid::NewGuid());
+
+	virtual void Serialize(FArchive& Ar) override;
+
+	virtual void Evaluate(Dataflow::FContext& Context, const FDataflowOutput* Out) const override;
+
+	//~ Begin FGCObject Interface
+	virtual void AddReferencedObjects(FReferenceCollector& Collector) override;
+	virtual FString GetReferencerName() const override;
+	//~ End FGCObject Interface
+};
diff --git a/Engine/Plugins/Experimental/ChaosClothAssetEditor/Source/ChaosClothAssetEditor/Private/ChaosClothAsset/DataflowNodes/TerminalNode.cpp b/Engine/Plugins/Experimental/ChaosClothAssetEditor/Source/ChaosClothAssetEditor/Private/ChaosClothAsset/DataflowNodes/TerminalNode.cpp
new file mode 100644
--- /dev/null
+++ b/Engine/Plugins/Experimental/ChaosClothAssetEditor/Source/ChaosClothAssetEditor/Private/ChaosClothAsset/DataflowNodes/TerminalNode.cpp
@@ -0,0 +1,92 @@
+// Copyright Epic Games, Inc. All Rights Reserved.
+
+#include "ChaosClothAsset/DataflowNodes/TerminalNode.h"
+#include "ChaosClothAsset/DataflowNodes/DataflowNodes.h"
+#include "ChaosClothAsset/ClothAsset.h"
+#include "ChaosClothAsset/ClothAdapter.h"
+#include "Animation/Skeleton.h"
+#include "Engine/SkinnedAssetCommon.h"
+#include "Materials/Material.h"
+#include "PhysicsEngine/PhysicsAsset.h"
+
+#include UE_INLINE_GENERATED_CPP_BY_NAME(TerminalNode)
+
+#define LOCTEXT_NAMESPACE "ChaosClothAssetTerminalNode"
+
+namespace UE::Chaos::ClothAsset::DataflowNodes
+{
+	static const FString DefaultSkeletonPathName = TEXT("/Engine/EditorMeshes/SkeletalMesh/DefaultSkeletalMesh_Skeleton.DefaultSkeletalMesh_Skeleton");
+}
+
+FChaosClothAssetTerminalNode::FChaosClothAssetTerminalNode(const Dataflow::FNodeParameters& InParam, FGuid InGuid)
+	: FDataflowTerminalNode(InParam, InGuid)
+{
+	RegisterInputConnection(&Collection);
+}
+
+void FChaosClothAssetTerminalNode::SetAssetValue(TObjectPtr<UObject> Asset, Dataflow::FContext& Context) const
+{
+	if (UChaosClothAsset* ClothAsset = Cast<UChaosClothAsset>(Asset.Get()))
+	{
+		using namespace UE::Chaos::ClothAsset;
+
+		const FManagedArrayCollection& InCollection = GetValue<FManagedArrayCollection>(Context, &Collection);
+
+		FClothCollection& ClothCollection = *ClothAsset->GetClothCollection();
+
+		// Replace collection
+		ClothCollection.Reset();
+		InCollection.CopyTo(&ClothCollection);
+
+		// Set materials
+		const int32 NumMaterials = ClothCollection.MaterialPathName.Num();
+
+		TArray<FSkeletalMaterial>& Materials = ClothAsset->GetMaterials();
+		Materials.Reset(NumMaterials);
+
+		for (int32 MaterialIndex = 0; MaterialIndex < NumMaterials; ++MaterialIndex)
+		{
+			if (UMaterial* const Material = LoadObject<UMaterial>(ClothAsset, *ClothCollection.MaterialPathName[MaterialIndex], nullptr, LOAD_None, nullptr))
+			{
+				Materials.Emplace(Material, true, false, Material->GetFName());
+			}
+			else
+			{
+				Materials.Emplace();
+			}
+		}
+
+		// Set reference skeleton
+		const FString SkeletonPathName = ClothCollection.SkeletonAssetPathName.Num() && !ClothCollection.SkeletonAssetPathName[0].IsEmpty() ?
+			ClothCollection.SkeletonAssetPathName[0] :
+			DataflowNodes::DefaultSkeletonPathName;
+
+		if (const USkeleton* const Skeleton = LoadObject<USkeleton>(nullptr, *SkeletonPathName, nullptr, LOAD_None, nullptr))
+		{
+			constexpr bool bRebuildClothSimulationModel = false;  // Avoid rebuilding the asset twice
+			ClothAsset->SetReferenceSkeleton(Skeleton->GetReferenceSkeleton(), bRebuildClothSimulationModel);
+		}
+
+		// Set physics asset (note: the cloth asset's physics asset is only replaced if a collection path name is found valid)
+		if (ClothCollection.PhysicsAssetPathName.Num() && !ClothCollection.PhysicsAssetPathName[0].IsEmpty())
+		{
+			check(ClothCollection.PhysicsAssetPathName.Num() == 1);  // Can't really deal with more than a single physics asset at the moment
+			if (UPhysicsAsset* const PhysicsAsset = LoadObject<UPhysicsAsset>(ClothAsset, *ClothCollection.PhysicsAssetPathName[0], nullptr, LOAD_None, nullptr))
+			{
+				ClothAsset->SetPhysicsAsset(PhysicsAsset);
+			}
+		}
+
+		// Check there is at least one empty LOD to avoid crashing the render data
+		FClothAdapter ClothAdapter(ClothAsset->GetClothCollection());
+		if (ClothAdapter.GetNumLods() < 1)
+		{
+			ClothAdapter.AddLod();
+		}
+
+		// Rebuild the asset static data
+		ClothAsset->Build();
+	}
+}
+
+#undef LOCTEXT_NAMESPACE
diff --git a/Engine/Plugins/Experimental/ChaosClothAssetEditor/Source/ChaosClothAssetEditor/Private/ChaosClothAsset/DataflowNodes/TerminalNode.h b/Engine/Plugins/Experimental/ChaosClothAssetEditor/Source/ChaosClothAssetEditor/Private/ChaosClothAsset/DataflowNodes/TerminalNode.h
new file mode 100644
--- /dev/null
+++ b/Engine/Plugins/Experimental/ChaosClothAssetEditor/Source/ChaosClothAssetEditor/Private/ChaosClothAsset/DataflowNodes/TerminalNode.h
@@ -0,0 +1,26 @@
+// Copyright Epic Games, Inc. All Rights Reserved.
+
+#pragma once 
+
+#include "Dataflow/DataflowTerminalNode.h"
+#include "GeometryCollection/ManagedArrayCollection.h"
+#include "TerminalNode.generated.h"
+
+USTRUCT(Meta = (DataflowCloth, DataflowTerminal))
+struct FChaosClothAssetTerminalNode : public FDataflowTerminalNode
+{
+	GENERATED_USTRUCT_BODY()
+	DATAFLOW_NODE_DEFINE_INTERNAL(FChaosClothAssetTerminalNode, "ClothAssetTerminal", "Cloth", "Cloth Terminal")  // TODO: Should the category be Terminal instead like all other terminal nodes
+	//DATAFLOW_NODE_RENDER_TYPE(FGeometryCollection::StaticType(), "Collection")  // TODO: Leave out the render type until there is something to render
+
+public:
+
+	UPROPERTY(meta = (DataflowInput, DisplayName = "Collection"))
+	FManagedArrayCollection Collection;
+
+	FChaosClothAssetTerminalNode(const Dataflow::FNodeParameters& InParam, FGuid InGuid = FGuid::NewGuid());
+
+	virtual void SetAssetValue(TObjectPtr<UObject> Asset, Dataflow::FContext& Context) const override;
+
+	virtual void Evaluate(Dataflow::FContext& Context, const FDataflowOutput* Out) const override {}
+};
