Chaos Cloth - Fixed crashes when creating a new empty Cloth Asset.

#rnx

diff --git a/Engine/Plugins/Experimental/ChaosClothAssetEditor/Source/ChaosClothAssetTools/Private/ChaosClothAsset/ClothAssetFactory.cpp b/Engine/Plugins/Experimental/ChaosClothAssetEditor/Source/ChaosClothAssetTools/Private/ChaosClothAsset/ClothAssetFactory.cpp
--- a/Engine/Plugins/Experimental/ChaosClothAssetEditor/Source/ChaosClothAssetTools/Private/ChaosClothAsset/ClothAssetFactory.cpp
+++ b/Engine/Plugins/Experimental/ChaosClothAssetEditor/Source/ChaosClothAssetTools/Private/ChaosClothAsset/ClothAssetFactory.cpp
@@ -2,6 +2,7 @@
 
 #include "ClothAssetFactory.h"
 #include "ChaosClothAsset/ClothAsset.h"
+#include "ChaosClothAsset/ClothAdapter.h"
 
 #include UE_INLINE_GENERATED_CPP_BY_NAME(ClothAssetFactory)
 
@@ -15,8 +16,30 @@ UChaosClothAssetFactory::UChaosClothAssetFactory(const FObjectInitializer& Objec
 
 UObject* UChaosClothAssetFactory::FactoryCreateNew(UClass* Class, UObject* Parent, FName Name, EObjectFlags Flags, UObject* /*Context*/, FFeedbackContext* /*Warn*/)
 {
-	UChaosClothAsset* const NewClothAsset = NewObject<UChaosClothAsset>(Parent, Class, Name, Flags | RF_Transactional | RF_Public | RF_Standalone);
-	NewClothAsset->MarkPackageDirty();
-	return NewClothAsset;
+	using namespace UE::Chaos::ClothAsset;
+
+	UChaosClothAsset* const ClothAsset = NewObject<UChaosClothAsset>(Parent, Class, Name, Flags | RF_Transactional | RF_Public | RF_Standalone);
+	ClothAsset->MarkPackageDirty();
+
+	// Add an empty default LOD, to avoid LOD mismatch with render data
+	TSharedPtr<FClothCollection> ClothCollection = ClothAsset->GetClothCollection();
+
+	FClothAdapter ClothAdapter(ClothCollection);
+	ClothAdapter.AddLod();
+
+	// Set reference skeleton in both the collection and the asset
+	ClothCollection->AddElements(1, FClothCollection::SkeletonsGroup);
+	ClothCollection->SkeletonAssetPathName[0] = TEXT("/Engine/EditorMeshes/SkeletalMesh/DefaultSkeletalMesh_Skeleton.DefaultSkeletalMesh_Skeleton");
+
+	if (const USkeleton* const Skeleton = LoadObject<USkeleton>(nullptr, *ClothCollection->SkeletonAssetPathName[0], nullptr, LOAD_None, nullptr))
+	{
+		constexpr bool bRebuildClothSimulationModel = false;  // Avoid rebuilding the asset twice
+		ClothAsset->SetReferenceSkeleton(Skeleton->GetReferenceSkeleton(), bRebuildClothSimulationModel);
+	}
+
+	// Build static data
+	ClothAsset->Build();
+
+	return ClothAsset;
 }
 
