CIS - Fix static Analysis on cloth terminal node priority of operator ?:.

#rnx

diff --git a/Engine/Plugins/Experimental/ChaosClothAssetEditor/Source/ChaosClothAssetEditor/Private/ChaosClothAsset/DataflowNodes/TerminalNode.cpp b/Engine/Plugins/Experimental/ChaosClothAssetEditor/Source/ChaosClothAssetEditor/Private/ChaosClothAsset/DataflowNodes/TerminalNode.cpp
--- a/Engine/Plugins/Experimental/ChaosClothAssetEditor/Source/ChaosClothAssetEditor/Private/ChaosClothAsset/DataflowNodes/TerminalNode.cpp
+++ b/Engine/Plugins/Experimental/ChaosClothAssetEditor/Source/ChaosClothAssetEditor/Private/ChaosClothAsset/DataflowNodes/TerminalNode.cpp
@@ -57,7 +57,7 @@ void FChaosClothAssetTerminalNode::SetAssetValue(TObjectPtr<UObject> Asset, Data
 		}
 
 		// Set reference skeleton
-		const FString SkeletonPathName = ClothCollection.SkeletonAssetPathName.Num() && !ClothCollection.SkeletonAssetPathName[0].IsEmpty() ?
+		const FString SkeletonPathName = (ClothCollection.SkeletonAssetPathName.Num() && !ClothCollection.SkeletonAssetPathName[0].IsEmpty()) ?
 			ClothCollection.SkeletonAssetPathName[0] :
 			DataflowNodes::DefaultSkeletonPathName;
 
