Cloth Editor: re-initialize the editor viewports when the Cloth Asset changes due to Dataflow graph evaluation
- clear any selection and gizmo in the Preview viewport
- maintain the SimulationSuspended state for the Editor

diff --git a/Engine/Plugins/Experimental/ChaosClothAssetEditor/Source/ChaosClothAssetEditor/Private/ChaosClothAsset/ClothEditor3DViewportClient.cpp b/Engine/Plugins/Experimental/ChaosClothAssetEditor/Source/ChaosClothAssetEditor/Private/ChaosClothAsset/ClothEditor3DViewportClient.cpp
--- a/Engine/Plugins/Experimental/ChaosClothAssetEditor/Source/ChaosClothAssetEditor/Private/ChaosClothAsset/ClothEditor3DViewportClient.cpp
+++ b/Engine/Plugins/Experimental/ChaosClothAssetEditor/Source/ChaosClothAssetEditor/Private/ChaosClothAsset/ClothEditor3DViewportClient.cpp
@@ -43,7 +43,7 @@ FChaosClothAssetEditor3DViewportClient::FChaosClothAssetEditor3DViewportClient(F
 	Gizmo = GizmoManager->Create3AxisTransformGizmo(this, GizmoIdentifier);
 
 	Gizmo->SetActiveTarget(TransformProxy);
-	Gizmo->SetVisibility(true);
+	Gizmo->SetVisibility(false);
 	Gizmo->bUseContextGizmoMode = false;
 	Gizmo->bUseContextCoordinateSystem = false;
 	Gizmo->ActiveGizmoMode = EToolContextTransformGizmoMode::Combined;
@@ -52,11 +52,15 @@ FChaosClothAssetEditor3DViewportClient::FChaosClothAssetEditor3DViewportClient(F
 	SetAdvancedShowFlagsForScene(UAssetViewerSettings::Get()->Profiles[GetMutableDefault<UEditorPerProjectUserSettings>()->AssetViewerProfileIndex].bPostProcessingEnabled);
 }
 
-void FChaosClothAssetEditor3DViewportClient::RegisterSettingsChangedDelegate()
+void FChaosClothAssetEditor3DViewportClient::RegisterDelegates()
 {
 	// Remove any existing delegate in case this function is called twice
 	UAssetViewerSettings::Get()->OnAssetViewerSettingsChanged().RemoveAll(this);
 	UAssetViewerSettings::Get()->OnAssetViewerSettingsChanged().AddSP(this, &FChaosClothAssetEditor3DViewportClient::OnAssetViewerSettingsChanged);
+
+	USelection* const SelectedComponents = ModeTools->GetSelectedComponents();
+	SelectedComponents->SelectionChangedEvent.RemoveAll(this);
+	SelectedComponents->SelectionChangedEvent.AddSP(this, &FChaosClothAssetEditor3DViewportClient::ComponentSelectionChanged);
 }
 
 FChaosClothAssetEditor3DViewportClient::~FChaosClothAssetEditor3DViewportClient()
@@ -64,6 +68,10 @@ FChaosClothAssetEditor3DViewportClient::~FChaosClothAssetEditor3DViewportClient(
 	DeleteViewportGizmo();
 
 	UAssetViewerSettings::Get()->OnAssetViewerSettingsChanged().RemoveAll(this);
+	if (USelection* const SelectedComponents = ModeTools->GetSelectedComponents())
+	{
+		SelectedComponents->SelectionChangedEvent.RemoveAll(this);
+	}
 }
 
 void FChaosClothAssetEditor3DViewportClient::DeleteViewportGizmo()
@@ -148,13 +156,8 @@ void FChaosClothAssetEditor3DViewportClient::EnableRenderMeshWireframe(bool bEna
 
 void FChaosClothAssetEditor3DViewportClient::SetClothComponent(TObjectPtr<UChaosClothComponent> InClothComponent)
 {
+	ClearSelectedComponents();
 	ClothComponent = InClothComponent;
-
-	if (ClothComponent)
-	{
-		TransformProxy->AddComponent(ClothComponent);
-		TransformProxy->SetTransform(ClothComponent->GetComponentToWorld());
-	}
 }
 
 void FChaosClothAssetEditor3DViewportClient::SetClothEdMode(TObjectPtr<UChaosClothAssetEditorMode> InClothEdMode)
@@ -237,6 +240,7 @@ void FChaosClothAssetEditor3DViewportClient::ProcessClick(FSceneView& View, HHit
 	SelectedComponents->GetSelectedObjects<UPrimitiveComponent>(PreviouslySelectedComponents);
 
 	SelectedComponents->Modify();
+	SelectedComponents->BeginBatchSelectOperation();
 
 	if (!bIsShiftKeyDown && !bIsCtrlKeyDown)
 	{
@@ -285,11 +289,33 @@ void FChaosClothAssetEditor3DViewportClient::ProcessClick(FSceneView& View, HHit
 		}
 	}
 
+	SelectedComponents->EndBatchSelectOperation();
+
+	for (UPrimitiveComponent* const Component : PreviouslySelectedComponents)
+	{
+		Component->PushSelectionToProxy();
+	}
+}
+
+
+void FChaosClothAssetEditor3DViewportClient::ClearSelectedComponents()
+{
+	USelection* const SelectedComponents = ModeTools->GetSelectedComponents();
+	TArray<UPrimitiveComponent*> PreviouslySelectedComponents;
+	SelectedComponents->GetSelectedObjects<UPrimitiveComponent>(PreviouslySelectedComponents);
+
 	for (UPrimitiveComponent* const Component : PreviouslySelectedComponents)
 	{
 		Component->PushSelectionToProxy();
 	}
 
+	SelectedComponents->DeselectAll();
+}
+
+void FChaosClothAssetEditor3DViewportClient::ComponentSelectionChanged(UObject* NewSelection)
+{
+	USelection* const SelectedComponents = ModeTools->GetSelectedComponents();
+
 	// Update TransformProxy
 
 	TransformProxy = NewObject<UTransformProxy>();
@@ -301,18 +327,19 @@ void FChaosClothAssetEditor3DViewportClient::ProcessClick(FSceneView& View, HHit
 	}
 
 	// Update gizmo
-
-	if (Components.Num() > 0)
+	if (Gizmo)
 	{
-		Gizmo->SetActiveTarget(TransformProxy);
-		Gizmo->SetVisibility(true);
-	}
-	else
-	{
-		Gizmo->ClearActiveTarget();
-		Gizmo->SetVisibility(false);
+		if (Components.Num() > 0)
+		{
+			Gizmo->SetActiveTarget(TransformProxy);
+			Gizmo->SetVisibility(true);
+		}
+		else
+		{
+			Gizmo->ClearActiveTarget();
+			Gizmo->SetVisibility(false);
+		}
 	}
-
 }
 
 
diff --git a/Engine/Plugins/Experimental/ChaosClothAssetEditor/Source/ChaosClothAssetEditor/Private/ChaosClothAsset/ClothEditor3DViewportClient.h b/Engine/Plugins/Experimental/ChaosClothAssetEditor/Source/ChaosClothAssetEditor/Private/ChaosClothAsset/ClothEditor3DViewportClient.h
--- a/Engine/Plugins/Experimental/ChaosClothAssetEditor/Source/ChaosClothAssetEditor/Private/ChaosClothAsset/ClothEditor3DViewportClient.h
+++ b/Engine/Plugins/Experimental/ChaosClothAssetEditor/Source/ChaosClothAssetEditor/Private/ChaosClothAsset/ClothEditor3DViewportClient.h
@@ -25,13 +25,15 @@ public:
 		const TWeakPtr<SEditorViewport>& InEditorViewportWidget = nullptr);
 
 	// Call this after construction to initialize callbacks when settings change
-	void RegisterSettingsChangedDelegate();
+	void RegisterDelegates();
 
 	virtual ~FChaosClothAssetEditor3DViewportClient();
 
 	// Delete the viewport gizmo and transform proxy
 	void DeleteViewportGizmo();
 
+	void ClearSelectedComponents();
+
 	void EnableSimMeshWireframe(bool bEnable ) { bSimMeshWireframe = bEnable; }
 	bool SimMeshWireframeEnabled() const { return bSimMeshWireframe; }
 
@@ -66,6 +68,8 @@ private:
 	void OnAssetViewerSettingsChanged(const FName& InPropertyName);
 	void SetAdvancedShowFlagsForScene(const bool bAdvancedShowFlags);
 
+	void ComponentSelectionChanged(UObject* NewSelection);
+
 	TWeakPtr<FChaosClothPreviewScene> ClothPreviewScene;
 
 	TObjectPtr<UChaosClothComponent> ClothComponent;
diff --git a/Engine/Plugins/Experimental/ChaosClothAssetEditor/Source/ChaosClothAssetEditor/Private/ChaosClothAsset/ClothEditorToolkit.cpp b/Engine/Plugins/Experimental/ChaosClothAssetEditor/Source/ChaosClothAssetEditor/Private/ChaosClothAsset/ClothEditorToolkit.cpp
--- a/Engine/Plugins/Experimental/ChaosClothAssetEditor/Source/ChaosClothAssetEditor/Private/ChaosClothAsset/ClothEditorToolkit.cpp
+++ b/Engine/Plugins/Experimental/ChaosClothAssetEditor/Source/ChaosClothAssetEditor/Private/ChaosClothAsset/ClothEditorToolkit.cpp
@@ -179,7 +179,7 @@ FChaosClothAssetEditorToolkit::FChaosClothAssetEditorToolkit(UAssetEditor* InOwn
 	//ClothPreviewInputRouter = ClothPreviewEditorModeManager->GetInteractiveToolsContext()->InputRouter;
 	ClothPreviewTabContent = MakeShareable(new FEditorViewportTabContent());
 	ClothPreviewViewportClient = MakeShared<FChaosClothAssetEditor3DViewportClient>(ClothPreviewEditorModeManager.Get(), ClothPreviewScene);
-	ClothPreviewViewportClient->RegisterSettingsChangedDelegate();
+	ClothPreviewViewportClient->RegisterDelegates();
 
 	ClothPreviewViewportDelegate = [this](FAssetEditorViewportConstructionArgs InArgs)
 	{
@@ -231,7 +231,14 @@ void FChaosClothAssetEditorToolkit::Tick(float DeltaTime)
 			LastDataflowNodeTimestamp = Dataflow::FTimestamp::Invalid;
 		}
 		DataflowTerminalPath = ClothEditorToolkitHelpers::GetDataflowTerminalFrom(ClothAsset);
+
+		Dataflow::FTimestamp OldTimestamp = LastDataflowNodeTimestamp;
 		FDataflowEditorCommands::EvaluateTerminalNode(*DataflowContext.Get(), LastDataflowNodeTimestamp, Dataflow, nullptr, nullptr, ClothAsset, DataflowTerminalPath);
+
+		if (OldTimestamp.Value < LastDataflowNodeTimestamp.Value)
+		{
+			OnClothAssetChanged();
+		}
 	}
 }
 
@@ -759,9 +766,13 @@ void FChaosClothAssetEditorToolkit::EvaluateNode(FDataflowNode* Node, FDataflowO
 		}
 		LastDataflowNodeTimestamp = Dataflow::FTimestamp::Invalid;
 
+		Dataflow::FTimestamp OldTimestamp = LastDataflowNodeTimestamp;
 		FDataflowEditorCommands::EvaluateTerminalNode(*DataflowContext.Get(), LastDataflowNodeTimestamp, Dataflow, Node, nullptr, Asset, DataflowTerminalPath);
 
-		OnClothAssetChanged();
+		if (OldTimestamp.Value < LastDataflowNodeTimestamp.Value)
+		{
+			OnClothAssetChanged();
+		}
 	}
 };
 
@@ -910,10 +921,15 @@ void FChaosClothAssetEditorToolkit::OnNodeSelectionChanged(const TSet<UObject*>&
 
 void FChaosClothAssetEditorToolkit::OnClothAssetChanged()
 {
+	ClothPreviewViewportClient->ClearSelectedComponents();
+
 	TArray<TObjectPtr<UObject>> ObjectsToEdit;
 	OwningAssetEditor->GetObjectsToEdit(ObjectsToEdit);
 
 	UChaosClothAssetEditorMode* const ClothMode = CastChecked<UChaosClothAssetEditorMode>(EditorModeManager->GetActiveScriptableMode(UChaosClothAssetEditorMode::EM_ChaosClothAssetEditorModeId));
+
+	const bool bWasSimulationSuspended = ClothMode->IsSimulationSuspended();
+
 	ClothMode->InitializeTargets(ObjectsToEdit);
 
 	if (UChaosClothAsset* const ClothAsset = Cast<UChaosClothAsset>(ObjectsToEdit[0]))
@@ -923,8 +939,20 @@ void FChaosClothAssetEditorToolkit::OnClothAssetChanged()
 		ensure(ClothAsset->HasAnyFlags(RF_Transactional));		// Ensure all objects are transactable for undo/redo in the details panel
 		SetEditingObject(ClothAsset);
 
+		ClothPreviewViewportClient->FocusViewportOnBox(ClothMode->PreviewBoundingBox());
+		ClothMode->RefocusRestSpaceViewportClient();
+
 		PopulateOutliner();
 	}
+
+	if (bWasSimulationSuspended)
+	{
+		ClothMode->SuspendSimulation();
+	}
+	else
+	{
+		ClothMode->ResumeSimulation();
+	}
 }
 
 
