Enable the use of skinning weights in Chaos Cloth.

diff --git a/Engine/Plugins/Enterprise/DatasmithImporter/Source/DatasmithImporter/Private/DatasmithImporter.cpp b/Engine/Plugins/Enterprise/DatasmithImporter/Source/DatasmithImporter/Private/DatasmithImporter.cpp
--- a/Engine/Plugins/Enterprise/DatasmithImporter/Source/DatasmithImporter/Private/DatasmithImporter.cpp
+++ b/Engine/Plugins/Enterprise/DatasmithImporter/Source/DatasmithImporter/Private/DatasmithImporter.cpp
@@ -547,6 +547,9 @@ void FDatasmithImporter::ImportClothes(FDatasmithImportContext& ImportContext)
 			ClothAsset->SetReferenceSkeleton(DefaultSkeleton->GetReferenceSkeleton(), false);
 		}
 
+		// Bind the mesh to the root bone.
+		ClothAsset->BindSimMeshToRootBone(); 
+
 		// Set the render mesh to duplicate the sim mesh
 		ClothAsset->CopySimMeshToRenderMesh(MaterialId);
 
diff --git a/Engine/Plugins/Experimental/ChaosClothAsset/Source/ChaosClothAsset/Private/ChaosClothAsset/ClothGeometryTools.cpp b/Engine/Plugins/Experimental/ChaosClothAsset/Source/ChaosClothAsset/Private/ChaosClothAsset/ClothGeometryTools.cpp
--- a/Engine/Plugins/Experimental/ChaosClothAsset/Source/ChaosClothAsset/Private/ChaosClothAsset/ClothGeometryTools.cpp
+++ b/Engine/Plugins/Experimental/ChaosClothAsset/Source/ChaosClothAsset/Private/ChaosClothAsset/ClothGeometryTools.cpp
@@ -115,6 +115,12 @@ namespace UE::Chaos::ClothAsset
 				const TArrayView<FVector3f> RenderNormal = ClothPatternAdapter.GetRenderNormal();
 				const TArrayView<TArray<FVector2f>> RenderUVs = ClothPatternAdapter.GetRenderUVs();
 				const TArrayView<FLinearColor> RenderColor = ClothPatternAdapter.GetRenderColor();
+				const TArrayView<int32> RenderNumBoneInfluences = ClothPatternAdapter.GetRenderNumBoneInfluences();
+				const TArrayView<TArray<int32>> RenderBoneIndices = ClothPatternAdapter.GetRenderBoneIndices();
+				const TArrayView<TArray<float>> RenderBoneWeights = ClothPatternAdapter.GetRenderBoneWeights();
+				const TArrayView<int32> SimNumBoneInfluences = ClothPatternAdapter.GetSimNumBoneInfluences();
+				const TArrayView<TArray<int32>> SimBoneIndices = ClothPatternAdapter.GetSimBoneIndices();
+				const TArrayView<TArray<float>> SimBoneWeights = ClothPatternAdapter.GetSimBoneWeights();
 
 				for (int32 VertexIndex = 0; VertexIndex < NumVertices; ++VertexIndex)
 				{
@@ -124,6 +130,9 @@ namespace UE::Chaos::ClothAsset
 					RenderColor[VertexIndex] = FLinearColor::White;
 					RenderTangentU[VertexIndex].Normalize();
 					RenderTangentV[VertexIndex].Normalize();
+					RenderNumBoneInfluences[VertexIndex] = SimNumBoneInfluences[VertexIndex];
+					RenderBoneIndices[VertexIndex] = SimBoneIndices[VertexIndex];
+					RenderBoneWeights[VertexIndex] = SimBoneWeights[VertexIndex];
 				}
 			}
 		}
@@ -190,4 +199,76 @@ namespace UE::Chaos::ClothAsset
 			}
 		}
 	}
+
+	void FClothGeometryTools::BindMeshToRootBone(const TSharedPtr<FClothCollection>& ClothCollection,
+												 bool bBindSimMesh,
+												 bool bBindRenderMesh,
+												 const TArray<int32> Lods)
+	{
+		if (!bBindSimMesh && !bBindRenderMesh)
+		{
+			return;
+		}
+		
+		FClothAdapter ClothAdapter(ClothCollection);
+
+		TArray<int32> LodsToBind;
+		if (Lods.IsEmpty())
+		{
+			LodsToBind.Reserve(ClothAdapter.GetNumLods());
+			
+			for (int32 LodIndex = 0; LodIndex < ClothAdapter.GetNumLods(); ++LodIndex)
+			{
+				LodsToBind.Add(LodIndex);
+			}
+		}
+		else
+		{	
+			LodsToBind.Reserve(Lods.Num());
+			
+			for (const int32 Lod : Lods)
+			{	
+				// Make sure the Lod indices are valid.
+				if (Lod < ClothAdapter.GetNumLods())
+				{
+					LodsToBind.Add(Lod);
+				}
+			}
+		}
+
+		for (const int32 LodIndex : LodsToBind)
+		{
+			FClothLodAdapter ClothLodAdapter = ClothAdapter.GetLod(LodIndex);
+			
+			if (bBindSimMesh)
+			{	
+				const int32 NumVertices = ClothLodAdapter.GetPatternsNumSimVertices();
+				TArrayView<int32> NumBoneInfluences = ClothLodAdapter.GetPatternsSimNumBoneInfluences();
+				TArrayView<TArray<int32>> BoneIndices = ClothLodAdapter.GetPatternsSimBoneIndices();
+				TArrayView<TArray<float>> BoneWeights = ClothLodAdapter.GetPatternsSimBoneWeights();
+
+				for (int32 VertexIndex = 0; VertexIndex < NumVertices; ++VertexIndex)
+				{
+					NumBoneInfluences[VertexIndex] = 1;
+					BoneIndices[VertexIndex] = {0};
+					BoneWeights[VertexIndex] = {1.0f};
+				}
+			}
+
+			if (bBindRenderMesh)
+			{
+				const int32 NumVertices = ClothLodAdapter.GetPatternsNumRenderVertices();
+				TArrayView<int32> NumBoneInfluences = ClothLodAdapter.GetPatternsRenderNumBoneInfluences();
+				TArrayView<TArray<int32>> BoneIndices = ClothLodAdapter.GetPatternsRenderBoneIndices();
+				TArrayView<TArray<float>> BoneWeights = ClothLodAdapter.GetPatternsRenderBoneWeights();
+
+				for (int32 VertexIndex = 0; VertexIndex < NumVertices; ++VertexIndex)
+				{
+					NumBoneInfluences[VertexIndex] = 1;
+					BoneIndices[VertexIndex] = {0};
+					BoneWeights[VertexIndex] = {1.0f};
+				}
+			}
+		}
+	}
 }  // End namespace UE::Chaos::ClothAsset
diff --git a/Engine/Plugins/Experimental/ChaosClothAsset/Source/ChaosClothAsset/Private/ChaosClothAsset/ClothLodAdapter.cpp b/Engine/Plugins/Experimental/ChaosClothAsset/Source/ChaosClothAsset/Private/ChaosClothAsset/ClothLodAdapter.cpp
--- a/Engine/Plugins/Experimental/ChaosClothAsset/Source/ChaosClothAsset/Private/ChaosClothAsset/ClothLodAdapter.cpp
+++ b/Engine/Plugins/Experimental/ChaosClothAsset/Source/ChaosClothAsset/Private/ChaosClothAsset/ClothLodAdapter.cpp
@@ -411,14 +411,14 @@ namespace UE::Chaos::ClothAsset
 		return FClothPatternConstAdapter(ClothCollection, LodIndex, PatternIndex);
 	}
 
-	void FClothLodConstAdapter::BuildSimulationMesh(TArray<FVector3f>& Positions, TArray<FVector3f>& Normals, TArray<uint32>& Indices) const
+	void FClothLodConstAdapter::BuildSimulationMesh(TArray<FVector3f>& Positions, TArray<FVector3f>& Normals, TArray<uint32>& Indices, TArray<int32>& WeldingMap) const
 	{
 		const int32 NumSimVertices = GetPatternsNumSimVertices();
 
 		// Initialize welding map with same index
 		// The welding map redirects to an existing vertex index if these two are part of the same welding group.
 		// The redirected index must be the smallest index in the group.
-		TArray<int32> WeldingMap;
+
 		WeldingMap.SetNumUninitialized(NumSimVertices);
 		for (int32 SimVertexIndex = 0; SimVertexIndex < NumSimVertices; ++SimVertexIndex)
 		{
diff --git a/Engine/Plugins/Experimental/ChaosClothAsset/Source/ChaosClothAsset/Public/ChaosClothAsset/ClothGeometryTools.h b/Engine/Plugins/Experimental/ChaosClothAsset/Source/ChaosClothAsset/Public/ChaosClothAsset/ClothGeometryTools.h
--- a/Engine/Plugins/Experimental/ChaosClothAsset/Source/ChaosClothAsset/Public/ChaosClothAsset/ClothGeometryTools.h
+++ b/Engine/Plugins/Experimental/ChaosClothAsset/Source/ChaosClothAsset/Public/ChaosClothAsset/ClothGeometryTools.h
@@ -23,5 +23,15 @@ namespace UE::Chaos::ClothAsset
 			bool bReverseSimMeshNormals,
 			bool bReverseRenderMeshNormals,
 			const TArray<int32>& PatternSelection);
+
+		/** 
+		 * Set the skinning weights for all of the sim/render vertices in ClothCollection to be bound to the root node. 
+		 * 
+		 * @param LODs if empty will apply the change to all LODs. Otherwise only LODs specified in the array (if exist) are affected. 
+		 */
+		static void BindMeshToRootBone(const TSharedPtr<FClothCollection>& ClothCollection,
+									   bool bBindSimMesh,
+									   bool bBindRenderMesh,
+									   const TArray<int32> Lods = {});
 	};
 }  // End namespace UE::Chaos::ClothAsset
\ No newline at end of file
diff --git a/Engine/Plugins/Experimental/ChaosClothAsset/Source/ChaosClothAsset/Public/ChaosClothAsset/ClothLodAdapter.h b/Engine/Plugins/Experimental/ChaosClothAsset/Source/ChaosClothAsset/Public/ChaosClothAsset/ClothLodAdapter.h
--- a/Engine/Plugins/Experimental/ChaosClothAsset/Source/ChaosClothAsset/Public/ChaosClothAsset/ClothLodAdapter.h
+++ b/Engine/Plugins/Experimental/ChaosClothAsset/Source/ChaosClothAsset/Public/ChaosClothAsset/ClothLodAdapter.h
@@ -95,7 +95,7 @@ namespace UE::Chaos::ClothAsset
 		const TSharedPtr<const FClothCollection>& GetClothCollection() const { return ClothCollection; }
 
 		/** Return the welded simulation mesh for this LOD. */
-		void BuildSimulationMesh(TArray<FVector3f>& Positions, TArray<FVector3f>& Normals, TArray<uint32>& Indices) const;
+		void BuildSimulationMesh(TArray<FVector3f>& Positions, TArray<FVector3f>& Normals, TArray<uint32>& Indices, TArray<int32>& WeldingMap) const;
 
 	private:
 		TSharedPtr<const FClothCollection> ClothCollection;
diff --git a/Engine/Plugins/Experimental/ChaosClothAsset/Source/ChaosClothAssetEngine/Private/ChaosClothAsset/ClothAsset.cpp b/Engine/Plugins/Experimental/ChaosClothAsset/Source/ChaosClothAssetEngine/Private/ChaosClothAsset/ClothAsset.cpp
--- a/Engine/Plugins/Experimental/ChaosClothAsset/Source/ChaosClothAssetEngine/Private/ChaosClothAsset/ClothAsset.cpp
+++ b/Engine/Plugins/Experimental/ChaosClothAsset/Source/ChaosClothAssetEngine/Private/ChaosClothAsset/ClothAsset.cpp
@@ -600,3 +600,11 @@ void UChaosClothAsset::ReregisterComponents()
 		}
 	}
 }
+
+void UChaosClothAsset::BindSimMeshToRootBone()
+{
+	using namespace UE::Chaos::ClothAsset;
+
+	check(ClothCollection.IsValid());
+	FClothGeometryTools::BindMeshToRootBone(ClothCollection, true, false);
+}
\ No newline at end of file
diff --git a/Engine/Plugins/Experimental/ChaosClothAsset/Source/ChaosClothAssetEngine/Private/ChaosClothAsset/ClothSimulationContext.cpp b/Engine/Plugins/Experimental/ChaosClothAsset/Source/ChaosClothAssetEngine/Private/ChaosClothAsset/ClothSimulationContext.cpp
--- a/Engine/Plugins/Experimental/ChaosClothAsset/Source/ChaosClothAssetEngine/Private/ChaosClothAsset/ClothSimulationContext.cpp
+++ b/Engine/Plugins/Experimental/ChaosClothAsset/Source/ChaosClothAssetEngine/Private/ChaosClothAsset/ClothSimulationContext.cpp
@@ -92,21 +92,21 @@ namespace UE::Chaos::ClothAsset
 		}
 
 		// Update bone matrices
-		if (!bIsInitialization)
-		{
-			ClothComponent.GetCurrentRefToLocalMatrices(RefToLocalMatrices, LodIndex);
-		}
-		else
+		if (bIsInitialization)
 		{
 			RefToLocalMatrices.Reset(NumBones);
 			RefToLocalMatrices.AddUninitialized(NumBones);
-
-			for (int32 BoneIndex = 0; BoneIndex < NumBones; ++BoneIndex)
-			{
-				RefToLocalMatrices[BoneIndex] = FMatrix44f::Identity;
-			}
 		}
-
+		
+		// We must set the matrices to identity on every call to Fill, otherwise every call to GetCurrentRefToLocalMatrices 
+		// below will multiply the content of the RefToLocalMatrices by another matrix and save the result back to
+		// RefToLocalMatrices which leads to the translation component increasing
+		for (int32 BoneIndex = 0; BoneIndex < NumBones; ++BoneIndex)
+		{
+			RefToLocalMatrices[BoneIndex] = FMatrix44f::Identity;
+		}
+		ClothComponent.GetCurrentRefToLocalMatrices(RefToLocalMatrices, LodIndex);
+		
 		// Update gravity
 		const UWorld* const World = ClothComponent.GetWorld();
 		constexpr float EarthGravity = -981.f;
diff --git a/Engine/Plugins/Experimental/ChaosClothAsset/Source/ChaosClothAssetEngine/Private/ChaosClothAsset/ClothSimulationModel.cpp b/Engine/Plugins/Experimental/ChaosClothAsset/Source/ChaosClothAssetEngine/Private/ChaosClothAsset/ClothSimulationModel.cpp
--- a/Engine/Plugins/Experimental/ChaosClothAsset/Source/ChaosClothAssetEngine/Private/ChaosClothAsset/ClothSimulationModel.cpp
+++ b/Engine/Plugins/Experimental/ChaosClothAsset/Source/ChaosClothAssetEngine/Private/ChaosClothAsset/ClothSimulationModel.cpp
@@ -120,7 +120,8 @@ FChaosClothSimulationModel::FChaosClothSimulationModel(const TSharedPtr<const UE
 	{
 		FChaosClothSimulationLodModel& LodModel = ClothSimulationLodModels[LodIndex];
 		const FClothLodConstAdapter ClothLod = Cloth.GetLod(LodIndex);
-		ClothLod.BuildSimulationMesh(LodModel.Positions, LodModel.Normals, LodModel.Indices);
+		TArray<int32> WeldingMap;
+		ClothLod.BuildSimulationMesh(LodModel.Positions, LodModel.Normals, LodModel.Indices, WeldingMap);
 
 		const TManagedArray<float>* const MaxDistanceValues = ClothCollection->FindAttributeTyped<float>("MaxDistance", FClothCollection::SimVerticesGroup);
 		if (MaxDistanceValues)
@@ -132,21 +133,33 @@ FChaosClothSimulationModel::FChaosClothSimulationModel(const TSharedPtr<const UE
 			LodModel.MaxDistance.Init(1.0, LodModel.Positions.Num());
 		}
 
-		const int32 NumSimVertices = LodModel.Positions.Num();
-		LodModel.BoneData.SetNum(NumSimVertices);
-		for (int32 VertexIndex = 0; VertexIndex < NumSimVertices; ++VertexIndex)
+		TConstArrayView<int32> NumBoneInfluences = ClothLod.GetPatternsSimNumBoneInfluences();
+		TConstArrayView<TArray<int32>> SimBoneIndices = ClothLod.GetPatternsSimBoneIndices();
+		TConstArrayView<TArray<float>> SimBoneWeights = ClothLod.GetPatternsSimBoneWeights();
+		LodModel.BoneData.SetNum(LodModel.Positions.Num());
+
+		uint32 WeldedIndex = 0;
+		for (int32 VertexIndex = 0; VertexIndex <  ClothLod.GetPatternsNumSimVertices(); ++VertexIndex)
 		{
-			// TODO: Skinning of the cloth asset needs to be done in order to be able to author theses values
-			LodModel.BoneData[VertexIndex].NumInfluences = 1;
-			LodModel.BoneData[VertexIndex].BoneIndices[0] = 0;
-			LodModel.BoneData[VertexIndex].BoneWeights[0] = 1.f;
+			if (WeldingMap[VertexIndex] == VertexIndex)
+			{
+				LodModel.BoneData[WeldedIndex].NumInfluences = NumBoneInfluences[VertexIndex];
+				for (int BoneIndex = 0; BoneIndex < LodModel.BoneData[WeldedIndex].NumInfluences; ++BoneIndex)
+				{
+					LodModel.BoneData[WeldedIndex].BoneIndices[BoneIndex] = SimBoneIndices[VertexIndex][BoneIndex];
+					LodModel.BoneData[WeldedIndex].BoneWeights[BoneIndex] = SimBoneWeights[VertexIndex][BoneIndex];
+				}
+				++WeldedIndex;
+			}
 		}
 	}
 
-	// Populate used bone names and indices  
-	// TODO: Skinning of the cloth asset needs to be done in order to be able to author theses values
-	UsedBoneNames.Add(TEXT("Root"));
-	UsedBoneIndices.Add(0);
+	// Populate used bone names and indices
+	for (int32 Index = 0; Index < ReferenceSkeleton.GetRawBoneNum(); ++Index)
+	{	
+		UsedBoneNames.Add(ReferenceSkeleton.GetRawRefBoneInfo()[Index].Name);
+		UsedBoneIndices.Add(Index);
+	}
 
 	// Initialize Reference bone index
 	ReferenceBoneIndex = CalculateReferenceBoneIndex(ClothSimulationLodModels, ReferenceSkeleton, UsedBoneIndices);
diff --git a/Engine/Plugins/Experimental/ChaosClothAsset/Source/ChaosClothAssetEngine/Public/ChaosClothAsset/ClothAsset.h b/Engine/Plugins/Experimental/ChaosClothAsset/Source/ChaosClothAssetEngine/Public/ChaosClothAsset/ClothAsset.h
--- a/Engine/Plugins/Experimental/ChaosClothAsset/Source/ChaosClothAssetEngine/Public/ChaosClothAsset/ClothAsset.h
+++ b/Engine/Plugins/Experimental/ChaosClothAsset/Source/ChaosClothAssetEngine/Public/ChaosClothAsset/ClothAsset.h
@@ -122,6 +122,9 @@ public:
 	/** Set the bone hierachy to use for this cloth. */
 	void SetReferenceSkeleton(const FReferenceSkeleton& InReferenceSkeleton, bool bRebuildClothSimulationModel = true) { RefSkeleton = InReferenceSkeleton; UpdateSkeleton(bRebuildClothSimulationModel); }
 
+	/** Set the skinning weights for all of the sim vertices to be bound to the root node of the reference skeleton. */
+	void BindSimMeshToRootBone();
+
 	//
 	// Dataflow
 	//
diff --git a/Engine/Plugins/Experimental/ChaosClothAssetEditor/Source/ChaosClothAssetEditorTools/Private/ChaosClothAsset/ClothTransferSkinWeightsTool.cpp b/Engine/Plugins/Experimental/ChaosClothAssetEditor/Source/ChaosClothAssetEditorTools/Private/ChaosClothAsset/ClothTransferSkinWeightsTool.cpp
--- a/Engine/Plugins/Experimental/ChaosClothAssetEditor/Source/ChaosClothAssetEditorTools/Private/ChaosClothAsset/ClothTransferSkinWeightsTool.cpp
+++ b/Engine/Plugins/Experimental/ChaosClothAssetEditor/Source/ChaosClothAssetEditorTools/Private/ChaosClothAsset/ClothTransferSkinWeightsTool.cpp
@@ -419,10 +419,14 @@ void UClothTransferSkinWeightsTool::TransferWeights()
 		UE::Chaos::ClothAsset::FClothLodAdapter ClothLodAdapter = ClothAdapter.GetLod(TargetLODIdx);
 
 		// Cloth collection data arrays we are writing to
-		TArrayView<int32> NumBoneInfluences = ClothLodAdapter.GetPatternsSimNumBoneInfluences();
+		TArrayView<int32> SimNumBoneInfluences = ClothLodAdapter.GetPatternsSimNumBoneInfluences();
 		TArrayView<TArray<int32>> SimBoneIndices = ClothLodAdapter.GetPatternsSimBoneIndices();
 		TArrayView<TArray<float>> SimBoneWeights = ClothLodAdapter.GetPatternsSimBoneWeights();
 
+		TArrayView<int32> RenderNumBoneInfluences = ClothLodAdapter.GetPatternsRenderNumBoneInfluences();
+		TArrayView<TArray<int32>> RenderBoneIndices = ClothLodAdapter.GetPatternsRenderBoneIndices();
+		TArrayView<TArray<float>> RenderBoneWeights = ClothLodAdapter.GetPatternsRenderBoneWeights();
+
 		const TArrayView<FVector3f> SimPositions =  ClothLodAdapter.GetPatternsSimRestPosition();
 		
 		checkSlow(SimPositions.Num() == SimBoneIndices.Num());
@@ -441,14 +445,21 @@ void UClothTransferSkinWeightsTool::TransferWeights()
 			
 			const int32 NumBones = BoneWeights.Num();
 			
-			NumBoneInfluences[VertexID] = NumBones;
+			SimNumBoneInfluences[VertexID] = NumBones;
 			SimBoneIndices[VertexID].SetNum(NumBones);
 			SimBoneWeights[VertexID].SetNum(NumBones);
 
+			RenderNumBoneInfluences[VertexID] = NumBones;
+			RenderBoneIndices[VertexID].SetNum(NumBones);
+			RenderBoneWeights[VertexID].SetNum(NumBones);
+
 			for (int BoneIdx = 0; BoneIdx < NumBones; ++BoneIdx) 
 			{
 				SimBoneIndices[VertexID][BoneIdx] = BoneWeights[BoneIdx].GetBoneIndex();
 				SimBoneWeights[VertexID][BoneIdx] = BoneWeights[BoneIdx].GetWeight();
+
+				RenderBoneIndices[VertexID][BoneIdx] = BoneWeights[BoneIdx].GetBoneIndex();
+				RenderBoneWeights[VertexID][BoneIdx] = BoneWeights[BoneIdx].GetWeight();
 			}
 
 		}, bUseParallel ? EParallelForFlags::None : EParallelForFlags::ForceSingleThread);
diff --git a/Engine/Plugins/Experimental/ChaosClothAssetEditor/Source/ChaosClothAssetTools/Private/ChaosClothAsset/ClothAssetBuilderEditor.cpp b/Engine/Plugins/Experimental/ChaosClothAssetEditor/Source/ChaosClothAssetTools/Private/ChaosClothAsset/ClothAssetBuilderEditor.cpp
--- a/Engine/Plugins/Experimental/ChaosClothAssetEditor/Source/ChaosClothAssetTools/Private/ChaosClothAsset/ClothAssetBuilderEditor.cpp
+++ b/Engine/Plugins/Experimental/ChaosClothAssetEditor/Source/ChaosClothAssetTools/Private/ChaosClothAsset/ClothAssetBuilderEditor.cpp
@@ -138,8 +138,11 @@ void UClothAssetBuilderEditor::BuildLod(FSkeletalMeshLODModel& LODModel, const U
 		Section.BaseVertexIndex = LODModel.NumVertices;
 		LODModel.NumVertices += (uint32)NumVertices;
 
-		TSet<FBoneIndexType> SectionBoneMap;
-		SectionBoneMap.Reserve(ClothAsset.RefSkeleton.GetNum());
+		// Map reference skeleton bone index to the index in the section's bone map
+		TMap<FBoneIndexType, FBoneIndexType> ReferenceToSectionBoneMap; 
+
+		// Track how many bones we added to the section's bone map so far
+		int CurSectionBoneMapNum = 0; 
 
 		for (int32 VertexIndex = 0; VertexIndex < NumVertices; ++VertexIndex)
 		{
@@ -164,19 +167,53 @@ void UClothAssetBuilderEditor::BuildLod(FSkeletalMeshLODModel& LODModel, const U
 				SoftVertex.UVs[TexCoord] = RenderUVs[TexCoord];
 			}
 
-			for (int32 Influence = 0; Influence < MAX_TOTAL_INFLUENCES; ++Influence)
+			const int32 NumInfluences = ClothCollection->RenderNumBoneInfluences[RenderIndex];
+			
+			// Add all of the bones that have non-zero influence to the section's bone map and keep track of the order
+			// that we added the reference bone via CurSectionBoneMapNum
+			for (int32 Influence = 0; Influence < NumInfluences; ++Influence)
 			{
-				const FBoneIndexType InfluenceBone = (FBoneIndexType)0;  // TODO: Set the correct bone influence from the cloth collection data
-				const int16 InfluenceWeight = (Influence == 0) ? UE::AnimationCore::MaxRawBoneWeight : 0;
+				const FBoneIndexType InfluenceBone = (FBoneIndexType)ClothCollection->RenderBoneIndices[RenderIndex][Influence];
 
-				SoftVertex.InfluenceBones[Influence] = InfluenceBone;
-				SoftVertex.InfluenceWeights[Influence] = InfluenceWeight;
+				if (ReferenceToSectionBoneMap.Contains(InfluenceBone) == false)
+				{
+					ReferenceToSectionBoneMap.Add(InfluenceBone, CurSectionBoneMapNum);
+					++CurSectionBoneMapNum; 
+				}
+			}
 
-				if (InfluenceWeight)
+			int32 Influence = 0;
+			for (;Influence < NumInfluences; ++Influence)
+			{
+				const FBoneIndexType InfluenceBone = (FBoneIndexType)ClothCollection->RenderBoneIndices[RenderIndex][Influence];
+				const float InWeight = ClothCollection->RenderBoneWeights[RenderIndex][Influence];
+				const uint16 InfluenceWeight = static_cast<uint16>(InWeight * static_cast<float>(UE::AnimationCore::MaxRawBoneWeight) + 0.5f);
+				
+				// FSoftSkinVertex::InfluenceBones contain indices into the section's bone map and not the reference
+				// skeleton, so we need to remap
+				const FBoneIndexType* const MappedIndexPtr = ReferenceToSectionBoneMap.Find(InfluenceBone);
+				
+				// ReferenceToSectionBoneMap should always contain InfluenceBone since it was added above
+				checkSlow(MappedIndexPtr);
+				if (MappedIndexPtr != nullptr)
 				{
-					SectionBoneMap.Add(InfluenceBone);
+					SoftVertex.InfluenceBones[Influence] = *MappedIndexPtr;
+					SoftVertex.InfluenceWeights[Influence] = InfluenceWeight;
 				}
 			}
+			
+			for (;Influence < MAX_TOTAL_INFLUENCES; ++Influence)
+			{
+				SoftVertex.InfluenceBones[Influence] = 0;
+				SoftVertex.InfluenceWeights[Influence] = 0;
+			}
+		}
+
+		// Initialize the section bone map
+		Section.BoneMap.SetNumUninitialized(ReferenceToSectionBoneMap.Num());
+		for (const TPair<FBoneIndexType, FBoneIndexType>& Pair : ReferenceToSectionBoneMap)
+		{
+			Section.BoneMap[Pair.Value] = Pair.Key;
 		}
 
 		// Remap the LOD indices with the new vertex indices
@@ -185,9 +222,7 @@ void UClothAssetBuilderEditor::BuildLod(FSkeletalMeshLODModel& LODModel, const U
 			RenderIndex = LodRenderIndexRemap[RenderIndex];
 		}
 
-		// Update the section bone map
-		Section.BoneMap = SectionBoneMap.Array();
-		ActiveBoneIndices.Append(MoveTemp(SectionBoneMap));
+		ActiveBoneIndices.Append(Section.BoneMap);
 
 		// Update max bone influences
 		Section.CalcMaxBoneInfluences();
diff --git a/Engine/Plugins/Experimental/ChaosClothAssetEditor/Source/ChaosClothAssetTools/Private/ChaosClothAsset/ClothingAssetToClothAssetExporter.cpp b/Engine/Plugins/Experimental/ChaosClothAssetEditor/Source/ChaosClothAssetTools/Private/ChaosClothAsset/ClothingAssetToClothAssetExporter.cpp
--- a/Engine/Plugins/Experimental/ChaosClothAssetEditor/Source/ChaosClothAssetTools/Private/ChaosClothAsset/ClothingAssetToClothAssetExporter.cpp
+++ b/Engine/Plugins/Experimental/ChaosClothAssetEditor/Source/ChaosClothAssetTools/Private/ChaosClothAsset/ClothingAssetToClothAssetExporter.cpp
@@ -60,6 +60,9 @@ void UClothingAssetToChaosClothAssetExporter::Export(const UClothingAssetBase* C
 		SkeletalMesh->GetRefSkeleton() :
 		LoadObject<USkeleton>(nullptr, TEXT("/Engine/EditorMeshes/SkeletalMesh/DefaultSkeletalMesh_Skeleton.DefaultSkeletalMesh_Skeleton"), nullptr, LOAD_None, nullptr)->GetReferenceSkeleton());
 
+	// Bind the mesh to the root bone.
+	ClothAsset->BindSimMeshToRootBone(); 
+
 	// Set the render mesh to duplicate the sim mesh
 	ClothAsset->CopySimMeshToRenderMesh(MaterialId);
 
