Added BindToRootBone DataFlow node to Chaos Cloth Editor.

diff --git a/Engine/Plugins/Experimental/ChaosClothAssetEditor/Source/ChaosClothAssetEditor/Private/ChaosClothAsset/DataflowNodes/BindToRootBoneNode.cpp b/Engine/Plugins/Experimental/ChaosClothAssetEditor/Source/ChaosClothAssetEditor/Private/ChaosClothAsset/DataflowNodes/BindToRootBoneNode.cpp
new file mode 100644
--- /dev/null
+++ b/Engine/Plugins/Experimental/ChaosClothAssetEditor/Source/ChaosClothAssetEditor/Private/ChaosClothAsset/DataflowNodes/BindToRootBoneNode.cpp
@@ -0,0 +1,38 @@
+// Copyright Epic Games, Inc. All Rights Reserved.
+
+#include "ChaosClothAsset/DataflowNodes/BindToRootBoneNode.h"
+#include "ChaosClothAsset/DataflowNodes/DataflowNodes.h"
+#include "ChaosClothAsset/ClothGeometryTools.h"
+
+#include UE_INLINE_GENERATED_CPP_BY_NAME(BindToRootBoneNode)
+
+#define LOCTEXT_NAMESPACE "ChaosClothAssetBindToRootBoneNode"
+
+FChaosClothAssetBindToRootBoneNode::FChaosClothAssetBindToRootBoneNode(const Dataflow::FNodeParameters& InParam, FGuid InGuid)
+	: FDataflowNode(InParam, InGuid)
+{
+	RegisterInputConnection(&Collection);
+	RegisterOutputConnection(&Collection, &Collection);
+}
+
+void FChaosClothAssetBindToRootBoneNode::Evaluate(Dataflow::FContext& Context, const FDataflowOutput* Out) const
+{
+	if (Out->IsA<FManagedArrayCollection>(&Collection))
+	{
+		using namespace UE::Chaos::ClothAsset;
+
+		// Evaluate in collection
+		const FManagedArrayCollection& InCollection = GetValue<FManagedArrayCollection>(Context, &Collection);
+
+		// TODO: Needs to make cloth collection a facade to avoid the copy of the cloth collection to a managed array
+		//       and remove the above const reference to operate on the InCollection instead.
+		TSharedPtr<FClothCollection> ClothCollection = MakeShared<FClothCollection>();
+		InCollection.CopyTo(ClothCollection.Get());
+
+        FClothGeometryTools::BindMeshToRootBone(ClothCollection, bBindSimMesh, bBindRenderMesh);
+
+		SetValue<FManagedArrayCollection>(Context, *ClothCollection, &Collection);
+	}
+}
+
+#undef LOCTEXT_NAMESPACE
diff --git a/Engine/Plugins/Experimental/ChaosClothAssetEditor/Source/ChaosClothAssetEditor/Private/ChaosClothAsset/DataflowNodes/BindToRootBoneNode.h b/Engine/Plugins/Experimental/ChaosClothAssetEditor/Source/ChaosClothAssetEditor/Private/ChaosClothAsset/DataflowNodes/BindToRootBoneNode.h
new file mode 100644
--- /dev/null
+++ b/Engine/Plugins/Experimental/ChaosClothAssetEditor/Source/ChaosClothAssetEditor/Private/ChaosClothAsset/DataflowNodes/BindToRootBoneNode.h
@@ -0,0 +1,32 @@
+// Copyright Epic Games, Inc. All Rights Reserved.
+
+#pragma once 
+
+#include "Dataflow/DataflowNode.h"
+#include "GeometryCollection/ManagedArrayCollection.h"
+#include "BindToRootBoneNode.generated.h"
+
+USTRUCT(Meta = (DataflowCloth))
+struct FChaosClothAssetBindToRootBoneNode : public FDataflowNode
+{
+	GENERATED_USTRUCT_BODY()
+	DATAFLOW_NODE_DEFINE_INTERNAL(FChaosClothAssetBindToRootBoneNode, "BindToRootBone", "Cloth", "Cloth Bind Skinning Weights To Root Bone")
+	//DATAFLOW_NODE_RENDER_TYPE(FGeometryCollection::StaticType(), "Collection")  // TODO: Leave out the render type until there is something to render
+
+public:
+
+	UPROPERTY(Meta = (Dataflowinput, DataflowOutput, DisplayName = "Collection", DataflowPassthrough = "Collection"))
+	FManagedArrayCollection Collection;
+
+	/** Whether to bind the simulation mesh. */
+	UPROPERTY(EditAnywhere, Category = "Bind To Root Bone")
+	bool bBindSimMesh = true;
+
+	/** Whether to bind the render mesh. */
+	UPROPERTY(EditAnywhere, Category = "Bind To Root Bone")
+	bool bBindRenderMesh = true;
+
+	FChaosClothAssetBindToRootBoneNode(const Dataflow::FNodeParameters& InParam, FGuid InGuid = FGuid::NewGuid());
+
+	virtual void Evaluate(Dataflow::FContext& Context, const FDataflowOutput* Out) const override;
+};
diff --git a/Engine/Plugins/Experimental/ChaosClothAssetEditor/Source/ChaosClothAssetEditor/Private/ChaosClothAsset/DataflowNodes/DataflowNodes.cpp b/Engine/Plugins/Experimental/ChaosClothAssetEditor/Source/ChaosClothAssetEditor/Private/ChaosClothAsset/DataflowNodes/DataflowNodes.cpp
--- a/Engine/Plugins/Experimental/ChaosClothAssetEditor/Source/ChaosClothAssetEditor/Private/ChaosClothAsset/DataflowNodes/DataflowNodes.cpp
+++ b/Engine/Plugins/Experimental/ChaosClothAssetEditor/Source/ChaosClothAssetEditor/Private/ChaosClothAsset/DataflowNodes/DataflowNodes.cpp
@@ -9,6 +9,7 @@
 #include "ChaosClothAsset/DataflowNodes/CopySimulationToRenderMeshNode.h"
 #include "ChaosClothAsset/DataflowNodes/ReverseNormalsNode.h"
 #include "ChaosClothAsset/DataflowNodes/SimulationDefaultConfigNode.h"
+#include "ChaosClothAsset/DataflowNodes/BindToRootBoneNode.h"
 #include "Dataflow/DataflowNodeFactory.h"
 #include "Framework/Notifications/NotificationManager.h"
 #include "Widgets/Notifications/SNotificationList.h"
@@ -27,6 +28,7 @@ namespace UE::Chaos::ClothAsset::DataflowNodes
 		DATAFLOW_NODE_REGISTER_CREATION_FACTORY(FChaosClothAssetCopySimulationToRenderMeshNode);
 		DATAFLOW_NODE_REGISTER_CREATION_FACTORY(FChaosClothAssetReverseNormalsNode);
 		DATAFLOW_NODE_REGISTER_CREATION_FACTORY(FChaosClothAssetSimulationDefaultConfigNode);
+		DATAFLOW_NODE_REGISTER_CREATION_FACTORY(FChaosClothAssetBindToRootBoneNode);
 	}
 
 	void LogAndToastWarning(const FText& Error)
