-TransferBoneWeights function now uses bone attributes to re-index the skin weight bone indices when transferring between different skeletons.
-Refactored Chaos cloth transfer tool to use the new TransferBoneWeights functionality.

diff --git a/Engine/Plugins/Experimental/ChaosClothAssetEditor/Source/ChaosClothAssetEditorTools/Private/ChaosClothAsset/ClothTransferSkinWeightsTool.cpp b/Engine/Plugins/Experimental/ChaosClothAssetEditor/Source/ChaosClothAssetEditorTools/Private/ChaosClothAsset/ClothTransferSkinWeightsTool.cpp
--- a/Engine/Plugins/Experimental/ChaosClothAssetEditor/Source/ChaosClothAssetEditorTools/Private/ChaosClothAsset/ClothTransferSkinWeightsTool.cpp
+++ b/Engine/Plugins/Experimental/ChaosClothAssetEditor/Source/ChaosClothAssetEditorTools/Private/ChaosClothAsset/ClothTransferSkinWeightsTool.cpp
@@ -19,6 +19,7 @@
 
 #include "DynamicMesh/DynamicMeshAttributeSet.h"
 #include "DynamicMesh/DynamicVertexSkinWeightsAttribute.h"
+#include "DynamicMesh/DynamicBoneAttribute.h"
 #include "DynamicMesh/DynamicMeshAABBTree3.h"
 #include "DynamicMesh/MeshTransforms.h"
 
@@ -43,33 +44,6 @@
 
 namespace ClothTransferSkinWeightsToolHelpers
 {
-	/** 
-	 * Compute mappings between indices and bone names.
-	 * 
-     * @note We assume that each mesh inherits its reference skeleton from the same USkeleton asset. However, their  
-     * internal indexing can be different and hence when transfering weights we need to make sure we reference the correct  
-     * bones via their names instead of indices. 
-     */
-	void GetBoneMaps(const USkinnedAsset* SourceSkinnedAsset, 
-						  const USkinnedAsset* TargetSkinnedAsset, 
-						  TMap<FBoneIndexType, FName>& SourceIndexToBone,
-						  TMap<FName, FBoneIndexType>& TargetBoneToIndex)
-	{
-		TargetBoneToIndex.Reset();
-		SourceIndexToBone.Reset();
-		const FReferenceSkeleton& SourceRefSkeleton = SourceSkinnedAsset->GetRefSkeleton();
-		for (int32 Index = 0; Index < SourceRefSkeleton.GetRawBoneNum(); ++Index)
-		{
-			SourceIndexToBone.Add(Index, SourceRefSkeleton.GetRawRefBoneInfo()[Index].Name);
-		}
-
-		const FReferenceSkeleton& TargetRefSkeleton = TargetSkinnedAsset->GetRefSkeleton();
-		for (int32 Index = 0; Index < TargetRefSkeleton.GetRawBoneNum(); ++Index)
-		{
-			TargetBoneToIndex.Add(TargetRefSkeleton.GetRawRefBoneInfo()[Index].Name, Index);
-		}
-	};
-
 	void SkeletalMeshToDynamicMesh(USkeletalMesh* FromSkeletalMeshAsset, int32 SourceLODIdx, FDynamicMesh3& ToDynamicMesh)
 	{
 		FMeshDescription SourceMesh;
@@ -294,8 +268,7 @@ void UClothTransferSkinWeightsTool::UpdatePreviewMesh()
 
 	// Convert source Skeletal Mesh to Dynamic Mesh
 	FDynamicMesh3 SourceDynamicMesh;
-	USkeletalMesh* FromSkeletalMeshAsset = ToolProperties->SourceMesh;
-	ClothTransferSkinWeightsToolHelpers::SkeletalMeshToDynamicMesh(FromSkeletalMeshAsset, SourceLODIdx, SourceDynamicMesh);
+	ClothTransferSkinWeightsToolHelpers::SkeletalMeshToDynamicMesh(ToolProperties->SourceMesh, SourceLODIdx, SourceDynamicMesh);
 	MeshTransforms::ApplyTransform(SourceDynamicMesh, ToolProperties->SourceMeshTransform, true);
 
 	// Convert target ClothComponent to Dynamic Mesh
@@ -304,16 +277,22 @@ void UClothTransferSkinWeightsTool::UpdatePreviewMesh()
 	TargetDynamicMesh.Attributes()->AttachSkinWeightsAttribute(FSkeletalMeshAttributes::DefaultSkinWeightProfileName, new FDynamicMeshVertexSkinWeightsAttribute(&TargetDynamicMesh));
 	ClothTransferSkinWeightsToolHelpers::ClothComponentToDynamicMesh(ClothComponent, TargetDynamicMesh);
 
+	// Enable bone attribute for the target mesh and copy over the bone data from the cloth asset
+	// TODO: Once we support skin weights and bones in ToDynamicMesh converter we can remove this logic
+	const UChaosClothAsset* TargetClothAsset = ClothComponent->GetClothAsset();
+	const FReferenceSkeleton& TargetRefSkeleton = TargetClothAsset->GetRefSkeleton();
+	TargetDynamicMesh.Attributes()->EnableBones(TargetRefSkeleton.GetRawBoneNum());
+	UE::Geometry::FDynamicMeshBoneNameAttribute* BoneNameAttrib = TargetDynamicMesh.Attributes()->GetBoneNames();
+	this->TargetMeshBoneNameToIndex.Reset();
+	for (int BoneID = 0; BoneID < TargetRefSkeleton.GetRawBoneNum(); ++BoneID)
+	{	
+		const FName& BoneName = TargetRefSkeleton.GetRawRefBoneInfo()[BoneID].Name;
+		BoneNameAttrib->SetValue(BoneID, BoneName);
+		this->TargetMeshBoneNameToIndex.Add(BoneName, BoneID);
+	}
 
+	// Do the transfer
 	FTransferBoneWeights TransferBoneWeights(&SourceDynamicMesh, FSkeletalMeshAttributes::DefaultSkinWeightProfileName);
-
-	// Compute bone index mappings
-	TMap<FBoneIndexType, FName> SourceIndexToBone;
-	ClothTransferSkinWeightsToolHelpers::GetBoneMaps(FromSkeletalMeshAsset, ClothComponent->GetClothAsset(), SourceIndexToBone, this->TargetMeshBoneNameToIndex);
-	TransferBoneWeights.SourceIndexToBone = &SourceIndexToBone;
-	TransferBoneWeights.TargetBoneToIndex = &this->TargetMeshBoneNameToIndex;
-
-	// Do the actual transfer
 	const FTransformSRT3d TargetToWorld = ClothComponent->GetComponentTransform();
 	if (TransferBoneWeights.Validate() == EOperationValidationResult::Ok)
 	{
@@ -322,7 +301,6 @@ void UClothTransferSkinWeightsTool::UpdatePreviewMesh()
 
 	// Get set of bone indices used in the target mesh
 	TMap<FName, FBoneIndexType> UsedBoneNames;
-	const FReferenceSkeleton& TargetRefSkeleton = ClothComponent->GetClothAsset()->GetRefSkeleton();
 
 	const TMap<FName, TUniquePtr<FDynamicMeshVertexSkinWeightsAttribute>>& WeightLayers = TargetDynamicMesh.Attributes()->GetSkinWeightsAttributes();
 	for (const TPair<FName, TUniquePtr<FDynamicMeshVertexSkinWeightsAttribute>>& LayerPair : WeightLayers)
@@ -388,22 +366,21 @@ void UClothTransferSkinWeightsTool::TransferWeights()
 
 	// Convert source Skeletal Mesh to Dynamic Mesh
 	FDynamicMesh3 SourceDynamicMesh;
-	USkeletalMesh* FromSkeletalMeshAsset = ToolProperties->SourceMesh;
-	ClothTransferSkinWeightsToolHelpers::SkeletalMeshToDynamicMesh(FromSkeletalMeshAsset, SourceLODIdx, SourceDynamicMesh);
-	FTransformSRT3d SourceToWorld; //TODO: Allows the user to set this value or infer it from an editor
-	MeshTransforms::ApplyTransform(SourceDynamicMesh, SourceToWorld, true);
+	ClothTransferSkinWeightsToolHelpers::SkeletalMeshToDynamicMesh(ToolProperties->SourceMesh, SourceLODIdx, SourceDynamicMesh);
+	MeshTransforms::ApplyTransform(SourceDynamicMesh, ToolProperties->SourceMeshTransform, true);
 
 	UChaosClothAsset* TargetClothAsset = ClothComponent->GetClothAsset(); 
 
 	// Compute bone index mappings
-	TMap<FBoneIndexType, FName> SourceIndexToBone;
 	TMap<FName, FBoneIndexType> TargetBoneToIndex;
-	ClothTransferSkinWeightsToolHelpers::GetBoneMaps(static_cast<USkinnedAsset*>(FromSkeletalMeshAsset), static_cast<USkinnedAsset*>(TargetClothAsset), SourceIndexToBone, TargetBoneToIndex);
-	
+	const FReferenceSkeleton& TargetRefSkeleton = TargetClothAsset->GetRefSkeleton();
+	for (int32 Index = 0; Index < TargetRefSkeleton.GetRawBoneNum(); ++Index)
+	{
+		TargetBoneToIndex.Add(TargetRefSkeleton.GetRawRefBoneInfo()[Index].Name, Index);
+	}
+
 	// Setup bone weight transfer operator
 	FTransferBoneWeights TransferBoneWeights(&SourceDynamicMesh, FSkeletalMeshAttributes::DefaultSkinWeightProfileName); 
-	TransferBoneWeights.SourceIndexToBone = &SourceIndexToBone;
-	TransferBoneWeights.TargetBoneToIndex = &TargetBoneToIndex;
 	if (TransferBoneWeights.Validate() != EOperationValidationResult::Ok)
 	{
 		//TODO: Display error message
@@ -411,7 +388,7 @@ void UClothTransferSkinWeightsTool::TransferWeights()
 	}
 
 	UE::Chaos::ClothAsset::FClothAdapter ClothAdapter(TargetClothAsset->GetClothCollection());
-    FTransformSRT3d TargetToWorld; //TODO: Allows the user to set this value or infer it from an editor
+    FTransformSRT3d TargetToWorld = ClothComponent->GetComponentTransform();
 
 	// Iterate over the LODs and transfer the bone weights from the source Skeletal mesh to the Cloth asset
 	for (int TargetLODIdx = 0; TargetLODIdx < ClothAdapter.GetNumLods(); ++TargetLODIdx) 
@@ -441,7 +418,7 @@ void UClothTransferSkinWeightsTool::TransferWeights()
 			const FVector3d PosD = FVector3d((double)Pos[0], (double)Pos[1], (double)Pos[2]);
 			
 			UE::AnimationCore::FBoneWeights BoneWeights;
-			TransferBoneWeights.Compute(PosD, TargetToWorld, BoneWeights);
+			TransferBoneWeights.Compute(PosD, TargetToWorld, BoneWeights, &TargetBoneToIndex);
 			
 			const int32 NumBones = BoneWeights.Num();
 			
diff --git a/Engine/Plugins/Runtime/GeometryProcessing/Source/DynamicMesh/Private/Operations/TransferBoneWeights.cpp b/Engine/Plugins/Runtime/GeometryProcessing/Source/DynamicMesh/Private/Operations/TransferBoneWeights.cpp
--- a/Engine/Plugins/Runtime/GeometryProcessing/Source/DynamicMesh/Private/Operations/TransferBoneWeights.cpp
+++ b/Engine/Plugins/Runtime/GeometryProcessing/Source/DynamicMesh/Private/Operations/TransferBoneWeights.cpp
@@ -3,6 +3,7 @@
 #include "Operations/TransferBoneWeights.h"
 #include "DynamicMesh/DynamicMesh3.h"
 #include "DynamicMesh/DynamicVertexSkinWeightsAttribute.h"
+#include "DynamicMesh/DynamicBoneAttribute.h"
 #include "DynamicMesh/DynamicMeshAttributeSet.h"
 #include "DynamicMesh/DynamicMeshAABBTree3.h"
 #include "Async/ParallelFor.h"
@@ -32,7 +33,7 @@ namespace TransferBoneWeightsLocals
 								const FIndex3i& TriVertices,
 								const FVector3f& Bary,
 								const FDynamicMeshVertexSkinWeightsAttribute* Attribute,
-								const TMap<FBoneIndexType, FName>* SourceIndexToBone = nullptr,
+								const TArray<FName>* SourceIndexToBone = nullptr,
 								const TMap<FName, FBoneIndexType>* TargetBoneToIndex = nullptr,
 								bool bNormalizeToOne = true)
 	{
@@ -61,7 +62,8 @@ namespace TransferBoneWeightsLocals
 				FBoneIndexType FromIdx = BoneWeight.GetBoneIndex();
 				uint16 FromWeight = BoneWeight.GetRawWeight();
 
-				if (ensure(SourceIndexToBone->Contains(FromIdx))) // the map must contain the index
+				checkSlow(FromIdx < SourceIndexToBone->Num());
+				if (FromIdx < SourceIndexToBone->Num())
 				{
 					FName BoneName = (*SourceIndexToBone)[FromIdx];
 					if (TargetBoneToIndex->Contains(BoneName))
@@ -72,7 +74,7 @@ namespace TransferBoneWeightsLocals
 					}
 					else 
 					{	
-						UE_LOG(LogGeometry, Warning, TEXT("FTransferBoneWeights: Bone name %s does not exist in the target mesh."), *BoneName.ToString());
+						UE_LOG(LogGeometry, Error, TEXT("FTransferBoneWeights: Bone name %s does not exist in the target mesh."), *BoneName.ToString());
 					}
 				}
 			}
@@ -91,11 +93,12 @@ namespace TransferBoneWeightsLocals
 
 			OutWeights = MappedWeights;
 		}
-	};
+	}
 
 	
 	FDynamicMeshVertexSkinWeightsAttribute* GetOrCreateSkinWeightsAttribute(FDynamicMesh3& InMesh, const FName& InProfileName)
 	{
+		checkSlow(InMesh.HasAttributes());
 		FDynamicMeshVertexSkinWeightsAttribute* Attribute = InMesh.Attributes()->GetSkinWeightsAttribute(InProfileName);
 		if (Attribute == nullptr)
 		{
@@ -103,7 +106,7 @@ namespace TransferBoneWeightsLocals
 			InMesh.Attributes()->AttachSkinWeightsAttribute(InProfileName, Attribute);
 		}
 		return Attribute;
-	};
+	}
 }
 
 FTransferBoneWeights::FTransferBoneWeights(const FDynamicMesh3* InSourceMesh, 
@@ -153,6 +156,11 @@ EOperationValidationResult FTransferBoneWeights::Validate()
 		return EOperationValidationResult::Failed_UnknownReason;
 	}
 
+	if (bIgnoreBoneAttributes == false && SourceMesh->Attributes()->HasBones() == false) 
+	{
+		return EOperationValidationResult::Failed_UnknownReason;
+	}
+
 	return EOperationValidationResult::Ok;
 }
 
@@ -160,12 +168,47 @@ bool FTransferBoneWeights::Compute(FDynamicMesh3& InOutTargetMesh, const FTransf
 {	
 	if (Validate() != EOperationValidationResult::Ok) 
 	{
-		checkNoEntry();
 		return false;
 	}
+
+	if (!InOutTargetMesh.HasAttributes())
+	{
+		InOutTargetMesh.EnableAttributes(); 
+	}
+
+	if (!bIgnoreBoneAttributes && !InOutTargetMesh.Attributes()->HasBones())
+	{
+		return false; // the target mesh must have bone attributes
+	}
 	
 	FDynamicMeshVertexSkinWeightsAttribute* TargetSkinWeights = TransferBoneWeightsLocals::GetOrCreateSkinWeightsAttribute(InOutTargetMesh, InTargetProfileName);
 	checkSlow(TargetSkinWeights);
+	
+	// Map the bone name to its index for the target mesh.
+	// Will be null if either the target and the source skeletons are the same or the caller forced the attributes to be ignored
+	TUniquePtr<TMap<FName, uint16>> TargetBoneToIndex;
+	if (!bIgnoreBoneAttributes)
+	{	
+		const TArray<FName>& SourceBoneNames = SourceMesh->Attributes()->GetBoneNames()->GetAttribValues();
+		const TArray<FName>& TargetBoneNames = InOutTargetMesh.Attributes()->GetBoneNames()->GetAttribValues();
+
+		if (SourceBoneNames != TargetBoneNames)
+		{
+			TargetBoneToIndex = MakeUnique<TMap<FName, uint16>>();
+			TargetBoneToIndex->Reserve(TargetBoneNames.Num());
+
+			for (int BoneID = 0; BoneID < TargetBoneNames.Num(); ++BoneID)
+			{
+				const FName& BoneName = TargetBoneNames[BoneID];
+				if (TargetBoneToIndex->Contains(BoneName))
+				{
+					checkSlow(false);
+					return false; // there should be no duplicates
+				}
+				TargetBoneToIndex->Add(BoneName, static_cast<uint16>(BoneID));
+			}
+		}
+	}
 
 	bool bFailed = false;
 	
@@ -181,7 +224,7 @@ bool FTransferBoneWeights::Compute(FDynamicMesh3& InOutTargetMesh, const FTransf
 			FVector3d Point = InOutTargetMesh.GetVertex(VertexID);
 		
 			FBoneWeights Weights;
-			if (Compute(Point, InToWorld, Weights) == false)
+			if (Compute(Point, InToWorld, Weights, TargetBoneToIndex.Get()) == false)
 			{
 				bFailed = true;
 				return;
@@ -200,11 +243,17 @@ bool FTransferBoneWeights::Compute(FDynamicMesh3& InOutTargetMesh, const FTransf
 	return true;
 }
 
-bool FTransferBoneWeights::Compute(const FVector3d& InPoint, const FTransformSRT3d& InToWorld, FBoneWeights& OutWeights) 
+bool FTransferBoneWeights::Compute(const FVector3d& InPoint, const FTransformSRT3d& InToWorld, FBoneWeights& OutWeights, const TMap<FName, uint16>* TargetBoneToIndex) 
 {
 	const FDynamicMeshVertexSkinWeightsAttribute* SourceSkinWeights = SourceMesh->Attributes()->GetSkinWeightsAttribute(SourceProfileName);
 	checkSlow(SourceSkinWeights);
 
+	const TArray<FName>* SourceBoneNames = nullptr;
+	if (!bIgnoreBoneAttributes) 
+	{
+		SourceBoneNames = &SourceMesh->Attributes()->GetBoneNames()->GetAttribValues();
+	}
+
 	IMeshSpatial::FQueryOptions Options;
 	double NearestDistSqr;
 	int32 NearTriID;
@@ -219,9 +268,8 @@ bool FTransferBoneWeights::Compute(const FVector3d& InPoint, const FTransformSRT
 		NearTriID = InternalSourceBVH->FindNearestTriangle(WorldPoint, NearestDistSqr, Options);
 	}
 
-	if (NearTriID == IndexConstants::InvalidID) 
+	if (!ensure(NearTriID != IndexConstants::InvalidID))
 	{
-		checkNoEntry();
 		return false;
 	}
 
@@ -238,7 +286,7 @@ bool FTransferBoneWeights::Compute(const FVector3d& InPoint, const FTransformSRT
 													  TriVertex,
 													  FVector3f((float)Bary[0], (float)Bary[1], (float)Bary[2]),
 													  SourceSkinWeights,
-													  SourceIndexToBone,
+													  SourceBoneNames,
 													  TargetBoneToIndex);
 
 	return true;
diff --git a/Engine/Plugins/Runtime/GeometryProcessing/Source/DynamicMesh/Public/Operations/TransferBoneWeights.h b/Engine/Plugins/Runtime/GeometryProcessing/Source/DynamicMesh/Public/Operations/TransferBoneWeights.h
--- a/Engine/Plugins/Runtime/GeometryProcessing/Source/DynamicMesh/Public/Operations/TransferBoneWeights.h
+++ b/Engine/Plugins/Runtime/GeometryProcessing/Source/DynamicMesh/Public/Operations/TransferBoneWeights.h
@@ -27,25 +27,33 @@ namespace UE
 namespace Geometry
 {
 /**
- * Transfer bone weights from one mesh (source) to another (target). 
+ * Transfer bone weights from one mesh (source) to another (target). Uses the dynamic mesh bone attributes to reindex 
+ * the bone indices of the transferred weights from the source to the target skeletons. If both meshes have identical 
+ * bone name attributes, then reindexing is skipped.
  * 
+ * During the reindexing, if a weighted bone in the source skeleton is not present in the target skeleton, then the
+ * weight is not transferred (skipped), and an error is printed to the console. For best results, the target skeleton
+ * should be a superset of all the bones that are indexed by the transferred weights.
  * 
- * Example usage:
  * 
- * FDynamicMesh SourceMesh = ...; // Mesh we transfering weights from
- * FDynamicMesh TargetMesh = ...; // Mesh we are transfering weights into
+ * Example usage:
  * 
- * // These can be initialized by accessing reference skeleton of the corresponding USkinnedAsset
- * TMap<uint16, FName> SourceIndexToBone = ...; // Map the bone index to the bone name for the source mesh
- * TMap<FName, uint16> TargetBoneToIndex = ...; // Map the bone name to the bone index for the target mesh
+ * FDynamicMesh SourceMesh = ...; // Mesh we transferring weights from. Must have bone attributes.
+ * FDynamicMesh TargetMesh = ...; // Mesh we are transferring weights to.
  *
  * FTransferBoneWeights TransferBoneWeights(&SourceMesh, FSkeletalMeshAttributes::DefaultSkinWeightProfileName);
- * TransferBoneWeights.SourceIndexToBone = &SourceIndexToBone; 
- * TransferBoneWeights.TargetBoneToIndex = &TargetBoneToIndex;
  * 
- * // Optionally transform the target mesh. This is useful when you want to align two meshes in world space.
+ * // Optionally, transform the target mesh. This is useful when you want to align the two meshes in world space.
  * FTransformSRT3d InToWorld = ...; 
  * 
+ * // When transferring weights from a dynamic mesh with bone attributes to a dynamic mesh without bone attributes,
+ * // first copy over the bone attributes from the source to the target.
+ * if (!TargetMesh.HasAttributes() || !TargetMesh.Attributes()->HasBones())
+ * {
+ *     TargetMesh.EnableAttributes();
+ *     TargetMesh.Attributes()->CopyBoneAttributes(*SourceMesh.Attributes());
+ * }
+ * 
  * if (TransferBoneWeights.Validate() == EOperationValidationResult::Ok) 
  * {
  *      TransferBoneWeights.Compute(TargetMesh, InToWorld, FSkeletalMeshAttributes::DefaultSkinWeightProfileName);
@@ -72,20 +80,16 @@ public:
 
 	/** Enable/disable multi-threading. */
 	bool bUseParallel = true;
-
-    /** 
-     * We assume that each mesh inherits its reference skeleton from the same USkeleton asset. However, their internal 
-     * indexing can be different and hence when transfering weights we need to make sure we reference the correct bones 
-     * via their names instead of indices.
-     * 
-     * TODO: We can store this information inside DynamicVertexSkinWeightsAttribute when converting Skeletal Meshes to Dynamic Meshes.
-     */
-	const TMap<uint16, FName>* SourceIndexToBone = nullptr; // Map the bone index to the bone name for the source mesh
-	const TMap<FName, uint16>* TargetBoneToIndex = nullptr; // Map the bone name to the bone index for the target mesh
 	
 	/** The transfer method to compute the bone weights. */
 	ETransferBoneWeightsMethod TransferMethod = ETransferBoneWeightsMethod::ClosestPointOnSurface;
 
+	/** 
+     * Completely ignore the source and target mesh bone attributes when transferring weights from one dynamic mesh to another.
+     * This skips re-indexing and simply copies skin weights over. Use with caution.
+	 */
+	bool bIgnoreBoneAttributes = false;
+
 protected:
 		
 	/** Source mesh we are transfering weights from. */
@@ -104,7 +108,15 @@ protected:
 	TUniquePtr<FDynamicMeshAABBTree3> InternalSourceBVH;
 
 public:
-
+	
+	/**
+	 * @param InSourceMesh The mesh we are transferring weights from 
+	 * @param InSourceProfileName The profile name of the skin weight attribute we are transferring weights from.
+	 * @param SourceBVH Optional source mesh BVH. If not provided, one will be computed internally. 
+	 * 
+	 * @note Assumes that the InSourceMesh has bone attributes, use bIgnoreBoneAttributes flag to ignore the bone 
+	 * 		 attributes and skip re-indexing.
+	 */
 	FTransferBoneWeights(const FDynamicMesh3* InSourceMesh, 
 					     const FName& InSourceProfileName,
 					     const FDynamicMeshAABBTree3* SourceBVH = nullptr); 
@@ -126,6 +138,9 @@ public:
      *       					  then the data will be overwritten, otherwise a new attribute will be created.
      * 
      * @return true if the algorithm succeeds, false if it failed or was canceled by the user.
+	 * 
+	 * @note Assumes that the InOutTargetMesh has bone attributes, use bIgnoreBoneAttributes flag to ignore the bone 
+	 * 		 attributes and skip re-indexing.
 	 */
 	virtual bool Compute(FDynamicMesh3& InOutTargetMesh, const FTransformSRT3d& InToWorld, const FName& InTargetProfileName);
 
@@ -135,10 +150,12 @@ public:
 	 * @param InPoint Point for which we are computing a bone weight
 	 * @param InToWorld Transform applied to the point
      * @param OutWeights Bone weight computed for the input transformed point
+	 * @param TargetBoneToIndex Optional map from the bone names to the bone indices of the target skeleton. 
+	 * 							If null, the bone indices of the skinning weights will not be re-indexed after the transfer.
 	 * 
      * @return true if the algorithm succeeds, false if it failed or was canceled by the user.
 	 */
-	virtual bool Compute(const FVector3d& InPoint, const FTransformSRT3d& InToWorld, UE::AnimationCore::FBoneWeights& OutWeights);
+	virtual bool Compute(const FVector3d& InPoint, const FTransformSRT3d& InToWorld, UE::AnimationCore::FBoneWeights& OutWeights, const TMap<FName, uint16>* TargetBoneToIndex = nullptr);
 
 protected:
 	
