CIS - Fixed warning The 'RenderMaterialPathName' reference becomes invalid in TerminalNode.cpp and ClothAssetBuilderEditor.cpp.

#rnx

diff --git a/Engine/Plugins/Experimental/ChaosClothAssetEditor/Source/ChaosClothAssetEditor/Private/ChaosClothAsset/DataflowNodes/TerminalNode.cpp b/Engine/Plugins/Experimental/ChaosClothAssetEditor/Source/ChaosClothAssetEditor/Private/ChaosClothAsset/DataflowNodes/TerminalNode.cpp
--- a/Engine/Plugins/Experimental/ChaosClothAssetEditor/Source/ChaosClothAssetEditor/Private/ChaosClothAsset/DataflowNodes/TerminalNode.cpp
+++ b/Engine/Plugins/Experimental/ChaosClothAssetEditor/Source/ChaosClothAssetEditor/Private/ChaosClothAsset/DataflowNodes/TerminalNode.cpp
@@ -63,9 +63,10 @@ void FChaosClothAssetTerminalNode::SetAssetValue(TObjectPtr<UObject> Asset, Data
 				TArray<FSkeletalMaterial>& Materials = ClothAsset->GetMaterials();
 				Materials.Reserve(Materials.Num() + NumLodMaterials);
 
+				const TConstArrayView<FString> LodRenderMaterialPathName = ClothLodFacade.GetRenderMaterialPathName();
 				for (int32 LodMaterialIndex = 0; LodMaterialIndex < NumLodMaterials; ++LodMaterialIndex)
 				{
-					const FString& RenderMaterialPathName = ClothLodFacade.GetRenderMaterialPathName()[LodMaterialIndex];
+					const FString& RenderMaterialPathName = LodRenderMaterialPathName[LodMaterialIndex];
 
 					auto Predicate = [&RenderMaterialPathName](const FSkeletalMaterial& SkeletalMaterial)
 						{
diff --git a/Engine/Plugins/Experimental/ChaosClothAssetEditor/Source/ChaosClothAssetTools/Private/ChaosClothAsset/ClothAssetBuilderEditor.cpp b/Engine/Plugins/Experimental/ChaosClothAssetEditor/Source/ChaosClothAssetTools/Private/ChaosClothAsset/ClothAssetBuilderEditor.cpp
--- a/Engine/Plugins/Experimental/ChaosClothAssetEditor/Source/ChaosClothAssetTools/Private/ChaosClothAsset/ClothAssetBuilderEditor.cpp
+++ b/Engine/Plugins/Experimental/ChaosClothAssetEditor/Source/ChaosClothAssetTools/Private/ChaosClothAsset/ClothAssetBuilderEditor.cpp
@@ -41,12 +41,14 @@ void UClothAssetBuilderEditor::BuildLod(FSkeletalMeshLODModel& LODModel, const U
 	const FCollectionClothLodConstFacade ClothLodFacade = ClothFacade.GetLod(LodIndex);
 
 	const int32 NumLodMaterials = ClothLodFacade.GetNumMaterials();
+	const TConstArrayView<FString> LodRenderMaterialPathName = ClothLodFacade.GetRenderMaterialPathName();
 
 	TArray<int32> LodMaterialMap;
 	LodMaterialMap.AddUninitialized(NumLodMaterials);
+
 	for (int32 LodMaterialIndex = 0; LodMaterialIndex < NumLodMaterials; ++LodMaterialIndex)
 	{
-		const FString& RenderMaterialPathName = ClothLodFacade.GetRenderMaterialPathName()[LodMaterialIndex];
+		const FString& RenderMaterialPathName = LodRenderMaterialPathName[LodMaterialIndex];
 		const int32 SkeletalMaterialIndex = Materials.IndexOfByPredicate([&RenderMaterialPathName](const FSkeletalMaterial& SkeletalMaterial)
 			{
 				return SkeletalMaterial.MaterialInterface && SkeletalMaterial.MaterialInterface->GetPathName() == RenderMaterialPathName;
