Cloth Editor: use per-triangle normals for 2D pattern rendering

diff --git a/Engine/Plugins/Experimental/ChaosClothAssetEditor/Source/ChaosClothAssetEditor/Private/ChaosClothAsset/ClothEditorMode.cpp b/Engine/Plugins/Experimental/ChaosClothAssetEditor/Source/ChaosClothAssetEditor/Private/ChaosClothAsset/ClothEditorMode.cpp
--- a/Engine/Plugins/Experimental/ChaosClothAssetEditor/Source/ChaosClothAssetEditor/Private/ChaosClothAsset/ClothEditorMode.cpp
+++ b/Engine/Plugins/Experimental/ChaosClothAssetEditor/Source/ChaosClothAssetEditor/Private/ChaosClothAsset/ClothEditorMode.cpp
@@ -516,6 +516,12 @@ void UChaosClothAssetEditorMode::ReinitializeDynamicMeshComponents()
 				FClothPatternToDynamicMesh Converter;
 				Converter.Convert(ChaosClothAsset, LodIndex, PatternIndex, bPattern2DMode, PatternMesh);
 
+				if (bPattern2DMode)
+				{
+					// Use per-triangle normals for the 2D view
+					UE::Geometry::FMeshNormals::InitializeMeshToPerTriangleNormals(&PatternMesh);
+				}
+
 				// We only need an actor to allow use of HHitProxy for selection
 				const FRotator Rotation(0.0f, 0.0f, 0.0f);
 				const FActorSpawnParameters SpawnInfo;
