FClothAssetDatasmithImportNode: create a temporary package to import into

diff --git a/Engine/Plugins/Experimental/ChaosClothAssetEditor/Source/ChaosClothAssetEditor/Private/ChaosClothAsset/DataflowNodes/DatasmithImportNode.cpp b/Engine/Plugins/Experimental/ChaosClothAssetEditor/Source/ChaosClothAssetEditor/Private/ChaosClothAsset/DataflowNodes/DatasmithImportNode.cpp
--- a/Engine/Plugins/Experimental/ChaosClothAssetEditor/Source/ChaosClothAssetEditor/Private/ChaosClothAsset/DataflowNodes/DatasmithImportNode.cpp
+++ b/Engine/Plugins/Experimental/ChaosClothAssetEditor/Source/ChaosClothAssetEditor/Private/ChaosClothAsset/DataflowNodes/DatasmithImportNode.cpp
@@ -16,10 +16,8 @@
 
 FChaosClothAssetDatasmithImportNode::FChaosClothAssetDatasmithImportNode(const Dataflow::FNodeParameters& InParam, FGuid InGuid)
 	: FDataflowNode(InParam, InGuid)
-	, DestPackageName("/Game/ClothAsset")
 {
 	RegisterInputConnection(&DatasmithFile);
-	RegisterInputConnection(&DestPackageName);
 	RegisterOutputConnection(&Collection);
 }
 
@@ -28,7 +26,6 @@ bool FChaosClothAssetDatasmithImportNode::EvaluateImpl(Dataflow::FContext& Conte
 	using namespace UE::DatasmithImporter;
 
 	const FFilePath& InFilePath = GetValue<FFilePath>(Context, &DatasmithFile);
-	const FString& InDestPackageName = GetValue<FString>(Context, &DestPackageName);
 
 	const FSourceUri SourceUri = FSourceUri::FromFilePath(InFilePath.FilePath);
 	const TSharedPtr<FExternalSource> ExternalSource = IExternalSourceModule::GetOrCreateExternalSource(SourceUri);
@@ -42,21 +39,22 @@ bool FChaosClothAssetDatasmithImportNode::EvaluateImpl(Dataflow::FContext& Conte
 	const FText LoggerLabel(NSLOCTEXT("ImportDatasmithClothNode", "LoggerLabel", "ImportDatasmithClothNode"));
 	FDatasmithImportContext DatasmithImportContext(ExternalSource, bLoadConfig, LoggerName, LoggerLabel);
 
-	const TStrongObjectPtr<const UPackage> DestinationPackage(CreatePackage(*InDestPackageName));		// TODO: would be nice to not have to create the package but it seems necessary for now
+	const FName PackageName = MakeUniqueObjectName(nullptr, UPackage::StaticClass(), FName(FPaths::Combine(GetTransientPackage()->GetPathName(), ExternalSource->GetSourceName())));
+	const TStrongObjectPtr<UPackage> DestinationPackage(CreatePackage(*PackageName.ToString()));
 	if (!ensure(DestinationPackage))
 	{
 		// Failed to create the package to hold this asset for some reason
 		return false;
 	}
-	
+
 	// Don't create the Actors in the level, just read the Assets
 	DatasmithImportContext.Options->BaseOptions.SceneHandling = EDatasmithImportScene::AssetsOnly;
 
-	constexpr EObjectFlags NewObjectFlags = RF_Public | RF_Standalone | RF_Transactional;
+	constexpr EObjectFlags NewObjectFlags = RF_Public | RF_Transactional | RF_Transient | RF_Standalone;
 	const TSharedPtr<FJsonObject> ImportSettingsJson;
 	constexpr bool bIsSilent = true;
-	const FString DestinationPath = DestinationPackage->GetName();
-	if (!DatasmithImportContext.Init(DestinationPath, NewObjectFlags, GWarn, ImportSettingsJson, bIsSilent))
+
+	if (!DatasmithImportContext.Init(DestinationPackage->GetPathName(), NewObjectFlags, GWarn, ImportSettingsJson, bIsSilent))
 	{
 		return false;
 	}
@@ -76,13 +74,14 @@ bool FChaosClothAssetDatasmithImportNode::EvaluateImpl(Dataflow::FContext& Conte
 
 	if (bImportSucceed && DatasmithImportContext.ImportedClothes.Num() > 0)
 	{
-		const UObject* const ClothObject = DatasmithImportContext.ImportedClothes.CreateIterator()->Value;
-		const UChaosClothAsset* const DatasmithClothAsset = Cast<UChaosClothAsset>(ClothObject);
+		UObject* const ClothObject = DatasmithImportContext.ImportedClothes.CreateIterator()->Value;
+		UChaosClothAsset* const DatasmithClothAsset = Cast<UChaosClothAsset>(ClothObject);
 		if (ensure(DatasmithClothAsset))
 		{
 			DatasmithClothAsset->GetClothCollection()->CopyTo(&OutCollection);
 			return true;
 		}
+		DatasmithClothAsset->ClearFlags(RF_Standalone);
 	}
 
 	return false;
diff --git a/Engine/Plugins/Experimental/ChaosClothAssetEditor/Source/ChaosClothAssetEditor/Private/ChaosClothAsset/DataflowNodes/DatasmithImportNode.h b/Engine/Plugins/Experimental/ChaosClothAssetEditor/Source/ChaosClothAssetEditor/Private/ChaosClothAsset/DataflowNodes/DatasmithImportNode.h
--- a/Engine/Plugins/Experimental/ChaosClothAssetEditor/Source/ChaosClothAssetEditor/Private/ChaosClothAsset/DataflowNodes/DatasmithImportNode.h
+++ b/Engine/Plugins/Experimental/ChaosClothAssetEditor/Source/ChaosClothAssetEditor/Private/ChaosClothAsset/DataflowNodes/DatasmithImportNode.h
@@ -19,10 +19,6 @@ public:
 	UPROPERTY(EditAnywhere, Category = "Datasmith Importer")
 	FFilePath DatasmithFile;
 
-	/** Package to import into. */
-	UPROPERTY(EditAnywhere, Category = "Datasmith Importer")
-	FString DestPackageName;
-
 	UPROPERTY(meta = (DataflowOutput, DisplayName = "Collection"))
 	FManagedArrayCollection Collection;
 
