Cloth Editor: set initial preview viewport view transform so the x-axis points to the right and y-axis points toward the camera

diff --git a/Engine/Plugins/Experimental/ChaosClothAssetEditor/Source/ChaosClothAssetEditor/Private/ChaosClothAsset/ClothEditorToolkit.cpp b/Engine/Plugins/Experimental/ChaosClothAssetEditor/Source/ChaosClothAssetEditor/Private/ChaosClothAsset/ClothEditorToolkit.cpp
--- a/Engine/Plugins/Experimental/ChaosClothAssetEditor/Source/ChaosClothAssetEditor/Private/ChaosClothAsset/ClothEditorToolkit.cpp
+++ b/Engine/Plugins/Experimental/ChaosClothAssetEditor/Source/ChaosClothAssetEditor/Private/ChaosClothAsset/ClothEditorToolkit.cpp
@@ -442,7 +442,7 @@ void FChaosClothAssetEditorToolkit::PostInitAssetEditor()
 	ClothPreviewViewportClient->SetClothEditorToolkit(StaticCastSharedRef<FChaosClothAssetEditorToolkit>(this->AsShared()));
 
 	SetCommonViewportClientOptions(ClothPreviewViewportClient.Get());
-	ClothPreviewViewportClient->SetInitialViewTransform(ELevelViewportType::LVT_Perspective, FVector(0, -100, 100), FRotator(0, 90, 0), DEFAULT_ORTHOZOOM);
+	ClothPreviewViewportClient->SetInitialViewTransform(ELevelViewportType::LVT_Perspective, FVector(0, 0, 0), FRotator(0, -90, 0), DEFAULT_ORTHOZOOM);
 
 	if (ClothPreviewViewportClient->Viewport != nullptr)
 	{
