Chaos Cloth - Re-added missing property copy from the cloth collection Terminal Node evaluation, that was causing cloth assets to misbehave after recent LOD code change.

#rnx

diff --git a/Engine/Plugins/Experimental/ChaosClothAssetEditor/Source/ChaosClothAssetEditor/Private/ChaosClothAsset/DataflowNodes/TerminalNode.cpp b/Engine/Plugins/Experimental/ChaosClothAssetEditor/Source/ChaosClothAssetEditor/Private/ChaosClothAsset/DataflowNodes/TerminalNode.cpp
--- a/Engine/Plugins/Experimental/ChaosClothAssetEditor/Source/ChaosClothAssetEditor/Private/ChaosClothAsset/DataflowNodes/TerminalNode.cpp
+++ b/Engine/Plugins/Experimental/ChaosClothAssetEditor/Source/ChaosClothAssetEditor/Private/ChaosClothAsset/DataflowNodes/TerminalNode.cpp
@@ -4,6 +4,7 @@
 #include "ChaosClothAsset/DataflowNodes/DataflowNodes.h"
 #include "ChaosClothAsset/ClothAsset.h"
 #include "ChaosClothAsset/CollectionClothFacade.h"
+#include "Chaos/CollectionPropertyFacade.h"
 #include "Animation/Skeleton.h"
 #include "Engine/SkinnedAssetCommon.h"
 #include "Materials/Material.h"
@@ -82,11 +83,15 @@ void FChaosClothAssetTerminalNode::SetAssetValue(TObjectPtr<UObject> Asset, Data
 					}
 				}
 
-				// Set skeleton and physics asset only with LOD 0 at the moment
+				// Set properties, skeleton, and physics asset only with LOD 0 at the moment
 				if (LodIndex == 0)
 				{
+					using namespace ::Chaos::Softs;
+
 					SkeletonPathName = InClothLodFacade.GetSkeletonAssetPathName();
 					PhysicsAssetPathName = InClothLodFacade.GetPhysicsAssetPathName();
+				
+					FCollectionPropertyMutableFacade(ClothCollection).Append(*InClothCollection);
 				}
 			}
 		}
