Chaos Cloth Dataflow: copy maps over if they exist at the terminal node.

diff --git a/Engine/Plugins/Experimental/ChaosClothAssetEditor/Source/ChaosClothAssetEditor/Private/ChaosClothAsset/DataflowNodes/TerminalNode.cpp b/Engine/Plugins/Experimental/ChaosClothAssetEditor/Source/ChaosClothAssetEditor/Private/ChaosClothAsset/DataflowNodes/TerminalNode.cpp
--- a/Engine/Plugins/Experimental/ChaosClothAssetEditor/Source/ChaosClothAssetEditor/Private/ChaosClothAsset/DataflowNodes/TerminalNode.cpp
+++ b/Engine/Plugins/Experimental/ChaosClothAssetEditor/Source/ChaosClothAssetEditor/Private/ChaosClothAsset/DataflowNodes/TerminalNode.cpp
@@ -83,6 +83,18 @@ void FChaosClothAssetTerminalNode::SetAssetValue(TObjectPtr<UObject> Asset, Data
 					}
 				}
 
+				const TArray<FName> MapNames = InClothFacade.GetWeightMapNames();
+				for (const FName& MapName : MapNames)
+				{
+					ClothFacade.AddWeightMap(MapName);
+					TArrayView<float> WeightMap = ClothLodFacade.GetWeightMap(MapName);
+					TConstArrayView<float> InWeightMap = InClothLodFacade.GetWeightMap(MapName);
+					for (int32 i = 0; i < WeightMap.Num(); ++i)
+					{
+						WeightMap[i] = InWeightMap[i];
+					}
+				}
+
 				// Set properties, skeleton, and physics asset only with LOD 0 at the moment
 				if (LodIndex == 0)
 				{
