Add Chaos Cloth outliner display of FIntVector2Type and FIntVector2ArrayType.

diff --git a/Engine/Plugins/Experimental/ChaosClothAssetEditor/Source/ChaosClothAssetEditor/Private/ChaosClothAsset/SClothCollectionOutliner.cpp b/Engine/Plugins/Experimental/ChaosClothAssetEditor/Source/ChaosClothAssetEditor/Private/ChaosClothAsset/SClothCollectionOutliner.cpp
--- a/Engine/Plugins/Experimental/ChaosClothAssetEditor/Source/ChaosClothAssetEditor/Private/ChaosClothAsset/SClothCollectionOutliner.cpp
+++ b/Engine/Plugins/Experimental/ChaosClothAssetEditor/Source/ChaosClothAssetEditor/Private/ChaosClothAsset/SClothCollectionOutliner.cpp
@@ -147,6 +147,11 @@ namespace ClothCollectionOutlinerHelpers
 		return Value;
 	}
 
+	FString AttributeValueToString(const FIntVector2& Value)
+	{
+		return FString::Printf(TEXT("X=%d Y=%d"), Value.X, Value.Y);
+	}
+
 	template<typename T>
 	FString AttributeValueToString(const TArray<T>& Array)
 	{
@@ -213,6 +218,12 @@ namespace ClothCollectionOutlinerHelpers
 		case FManagedArrayCollection::EArrayType::FStringType:
 			ValueAsString = AttributeValueToString<FString>(ClothCollection, AttributeName, GroupName, AttributeArrayIndex);
 			break;
+		case FManagedArrayCollection::EArrayType::FIntVector2Type:
+			ValueAsString = AttributeValueToString<FIntVector2>(ClothCollection, AttributeName, GroupName, AttributeArrayIndex);
+			break;
+		case FManagedArrayCollection::EArrayType::FIntVector2ArrayType:
+			ValueAsString = AttributeValueToString<TArray<FIntVector2>>(ClothCollection, AttributeName, GroupName, AttributeArrayIndex);
+			break;
 		default:
 			ensure(false);
 			ValueAsString = "(Unknown Data Type)";
