ClothTrainingTool: add functions to run cloth simulation automatically

diff --git a/Engine/Plugins/Experimental/ChaosClothAsset/Source/ChaosClothAssetEngine/Private/ChaosClothAsset/ClothSimulationProxy.cpp b/Engine/Plugins/Experimental/ChaosClothAsset/Source/ChaosClothAssetEngine/Private/ChaosClothAsset/ClothSimulationProxy.cpp
--- a/Engine/Plugins/Experimental/ChaosClothAsset/Source/ChaosClothAssetEngine/Private/ChaosClothAsset/ClothSimulationProxy.cpp
+++ b/Engine/Plugins/Experimental/ChaosClothAsset/Source/ChaosClothAssetEngine/Private/ChaosClothAsset/ClothSimulationProxy.cpp
@@ -4,6 +4,7 @@
 #include "ChaosClothAsset/ClothAsset.h"
 #include "ChaosClothAsset/ClothAssetPrivate.h"
 #include "ChaosClothAsset/ClothComponent.h"
+#include "ChaosClothAsset/ClothSimulationContext.h"
 #include "ChaosClothAsset/ClothSimulationMesh.h"
 #include "ChaosClothAsset/ClothSimulationModel.h"
 #include "ChaosCloth/ChaosClothingSimulationCloth.h"
@@ -77,7 +78,7 @@ namespace UE::Chaos::ClothAsset
 			FScopeCycleCounterUObject ContextScope(ClothSimulationProxy.ClothComponent.GetSkinnedAsset());
 			CSV_SCOPED_TIMING_STAT(Animation, Cloth);
 
-			ClothSimulationProxy.Tick_PhysicsThread();
+			ClothSimulationProxy.Tick();
 		}
 
 	private:
@@ -86,6 +87,7 @@ namespace UE::Chaos::ClothAsset
 
 	FClothSimulationProxy::FClothSimulationProxy(const UChaosClothComponent& InClothComponent)
 		: ClothComponent(InClothComponent)
+		, ClothSimulationContext(MakeUnique<FClothSimulationContext>())
 		, Solver(MakeUnique<::Chaos::FClothingSimulationSolver>())
 		, Visualization(MakeUnique<::Chaos::FClothVisualization>(Solver.Get()))
 		, MaxDeltaTime(UPhysicsSettings::Get()->MaxPhysicsDeltaTime)
@@ -97,12 +99,12 @@ namespace UE::Chaos::ClothAsset
 		// Need a valid context to initialize the mesh
 		constexpr bool bIsInitialization = true;
 		constexpr Softs::FSolverReal NoAdvanceDt = 0.f;
-		ClothSimulationContext.Fill(ClothComponent, NoAdvanceDt, MaxDeltaTime, bIsInitialization);
+		ClothSimulationContext->Fill(ClothComponent, NoAdvanceDt, MaxDeltaTime, bIsInitialization);
 
 		// Setup startup transforms
 		constexpr bool bNeedsReset = true;
-		Solver->SetLocalSpaceLocation((FVec3)ClothSimulationContext.ComponentTransform.GetLocation(), bNeedsReset);
-		Solver->SetLocalSpaceRotation((FQuat)ClothSimulationContext.ComponentTransform.GetRotation());
+		Solver->SetLocalSpaceLocation((FVec3)ClothSimulationContext->ComponentTransform.GetLocation(), bNeedsReset);
+		Solver->SetLocalSpaceRotation((FQuat)ClothSimulationContext->ComponentTransform.GetRotation());
 
 		// Create mesh simulation thread object
 		const UChaosClothAsset* const ClothAsset = ClothComponent.GetClothAsset();
@@ -119,7 +121,7 @@ namespace UE::Chaos::ClothAsset
 			FString::Format(TEXT("{0}|{1}"), { ClothComponent.GetOwner()->GetName(), ClothComponent.GetName() }) :
 			ClothComponent.GetName();
 #endif
-		const int32 MeshIndex = Meshes.Emplace(MakeUnique<FClothSimulationMesh>(*ClothSimulationModel, ClothSimulationContext, DebugName));
+		const int32 MeshIndex = Meshes.Emplace(MakeUnique<FClothSimulationMesh>(*ClothSimulationModel, *ClothSimulationContext, DebugName));
 
 		// Create collider simulation thread object
 PRAGMA_DISABLE_DEPRECATION_WARNINGS  // TODO: CHAOS_IS_CLOTHINGSIMULATIONMESH_ABSTRACT
@@ -150,8 +152,8 @@ PRAGMA_ENABLE_DEPRECATION_WARNINGS
 		Solver->SetConfig(Configs[SolverConfigIndex].Get());
 
 		// Set start pose (update the context, then the solver without advancing the simulation)
-		ClothSimulationContext.Fill(ClothComponent, NoAdvanceDt, MaxDeltaTime);
-		Solver->Update((Softs::FSolverReal)ClothSimulationContext.DeltaTime);
+		ClothSimulationContext->Fill(ClothComponent, NoAdvanceDt, MaxDeltaTime);
+		Solver->Update((Softs::FSolverReal)ClothSimulationContext->DeltaTime);
 	}
 
 	FClothSimulationProxy::~FClothSimulationProxy()
@@ -164,9 +166,9 @@ PRAGMA_ENABLE_DEPRECATION_WARNINGS
 		SCOPE_CYCLE_COUNTER(STAT_ClothSimulationProxy_TickGame);
 
 		// Fill a new context, note the context is also needed when the simulation is suspended
-		ClothSimulationContext.Fill(ClothComponent, DeltaTime, MaxDeltaTime, false, CacheData);
+		ClothSimulationContext->Fill(ClothComponent, DeltaTime, MaxDeltaTime, false, CacheData);
 
-		const bool bUseCache = ClothSimulationContext.CacheData.CacheIndices.Num() > 0;
+		const bool bUseCache = ClothSimulationContext->CacheData.CacheIndices.Num() > 0;
 		if (bUseCache)
 		{
 			Solver->SetEnableSolver(false);
@@ -175,11 +177,7 @@ PRAGMA_ENABLE_DEPRECATION_WARNINGS
 		const bool bCreateParallelTask = (DeltaTime > 0.f && !ClothComponent.IsSimulationSuspended()) || bUseCache;
 		if (bCreateParallelTask)
 		{
-			// Replace physics thread's configs with the game thread's configs
-			for (const TUniquePtr<::Chaos::FClothingSimulationConfig>& Config : Configs)
-			{
-				Config->Initialize(ClothComponent.GetPropertyCollection());  // TODO: Outfit and multi-cloths/solver config update
-			}
+			InitializeConfigs();
 
 			// Start the the cloth simulation thread
 			ParallelTask = TGraphTask<FClothSimulationProxyParallelTask>::CreateTask(nullptr, ENamedThreads::GameThread).ConstructAndDispatchWhenReady(*this);
@@ -187,19 +185,19 @@ PRAGMA_ENABLE_DEPRECATION_WARNINGS
 		else
 		{
 			// It still needs to write back to the GT cache as the context has changed
-			WriteSimulationData_GameThread();
+			WriteSimulationData();
 		}
 	}
 
-	void FClothSimulationProxy::Tick_PhysicsThread()
+	void FClothSimulationProxy::Tick()
 	{
 		using namespace ::Chaos;
 
 		TRACE_CPUPROFILER_EVENT_SCOPE(FClothSimulationProxy_TickPhysics);
 		SCOPE_CYCLE_COUNTER(STAT_ClothSimulationProxy_TickPhysics);
-		const bool bUseCache = ClothSimulationContext.CacheData.HasData();
+		const bool bUseCache = ClothSimulationContext->CacheData.HasData();
 
-		if (ClothSimulationContext.DeltaTime == 0.f && !bUseCache)
+		if (ClothSimulationContext->DeltaTime == 0.f && !bUseCache)
 		{
 			return;
 		}
@@ -207,15 +205,15 @@ PRAGMA_ENABLE_DEPRECATION_WARNINGS
 		const double StartTime = FPlatformTime::Seconds();
 		const float PrevSimulationTime = SimulationTime;  // Copy the atomic to prevent a re-read
 
-		const bool bNeedsReset = (ClothSimulationContext.bReset || PrevSimulationTime == 0.f);  // Reset on the first frame too since the simulation is created in bind pose, and not in start pose
-		const bool bNeedsTeleport = ClothSimulationContext.bTeleport;
+		const bool bNeedsReset = (ClothSimulationContext->bReset || PrevSimulationTime == 0.f);  // Reset on the first frame too since the simulation is created in bind pose, and not in start pose
+		const bool bNeedsTeleport = ClothSimulationContext->bTeleport;
 		bIsTeleported = bNeedsTeleport;
 
 		// Update Solver animatable parameters
-		Solver->SetLocalSpaceLocation((FVec3)ClothSimulationContext.ComponentTransform.GetLocation(), bNeedsReset);
-		Solver->SetLocalSpaceRotation((FQuat)ClothSimulationContext.ComponentTransform.GetRotation());
-		Solver->SetWindVelocity(ClothSimulationContext.WindVelocity);
-		Solver->SetGravity(ClothSimulationContext.WorldGravity);
+		Solver->SetLocalSpaceLocation((FVec3)ClothSimulationContext->ComponentTransform.GetLocation(), bNeedsReset);
+		Solver->SetLocalSpaceRotation((FQuat)ClothSimulationContext->ComponentTransform.GetRotation());
+		Solver->SetWindVelocity(ClothSimulationContext->WindVelocity);
+		Solver->SetGravity(ClothSimulationContext->WorldGravity);
 
 		// Check teleport modes
 		for (const TUniquePtr<FClothingSimulationCloth>& Cloth : Cloths)
@@ -236,11 +234,11 @@ PRAGMA_ENABLE_DEPRECATION_WARNINGS
 		// Step the simulation
 		if (Solver->GetEnableSolver() || !bUseCache)
 		{
-			Solver->Update((Softs::FSolverReal)ClothSimulationContext.DeltaTime);
+			Solver->Update((Softs::FSolverReal)ClothSimulationContext->DeltaTime);
 		}
 		else
 		{
-			Solver->UpdateFromCache(ClothSimulationContext.CacheData);
+			Solver->UpdateFromCache(ClothSimulationContext->CacheData);
 		}
 
 		// Keep the actual used number of iterations for the stats
@@ -273,16 +271,15 @@ PRAGMA_ENABLE_DEPRECATION_WARNINGS
 			ParallelTask.SafeRelease();
 
 			// Write back to the GT cache
-			WriteSimulationData_GameThread();
+			WriteSimulationData();
 		}
 	}
 
-	void FClothSimulationProxy::WriteSimulationData_GameThread()
+
+	void FClothSimulationProxy::WriteSimulationData()
 	{
 		using namespace ::Chaos;
 
-		check(IsInGameThread());
-
 		CSV_SCOPED_TIMING_STAT(Animation, Cloth);
 		TRACE_CPUPROFILER_EVENT_SCOPE(FClothSimulationProxy_WriteSimulationData);
 		SCOPE_CYCLE_COUNTER(STAT_ClothSimulationProxy_WriteSimulationData);
@@ -356,13 +353,13 @@ PRAGMA_ENABLE_DEPRECATION_WARNINGS
 
 			// Get the reference transform used in the current animation pose
 			FTransform ReferenceBoneTransform = ComponentSpaceTransforms[ReferenceBoneIndex];
-			ReferenceBoneTransform *= ClothSimulationContext.ComponentTransform;
+			ReferenceBoneTransform *= ClothSimulationContext->ComponentTransform;
 			ReferenceBoneTransform.SetScale3D(FVector(1.0f));  // Scale is already baked in the cloth mesh
 
 			// Set the world space transform to be this cloth's reference bone
 			FClothSimulData& Data = CurrentSimulationData.FindOrAdd(AssetIndex);
 			Data.Transform = ReferenceBoneTransform;
-			Data.ComponentRelativeTransform = ReferenceBoneTransform.GetRelativeTransform(ClothSimulationContext.ComponentTransform);
+			Data.ComponentRelativeTransform = ReferenceBoneTransform.GetRelativeTransform(ClothSimulationContext->ComponentTransform);
 
 			// Retrieve the last reference space transform used for this cloth
 			// Note: This won't necessary match the current bone reference transform when the simulation is paused,
@@ -436,4 +433,18 @@ PRAGMA_ENABLE_DEPRECATION_WARNINGS
 		}
 		return FBoxSphereBounds(ForceInit);
 	}
+
+	void FClothSimulationProxy::InitializeConfigs()
+	{
+		// Replace physics thread's configs with the game thread's configs
+		for (const TUniquePtr<::Chaos::FClothingSimulationConfig>& Config : Configs)
+		{
+			Config->Initialize(ClothComponent.GetPropertyCollection());  // TODO: Outfit and multi-cloths/solver config update
+		}
+	}
+
+	void FClothSimulationProxy::FillSimulationContext(float DeltaTime, bool bIsInitialization, FClothingSimulationCacheData* CacheData)
+	{
+		ClothSimulationContext->Fill(ClothComponent, DeltaTime, MaxDeltaTime, bIsInitialization, CacheData);
+	}
 }
diff --git a/Engine/Plugins/Experimental/ChaosClothAsset/Source/ChaosClothAssetEngine/Public/ChaosClothAsset/ClothSimulationProxy.h b/Engine/Plugins/Experimental/ChaosClothAsset/Source/ChaosClothAssetEngine/Public/ChaosClothAsset/ClothSimulationProxy.h
--- a/Engine/Plugins/Experimental/ChaosClothAsset/Source/ChaosClothAssetEngine/Public/ChaosClothAsset/ClothSimulationProxy.h
+++ b/Engine/Plugins/Experimental/ChaosClothAsset/Source/ChaosClothAssetEngine/Public/ChaosClothAsset/ClothSimulationProxy.h
@@ -2,7 +2,6 @@
 
 #pragma once
 
-#include "ChaosClothAsset/ClothSimulationContext.h"
 #include "Async/TaskGraphInterfaces.h"
 #include "Containers/Array.h"
 #include "Math/Transform.h"
@@ -27,11 +26,13 @@ struct FClothingSimulationCacheData;
 
 namespace UE::Chaos::ClothAsset
 {
+	struct FClothSimulationContext;
+
 	/**
 	 * Cloth simulation proxy.
 	 * Class used to share data between the cloth simulation and the cloth component.
 	 */
-	class FClothSimulationProxy final
+	class CHAOSCLOTHASSETENGINE_API FClothSimulationProxy
 	{
 	public:
 		explicit FClothSimulationProxy(const UChaosClothComponent& InClothComponent);
@@ -57,11 +58,13 @@ namespace UE::Chaos::ClothAsset
 
 		FBoxSphereBounds CalculateBounds_AnyThread() const;
 
-	private:
-		void Tick_PhysicsThread();
-
-		void WriteSimulationData_GameThread();
+	protected:
+		void Tick();
+		void WriteSimulationData();
+		void InitializeConfigs();
+		void FillSimulationContext(float DeltaTime, bool bIsInitialization = false, FClothingSimulationCacheData* CacheData = nullptr);
 
+	private:
 		// Internal physics thread object
 		friend class FClothSimulationProxyParallelTask;
 
@@ -75,7 +78,7 @@ namespace UE::Chaos::ClothAsset
 		const UChaosClothComponent& ClothComponent;
 
 		// Simulation context used to store the required component data for the duration of the simulation
-		FClothSimulationContext ClothSimulationContext;
+		TUniquePtr<FClothSimulationContext> ClothSimulationContext;
 
 		// The cloth simulation model used to create this simulation, ownership might get transferred to this proxy if it changes during the simulation
 		TSharedPtr<const FChaosClothSimulationModel> ClothSimulationModel;
diff --git a/Engine/Plugins/Experimental/ChaosClothAssetEditor/ChaosClothAssetEditor.uplugin b/Engine/Plugins/Experimental/ChaosClothAssetEditor/ChaosClothAssetEditor.uplugin
--- a/Engine/Plugins/Experimental/ChaosClothAssetEditor/ChaosClothAssetEditor.uplugin
+++ b/Engine/Plugins/Experimental/ChaosClothAssetEditor/ChaosClothAssetEditor.uplugin
@@ -63,6 +63,10 @@
 		{
 			"Name": "DatasmithImporter",
 			"Enabled": true
+		},
+		{
+			"Name": "ChaosCaching",
+			"Enabled": true
 		}
 	]
 }
\ No newline at end of file
diff --git a/Engine/Plugins/Experimental/ChaosClothAssetEditor/Source/ChaosClothAssetEditorTools/ChaosClothAssetEditorTools.Build.cs b/Engine/Plugins/Experimental/ChaosClothAssetEditor/Source/ChaosClothAssetEditorTools/ChaosClothAssetEditorTools.Build.cs
--- a/Engine/Plugins/Experimental/ChaosClothAssetEditor/Source/ChaosClothAssetEditorTools/ChaosClothAssetEditorTools.Build.cs
+++ b/Engine/Plugins/Experimental/ChaosClothAssetEditor/Source/ChaosClothAssetEditorTools/ChaosClothAssetEditorTools.Build.cs
@@ -53,7 +53,8 @@ public class ChaosClothAssetEditorTools : ModuleRules
 				"MeshConversion",
 				"MeshDescription",
 				"SkeletalMeshDescription",
-				"ChaosClothAsset"
+				"ChaosClothAsset",
+				"ChaosCaching",
 				// ... add private dependencies that you statically link with here ...	
 			}
 			);
diff --git a/Engine/Plugins/Experimental/ChaosClothAssetEditor/Source/ChaosClothAssetEditorTools/Private/ChaosClothAsset/ClothTrainingTool.cpp b/Engine/Plugins/Experimental/ChaosClothAssetEditor/Source/ChaosClothAssetEditorTools/Private/ChaosClothAsset/ClothTrainingTool.cpp
--- a/Engine/Plugins/Experimental/ChaosClothAssetEditor/Source/ChaosClothAssetEditorTools/Private/ChaosClothAsset/ClothTrainingTool.cpp
+++ b/Engine/Plugins/Experimental/ChaosClothAssetEditor/Source/ChaosClothAssetEditorTools/Private/ChaosClothAsset/ClothTrainingTool.cpp
@@ -2,40 +2,94 @@
 
 #include "ChaosClothAsset/ClothTrainingTool.h"
 #include "CoreMinimal.h"
+#include "ChaosCloth/ChaosClothingSimulationConfig.h"
 #include "ChaosClothAsset/ClothComponent.h"
+#include "ChaosClothAsset/ClothComponentToolTarget.h"
+#include "ChaosClothAsset/ClothCollection.h"
+#include "ChaosClothAsset/ClothSimulationProxy.h"
+#include "ClothingSystemRuntimeTypes.h"
+#include "Chaos/ChaosCache.h"
+#include "Chaos/CacheCollection.h"
+#include "ComponentReregisterContext.h"
 #include "ModelingOperators.h"
 #include "Misc/ScopedSlowTask.h"
 #include "ContextObjectStore.h"
 #include "InteractiveToolManager.h"
 #include "ToolTargetManager.h"
-#include "ChaosClothAsset/ClothComponentToolTarget.h"
+#include "UObject/SavePackage.h"
 
 #define LOCTEXT_NAMESPACE "ClothTrainingTool"
 
-// ------------------- Operator -------------------
+class UClothTrainingTool::FClothSimulationDataGenerationProxy : public UE::Chaos::ClothAsset::FClothSimulationProxy
+{
+public:
+	explicit FClothSimulationDataGenerationProxy(const UChaosClothComponent& InClothComponent);
+	~FClothSimulationDataGenerationProxy();
+
+	using UE::Chaos::ClothAsset::FClothSimulationProxy::Tick;
+	using UE::Chaos::ClothAsset::FClothSimulationProxy::FillSimulationContext;
+	using UE::Chaos::ClothAsset::FClothSimulationProxy::InitializeConfigs;
+	using UE::Chaos::ClothAsset::FClothSimulationProxy::WriteSimulationData;
+};
 
-class FTrainClothOp : public UE::Geometry::TGenericDataOperator<FSkinnedMeshVertices>
+UClothTrainingTool::FClothSimulationDataGenerationProxy::FClothSimulationDataGenerationProxy(const UChaosClothComponent& InClothComponent)
+	: FClothSimulationProxy(InClothComponent)
+{	
+}
+
+UClothTrainingTool::FClothSimulationDataGenerationProxy::~FClothSimulationDataGenerationProxy()
 {
-public: 
+}
 
-	TObjectPtr<UAnimSequence> AnimSequence;
+class UClothTrainingTool::FGenerateClothOp : public UE::Geometry::TGenericDataOperator<FSkinnedMeshVertices>
+{
+public: 
+	FGenerateClothOp(TObjectPtr<UAnimSequence> InAnimSequence, TObjectPtr<UChaosCache> InCache, TObjectPtr<UChaosClothComponent> InClothComponent, FClothSimulationDataGenerationProxy* InDataGenerationProxy)
+		: AnimSequence(InAnimSequence)
+		, Cache(InCache)
+		, ClothComponent(InClothComponent)
+		, DataGenerationProxy(InDataGenerationProxy)
+	{
+	}
 
 	virtual void CalculateResult(FProgressCancel* Progress) override;
-};
 
+private:
+	void Simulate(float DeltaTime, int32 NumSteps);
+	void BackupClothComponentState();
+	void RestoreClothComponentState();
+	const TObjectPtr<UAnimSequence> AnimSequence;
+	TObjectPtr<UChaosCache> Cache;
+	TObjectPtr<UChaosClothComponent> ClothComponent;
+	FClothSimulationDataGenerationProxy *DataGenerationProxy = nullptr;
+
+	bool bIsSimulationSuspendedBackup = false;
+	bool bTeleportBackup = false;
+	bool bResetBackup = false;
+};
 
-void FTrainClothOp::CalculateResult(FProgressCancel* Progress)
+void UClothTrainingTool::FGenerateClothOp::CalculateResult(FProgressCancel* Progress)
 {
-	// @@@@@@@@@ TODO: Do the posing, simulation, getting outputs, etc. Check for Progress->Cancelled() regularly if possible.
+	if (ClothComponent == nullptr || Cache == nullptr || DataGenerationProxy == nullptr)
+	{
+		return;
+	}
 
-	FSkinnedMeshVertices OutputVertices;
+	static const FName PositionXName = TEXT("PositionX");
+	static const FName PositionYName = TEXT("PositionY");
+	static const FName PositionZName = TEXT("PositionZ");
 	
 	bool bCancelled = false;
+	// TODO: change this part to data from AnimSequence
+	constexpr float DeltaTime = 1e-3;
+	constexpr float CacheFPS = 30;
+	const int32 NumFrames = 10;
 
-	for (int i = 0; i < 10; ++i)
-	{
-		FPlatformProcess::Sleep(1.0f);
+	BackupClothComponentState();
 
+	ClothComponent->ResumeSimulation();
+	for (int32 Frame = 0; Frame < NumFrames; Frame++)
+	{
 		if (Progress)
 		{
 			if (Progress->Cancelled())
@@ -44,15 +98,102 @@ void FTrainClothOp::CalculateResult(FProgressCancel* Progress)
 				break;
 			}
 
-			OutputVertices.Vertices.Add(FVector3f(i,i,i));
+			FPendingFrameWrite NewFrame;
+			const float TotalTime = Frame * DeltaTime * 5;
+			Simulate(DeltaTime, (Frame + 1) * 5);
+
+			const TMap<int32, FClothSimulData> &SimulDataMap = DataGenerationProxy->GetCurrentSimulationData_AnyThread();
+			const FClothSimulData* const SimulData = SimulDataMap.Find(0);
+			if (SimulDataMap.Num() > 1)
+			{
+				ensureMsgf(false, TEXT("Only support single cloth for now."));
+				continue;
+			}
+			if (SimulData == nullptr)
+			{
+				ensureMsgf(false, TEXT("ClothSimulData is nullptr"));
+				continue;
+			}
+
+			const TArray<FVector3f>& SimPositions = SimulData->Positions;
+			const int32 NumParticles = SimPositions.Num();
+
+			const float Time = Frame / CacheFPS;
+			NewFrame.Time = Time;
+
+			TArray<int32>& PendingID = NewFrame.PendingChannelsIndices;
+			TArray<float> PendingPX, PendingPY, PendingPZ;
+			TArray<float> PendingVX, PendingVY, PendingVZ;
+			PendingID.SetNum(NumParticles);
+			PendingPX.SetNum(NumParticles);
+			PendingPY.SetNum(NumParticles);
+			PendingPZ.SetNum(NumParticles);
 
-			Progress->AdvanceCurrentScopeProgressBy(0.1f);
+			for (int32 ParticleIndex = 0; ParticleIndex < NumParticles; ++ParticleIndex)
+			{
+				const FVector3f& Position = SimPositions[ParticleIndex];
+				PendingID[ParticleIndex] = ParticleIndex;
+				PendingPX[ParticleIndex] = Position.X;
+				PendingPY[ParticleIndex] = Position.Y;
+				PendingPZ[ParticleIndex] = Position.Z + Frame;
+			}
+
+			NewFrame.PendingChannelsData.Add(PositionXName, PendingPX);
+			NewFrame.PendingChannelsData.Add(PositionYName, PendingPY);
+			NewFrame.PendingChannelsData.Add(PositionZName, PendingPZ);
+
+			Cache->AddFrame_Concurrent(MoveTemp(NewFrame));
+			Progress->AdvanceCurrentScopeProgressBy(1.f / NumFrames);
 		}
 	}
+	RestoreClothComponentState();
 
 	if (!bCancelled)
 	{
-		SetResult(MakeUnique<FSkinnedMeshVertices>(OutputVertices));
+		UE_LOG(LogTemp, Display, TEXT("Data generation complete."));
+	}
+}
+
+void UClothTrainingTool::FGenerateClothOp::Simulate(float DeltaTime, int32 NumSteps)
+{
+	ClothComponent->ForceNextUpdateTeleportAndReset();
+	DataGenerationProxy->FillSimulationContext(DeltaTime);	
+	DataGenerationProxy->InitializeConfigs();
+	for (int32 Step = 0; Step < NumSteps; ++Step)
+	{
+		DataGenerationProxy->Tick();
+
+		// Clear any reset flags at the end of the first step
+		if (Step == 0 && NumSteps > 1)
+		{
+			ClothComponent->ResetTeleportMode();
+			DataGenerationProxy->FillSimulationContext(DeltaTime);	
+		}
+	}
+	DataGenerationProxy->WriteSimulationData();
+}
+
+void UClothTrainingTool::FGenerateClothOp::BackupClothComponentState()
+{
+	bIsSimulationSuspendedBackup = ClothComponent->IsSimulationSuspended();
+	bTeleportBackup = ClothComponent->NeedsTeleport();
+	bResetBackup = ClothComponent->NeedsReset();
+}
+
+void UClothTrainingTool::FGenerateClothOp::RestoreClothComponentState()
+{
+	bIsSimulationSuspendedBackup ? ClothComponent->SuspendSimulation() : ClothComponent->ResumeSimulation();
+	if (bResetBackup)
+	{
+		ClothComponent->ForceNextUpdateTeleportAndReset();
+	}
+	else if (bTeleportBackup)
+	{
+		ClothComponent->ForceNextUpdateTeleport();
+	}
+	else
+	{
+		ClothComponent->ResetTeleportMode();
 	}
 }
 
@@ -95,15 +236,23 @@ UInteractiveTool* UClothTrainingToolBuilder::BuildTool(const FToolBuilderState&
 }
 
 // ------------------- Tool -------------------
+UClothTrainingTool::UClothTrainingTool() = default;
+
+UClothTrainingTool::UClothTrainingTool(FVTableHelper& Helper)
+	: Super(Helper)
+{
+}
+
+UClothTrainingTool::~UClothTrainingTool() = default;
 
 
 void UClothTrainingTool::Setup()
 {
 	UInteractiveTool::Setup();
 
-	if (const UClothComponentToolTarget* const ClothComponentTarget = Cast<UClothComponentToolTarget>(Target))
+	if (UClothComponentToolTarget* ClothComponentTarget = Cast<UClothComponentToolTarget>(Target))
 	{
-		const UChaosClothComponent* const TargetClothComponent = ClothComponentTarget->GetClothComponent();
+		UChaosClothComponent* TargetClothComponent = ClothComponentTarget->GetClothComponent();
 		ClothComponent = TargetClothComponent;
 	}
 
@@ -117,15 +266,26 @@ void UClothTrainingTool::Setup()
 
 void UClothTrainingTool::RunTraining()
 {
-	const FText DefaultMessage(LOCTEXT("ClothTrainingMessage", "Training..."));
+	if (ClothComponent == nullptr || ToolProperties->CacheCollection == nullptr)
+	{
+		return;
+	}
+
+	static const FName CacheName = FName("SimulatedCache");
+	UChaosCacheCollection* CacheCollection = ToolProperties->CacheCollection;
+	UChaosCache* Cache = CacheCollection->FindOrAddCache(CacheName);
+
+	const FText DefaultMessage(LOCTEXT("ClothTrainingMessage", "Generate training data..."));
 	
-	using FTaskType = UE::Geometry::TModelingOpTask<FTrainClothOp>;
+	using FTaskType = UE::Geometry::TModelingOpTask<FGenerateClothOp>;
 	using FExecuterType = UE::Geometry::FAsyncTaskExecuterWithProgressCancel<FTaskType>;
 
-	TUniquePtr<FTrainClothOp> NewOp = MakeUnique<FTrainClothOp>();
-	NewOp->AnimSequence = ToolProperties->AnimationSequence;
-	// @@@@@@@@@@@ TODO: Add anything else needed by the Operator
+	if (!DataGenerationProxy.IsValid())
+	{
+		DataGenerationProxy = MakeUnique<FClothSimulationDataGenerationProxy>(*ClothComponent);
+	}
 
+	TUniquePtr<FGenerateClothOp> NewOp = MakeUnique<FGenerateClothOp>(ToolProperties->AnimationSequence, Cache, ClothComponent, DataGenerationProxy.Get());
 
 	TUniquePtr<FExecuterType> BackgroundTaskExecuter = MakeUnique<FExecuterType>(MoveTemp(NewOp));
 	BackgroundTaskExecuter->StartBackgroundTask();
@@ -134,6 +294,7 @@ void UClothTrainingTool::RunTraining()
 	SlowTask.MakeDialog(true);
 
 	bool bSuccess = false;
+	FCacheUserToken CacheUserToken = Cache->BeginRecord(ClothComponent, FGuid(), FTransform::Identity);
 	while (true)
 	{
 		if (SlowTask.ShouldCancel())
@@ -163,15 +324,28 @@ void UClothTrainingTool::RunTraining()
 		}
 	}
 
+	ClothComponent->SuspendSimulation();
+	Cache->bCompressChannels = true;
+	Cache->EndRecord(CacheUserToken);
+
 	if (bSuccess)
 	{
-		check(BackgroundTaskExecuter != nullptr && BackgroundTaskExecuter->IsDone());
-		TUniquePtr<FTrainClothOp> Op = BackgroundTaskExecuter->GetTask().ExtractOperator();
-
-		TUniquePtr<FSkinnedMeshVertices> Result = Op->ExtractResult();
-
-		// @@@@@@@@@@@ TODO: Do something with the result
+		UPackage* const Package = CacheCollection->GetPackage();
 
+		const FString SavePath = Package->GetFName().ToString();
+		UE_LOG(LogTemp, Display, TEXT("Save to %s"), *SavePath);
+		if (Package == nullptr)
+		{
+			UE_LOG(LogTemp, Error, TEXT("Failed to get package for %s"), *SavePath);
+			return;
+		}
+		FSavePackageArgs SaveArgs;
+		SaveArgs.SaveFlags = SAVE_NoError;
+		const bool bSaveSucced = UPackage::SavePackage(Package, CacheCollection, *SavePath, SaveArgs);
+		if (!bSaveSucced)
+		{
+			UE_LOG(LogTemp, Display, TEXT("Failed to save cache collection"));
+		}
 	}
 
 }
diff --git a/Engine/Plugins/Experimental/ChaosClothAssetEditor/Source/ChaosClothAssetEditorTools/Public/ChaosClothAsset/ClothTrainingTool.h b/Engine/Plugins/Experimental/ChaosClothAssetEditor/Source/ChaosClothAssetEditorTools/Public/ChaosClothAsset/ClothTrainingTool.h
--- a/Engine/Plugins/Experimental/ChaosClothAssetEditor/Source/ChaosClothAssetEditorTools/Public/ChaosClothAsset/ClothTrainingTool.h
+++ b/Engine/Plugins/Experimental/ChaosClothAssetEditor/Source/ChaosClothAssetEditorTools/Public/ChaosClothAsset/ClothTrainingTool.h
@@ -13,6 +13,7 @@
 #include "ClothTrainingTool.generated.h"
 
 class UChaosClothComponent;
+class UChaosCacheCollection;
 class UClothTrainingTool;
 
 
@@ -22,7 +23,6 @@ struct FSkinnedMeshVertices
 	TArray<FVector3f> Vertices;
 };
 
-
 UCLASS()
 class CHAOSCLOTHASSETEDITORTOOLS_API UClothTrainingToolProperties : public UObject
 {
@@ -32,8 +32,8 @@ public:
 	UPROPERTY(EditAnywhere, Category = Files)
 	TObjectPtr<UAnimSequence> AnimationSequence;
 
-	UPROPERTY(EditAnywhere, Category = Files)
-	FText OutputBufferLocation;
+	UPROPERTY(EditAnywhere, Category = Files, meta = (DisplayName = "Generated Cache"))
+	TObjectPtr<UChaosCacheCollection> CacheCollection;
 };
 
 UENUM()
@@ -86,12 +86,16 @@ class CHAOSCLOTHASSETEDITORTOOLS_API UClothTrainingTool : public USingleSelectio
 {
 	GENERATED_BODY()
 public:
-
+	UClothTrainingTool();
+	UClothTrainingTool(FVTableHelper& Helper);
+	~UClothTrainingTool();
 	// UInteractiveTool
 	virtual void Setup() override;
 	virtual void OnTick(float DeltaTime) override;
 
 private:
+	class FClothSimulationDataGenerationProxy;
+	class FGenerateClothOp;
 
 	friend class UClothTrainingToolActionProperties;
 
@@ -104,7 +108,9 @@ private:
 	EClothTrainingToolActions PendingAction = EClothTrainingToolActions::NoAction;
 
 	UPROPERTY()
-	TObjectPtr<const UChaosClothComponent> ClothComponent;
+	TObjectPtr<UChaosClothComponent> ClothComponent;
+
+	TUniquePtr<FClothSimulationDataGenerationProxy> DataGenerationProxy;
 
 	void RequestAction(EClothTrainingToolActions ActionType);
 
