Create a new module for Cloth Asset-related Dataflow nodes.

This should make it possible for the Cloth Editor Tools to create/modify nodes in the Dataflow graph (e.g. the Paint tool should be able to create a "AddMap" node and add it to the graph.)

diff --git a/Engine/Plugins/Experimental/ChaosClothAssetEditor/ChaosClothAssetEditor.uplugin b/Engine/Plugins/Experimental/ChaosClothAssetEditor/ChaosClothAssetEditor.uplugin
--- a/Engine/Plugins/Experimental/ChaosClothAssetEditor/ChaosClothAssetEditor.uplugin
+++ b/Engine/Plugins/Experimental/ChaosClothAssetEditor/ChaosClothAssetEditor.uplugin
@@ -26,6 +26,12 @@
 			"LoadingPhase": "Default",
 			"PlatformAllowList": [ "Win64" ]
 		},
+		{
+			"Name": "ChaosClothAssetDataflowNodes",
+			"Type": "Editor",
+			"LoadingPhase": "Default",
+			"PlatformAllowList": [ "Win64" ]
+		},
 		{
 			"Name": "ChaosClothAssetTools",
 			"Type": "UncookedOnly",
diff --git a/Engine/Plugins/Experimental/ChaosClothAssetEditor/Source/ChaosClothAssetDataflowNodes/ChaosClothAssetDataflowNodes.Build.cs b/Engine/Plugins/Experimental/ChaosClothAssetEditor/Source/ChaosClothAssetDataflowNodes/ChaosClothAssetDataflowNodes.Build.cs
new file mode 100644
--- /dev/null
+++ b/Engine/Plugins/Experimental/ChaosClothAssetEditor/Source/ChaosClothAssetDataflowNodes/ChaosClothAssetDataflowNodes.Build.cs
@@ -0,0 +1,38 @@
+// Copyright Epic Games, Inc. All Rights Reserved.
+
+using UnrealBuildTool;
+
+public class ChaosClothAssetDataflowNodes : ModuleRules
+{
+	public ChaosClothAssetDataflowNodes(ReadOnlyTargetRules Target) : base(Target)
+	{	
+		PublicDependencyModuleNames.AddRange(
+			new string[]
+			{
+				"Core",
+				// ... add other public dependencies that you statically link with here ...
+			}
+			);	
+		
+		PrivateDependencyModuleNames.AddRange(
+			new string[]
+			{
+				"Chaos",
+				"ChaosCloth",
+				"ChaosClothAsset",
+				"ChaosClothAssetEngine",
+				"CoreUObject",
+				"DataflowCore",
+				"DataflowEditor",
+				"DataflowEngine",
+				"DatasmithImporter",
+				"Engine",
+				"ExternalSource",
+				"Slate",
+				"GeometryCore",
+				"MeshConversion"
+				// ... add private dependencies that you statically link with here ...	
+			}
+			);
+	}
+}
diff --git a/Engine/Plugins/Experimental/ChaosClothAssetEditor/Source/ChaosClothAssetDataflowNodes/Private/ChaosClothAsset/BindToRootBoneNode.cpp b/Engine/Plugins/Experimental/ChaosClothAssetEditor/Source/ChaosClothAssetDataflowNodes/Private/ChaosClothAsset/BindToRootBoneNode.cpp
new file mode 100644
--- /dev/null
+++ b/Engine/Plugins/Experimental/ChaosClothAssetEditor/Source/ChaosClothAssetDataflowNodes/Private/ChaosClothAsset/BindToRootBoneNode.cpp
@@ -0,0 +1,34 @@
+// Copyright Epic Games, Inc. All Rights Reserved.
+
+#include "ChaosClothAsset/BindToRootBoneNode.h"
+#include "ChaosClothAsset/DataflowNodes.h"
+#include "ChaosClothAsset/ClothGeometryTools.h"
+
+#include UE_INLINE_GENERATED_CPP_BY_NAME(BindToRootBoneNode)
+
+#define LOCTEXT_NAMESPACE "ChaosClothAssetBindToRootBoneNode"
+
+FChaosClothAssetBindToRootBoneNode::FChaosClothAssetBindToRootBoneNode(const Dataflow::FNodeParameters& InParam, FGuid InGuid)
+	: FDataflowNode(InParam, InGuid)
+{
+	RegisterInputConnection(&Collection);
+	RegisterOutputConnection(&Collection, &Collection);
+}
+
+void FChaosClothAssetBindToRootBoneNode::Evaluate(Dataflow::FContext& Context, const FDataflowOutput* Out) const
+{
+	if (Out->IsA<FManagedArrayCollection>(&Collection))
+	{
+		using namespace UE::Chaos::ClothAsset;
+
+		// Evaluate in collection
+		FManagedArrayCollection InCollection = GetValue<FManagedArrayCollection>(Context, &Collection);
+		const TSharedRef<FManagedArrayCollection> ClothCollection = MakeShared<FManagedArrayCollection>(MoveTemp(InCollection));
+
+		FClothGeometryTools::BindMeshToRootBone(ClothCollection, bBindSimMesh, bBindRenderMesh);
+
+		SetValue<FManagedArrayCollection>(Context, *ClothCollection, &Collection);
+	}
+}
+
+#undef LOCTEXT_NAMESPACE
diff --git a/Engine/Plugins/Experimental/ChaosClothAssetEditor/Source/ChaosClothAssetDataflowNodes/Private/ChaosClothAsset/ChaosClothAssetDataflowNodesModule.cpp b/Engine/Plugins/Experimental/ChaosClothAssetEditor/Source/ChaosClothAssetDataflowNodes/Private/ChaosClothAsset/ChaosClothAssetDataflowNodesModule.cpp
new file mode 100644
--- /dev/null
+++ b/Engine/Plugins/Experimental/ChaosClothAssetEditor/Source/ChaosClothAssetDataflowNodes/Private/ChaosClothAsset/ChaosClothAssetDataflowNodesModule.cpp
@@ -0,0 +1,20 @@
+// Copyright Epic Games, Inc. All Rights Reserved.
+
+#include "ChaosClothAsset/ChaosClothAssetDataflowNodesModule.h"
+#include "ChaosClothAsset/DataflowNodes.h"
+#include "Modules/ModuleManager.h"
+
+#define LOCTEXT_NAMESPACE "FChaosClothAssetDataflowNodesModule"
+
+void FChaosClothAssetDataflowNodesModule::StartupModule()
+{
+	UE::Chaos::ClothAsset::DataflowNodes::Register();
+}
+
+void FChaosClothAssetDataflowNodesModule::ShutdownModule()
+{
+}
+
+#undef LOCTEXT_NAMESPACE
+	
+IMPLEMENT_MODULE(FChaosClothAssetDataflowNodesModule, ChaosClothAssetDataflowNodes)
diff --git a/Engine/Plugins/Experimental/ChaosClothAssetEditor/Source/ChaosClothAssetDataflowNodes/Private/ChaosClothAsset/ChaosClothAssetDataflowNodesModule.h b/Engine/Plugins/Experimental/ChaosClothAssetEditor/Source/ChaosClothAssetDataflowNodes/Private/ChaosClothAsset/ChaosClothAssetDataflowNodesModule.h
new file mode 100644
--- /dev/null
+++ b/Engine/Plugins/Experimental/ChaosClothAssetEditor/Source/ChaosClothAssetDataflowNodes/Private/ChaosClothAsset/ChaosClothAssetDataflowNodesModule.h
@@ -0,0 +1,16 @@
+// Copyright Epic Games, Inc. All Rights Reserved.
+
+#pragma once
+
+#include "CoreMinimal.h"
+#include "Modules/ModuleInterface.h"
+
+class CHAOSCLOTHASSETDATAFLOWNODES_API FChaosClothAssetDataflowNodesModule : public IModuleInterface
+{
+public:
+
+	/** IModuleInterface implementation */
+	virtual void StartupModule() override;
+	virtual void ShutdownModule() override;
+
+};
diff --git a/Engine/Plugins/Experimental/ChaosClothAssetEditor/Source/ChaosClothAssetDataflowNodes/Private/ChaosClothAsset/CopySimulationToRenderMeshNode.cpp b/Engine/Plugins/Experimental/ChaosClothAssetEditor/Source/ChaosClothAssetDataflowNodes/Private/ChaosClothAsset/CopySimulationToRenderMeshNode.cpp
new file mode 100644
--- /dev/null
+++ b/Engine/Plugins/Experimental/ChaosClothAssetEditor/Source/ChaosClothAssetDataflowNodes/Private/ChaosClothAsset/CopySimulationToRenderMeshNode.cpp
@@ -0,0 +1,45 @@
+// Copyright Epic Games, Inc. All Rights Reserved.
+
+#include "ChaosClothAsset/CopySimulationToRenderMeshNode.h"
+#include "ChaosClothAsset/DataflowNodes.h"
+#include "ChaosClothAsset/ClothGeometryTools.h"
+#include "Materials/Material.h"
+
+#include UE_INLINE_GENERATED_CPP_BY_NAME(CopySimulationToRenderMeshNode)
+
+#define LOCTEXT_NAMESPACE "ChaosClothAssetCopySimulationToRenderMeshNode"
+
+FChaosClothAssetCopySimulationToRenderMeshNode::FChaosClothAssetCopySimulationToRenderMeshNode(const Dataflow::FNodeParameters& InParam, FGuid InGuid)
+	: FDataflowNode(InParam, InGuid)
+{
+	RegisterInputConnection(&Collection);
+	RegisterInputConnection(&Patterns);
+	RegisterOutputConnection(&Collection, &Collection);
+}
+
+void FChaosClothAssetCopySimulationToRenderMeshNode::Evaluate(Dataflow::FContext& Context, const FDataflowOutput* Out) const
+{
+	if (Out->IsA<FManagedArrayCollection>(&Collection))
+	{
+		using namespace UE::Chaos::ClothAsset;
+
+		// Evaluate in collection
+		FManagedArrayCollection InCollection = GetValue<FManagedArrayCollection>(Context, &Collection);
+		const TSharedRef<FManagedArrayCollection> ClothCollection = MakeShared<FManagedArrayCollection>(MoveTemp(InCollection));
+
+		// Empty mesh and materials
+		FClothGeometryTools::DeleteRenderMesh(ClothCollection);
+
+		// Find the material asset path
+		const FString MaterialPathName = Material ? 
+			Material->GetPathName() :
+			FString(TEXT("/Engine/EditorMaterials/Cloth/CameraLitDoubleSided.CameraLitDoubleSided"));
+
+		// Copy the mesh data
+		FClothGeometryTools::CopySimMeshToRenderMesh(ClothCollection, MaterialPathName);
+
+		SetValue<FManagedArrayCollection>(Context, *ClothCollection, &Collection);
+	}
+}
+
+#undef LOCTEXT_NAMESPACE
diff --git a/Engine/Plugins/Experimental/ChaosClothAssetEditor/Source/ChaosClothAssetDataflowNodes/Private/ChaosClothAsset/DataflowNodes.cpp b/Engine/Plugins/Experimental/ChaosClothAssetEditor/Source/ChaosClothAssetDataflowNodes/Private/ChaosClothAsset/DataflowNodes.cpp
new file mode 100644
--- /dev/null
+++ b/Engine/Plugins/Experimental/ChaosClothAssetEditor/Source/ChaosClothAssetDataflowNodes/Private/ChaosClothAsset/DataflowNodes.cpp
@@ -0,0 +1,46 @@
+// Copyright Epic Games, Inc. All Rights Reserved.
+
+#include "ChaosClothAsset/DataflowNodes.h"
+#include "ChaosClothAsset/TerminalNode.h"
+#include "ChaosClothAsset/DatasmithImportNode.h"
+#include "ChaosClothAsset/ImportNode.h"
+#include "ChaosClothAsset/StaticMeshImportNode.h"
+#include "ChaosClothAsset/PatternSelectionNode.h"
+#include "ChaosClothAsset/DeleteRenderMeshNode.h"
+#include "ChaosClothAsset/CopySimulationToRenderMeshNode.h"
+#include "ChaosClothAsset/ReverseNormalsNode.h"
+#include "ChaosClothAsset/SimulationDefaultConfigNode.h"
+#include "ChaosClothAsset/SimulationXPBDConfigNode.h"
+#include "ChaosClothAsset/BindToRootBoneNode.h"
+#include "Dataflow/DataflowNodeFactory.h"
+#include "Framework/Notifications/NotificationManager.h"
+#include "Widgets/Notifications/SNotificationList.h"
+
+DEFINE_LOG_CATEGORY(LogChaosClothAssetDataflowNodes);
+
+namespace UE::Chaos::ClothAsset::DataflowNodes
+{
+	void Register()
+	{
+		DATAFLOW_NODE_REGISTER_CREATION_FACTORY(FChaosClothAssetTerminalNode);
+		DATAFLOW_NODE_REGISTER_CREATION_FACTORY(FChaosClothAssetDatasmithImportNode);
+		DATAFLOW_NODE_REGISTER_CREATION_FACTORY(FChaosClothAssetImportNode);
+		DATAFLOW_NODE_REGISTER_CREATION_FACTORY(FChaosClothAssetStaticMeshImportNode);
+		DATAFLOW_NODE_REGISTER_CREATION_FACTORY(FChaosClothAssetPatternSelectionNode);
+		DATAFLOW_NODE_REGISTER_CREATION_FACTORY(FChaosClothAssetDeleteRenderMeshNode);
+		DATAFLOW_NODE_REGISTER_CREATION_FACTORY(FChaosClothAssetCopySimulationToRenderMeshNode);
+		DATAFLOW_NODE_REGISTER_CREATION_FACTORY(FChaosClothAssetReverseNormalsNode);
+		DATAFLOW_NODE_REGISTER_CREATION_FACTORY(FChaosClothAssetSimulationDefaultConfigNode);
+		DATAFLOW_NODE_REGISTER_CREATION_FACTORY(FChaosClothAssetSimulationXPBDConfigNode);
+		DATAFLOW_NODE_REGISTER_CREATION_FACTORY(FChaosClothAssetBindToRootBoneNode);
+	}
+
+	void LogAndToastWarning(const FText& Error)
+	{
+		FNotificationInfo Info(Error);
+		Info.ExpireDuration = 5.0f;
+		FSlateNotificationManager::Get().AddNotification(Info);
+
+		UE_LOG(LogChaosClothAssetDataflowNodes, Warning, TEXT("%s"), *Error.ToString());
+	}
+}
diff --git a/Engine/Plugins/Experimental/ChaosClothAssetEditor/Source/ChaosClothAssetDataflowNodes/Private/ChaosClothAsset/DatasmithImportNode.cpp b/Engine/Plugins/Experimental/ChaosClothAssetEditor/Source/ChaosClothAssetDataflowNodes/Private/ChaosClothAsset/DatasmithImportNode.cpp
new file mode 100644
--- /dev/null
+++ b/Engine/Plugins/Experimental/ChaosClothAssetEditor/Source/ChaosClothAssetDataflowNodes/Private/ChaosClothAsset/DatasmithImportNode.cpp
@@ -0,0 +1,113 @@
+// Copyright Epic Games, Inc. All Rights Reserved.
+
+#include "ChaosClothAsset/DatasmithImportNode.h"
+#include "ChaosClothAsset/DataflowNodes.h"
+#include "ChaosClothAsset/ClothAsset.h"
+#include "ChaosClothAsset/CollectionClothFacade.h"
+#include "DatasmithImportContext.h"
+#include "DatasmithImportFactory.h"
+#include "ExternalSource.h"
+#include "ExternalSourceModule.h"
+#include "SourceUri.h"
+
+#include UE_INLINE_GENERATED_CPP_BY_NAME(DatasmithImportNode)
+
+#define LOCTEXT_NAMESPACE "ChaosClothAssetDatasmithImportNode"
+
+FChaosClothAssetDatasmithImportNode::FChaosClothAssetDatasmithImportNode(const Dataflow::FNodeParameters& InParam, FGuid InGuid)
+	: FDataflowNode(InParam, InGuid)
+{
+	RegisterInputConnection(&DatasmithFile);
+	RegisterOutputConnection(&Collection);
+}
+
+bool FChaosClothAssetDatasmithImportNode::EvaluateImpl(Dataflow::FContext& Context, FManagedArrayCollection& OutCollection) const
+{
+	using namespace UE::DatasmithImporter;
+
+	const FFilePath& InFilePath = GetValue<FFilePath>(Context, &DatasmithFile);
+
+	const FSourceUri SourceUri = FSourceUri::FromFilePath(InFilePath.FilePath);
+	const TSharedPtr<FExternalSource> ExternalSource = IExternalSourceModule::GetOrCreateExternalSource(SourceUri);
+	if (!ExternalSource)
+	{
+		return false;
+	}
+
+	constexpr bool bLoadConfig = false; // TODO: Not sure what this does
+	const FName LoggerName(TEXT("ImportDatasmithClothNode"));
+	const FText LoggerLabel(NSLOCTEXT("ImportDatasmithClothNode", "LoggerLabel", "ImportDatasmithClothNode"));
+	FDatasmithImportContext DatasmithImportContext(ExternalSource, bLoadConfig, LoggerName, LoggerLabel);
+
+	const FName PackageName = MakeUniqueObjectName(nullptr, UPackage::StaticClass(), FName(FPaths::Combine(GetTransientPackage()->GetPathName(), ExternalSource->GetSourceName())));
+	const TStrongObjectPtr<UPackage> DestinationPackage(CreatePackage(*PackageName.ToString()));
+	if (!ensure(DestinationPackage))
+	{
+		// Failed to create the package to hold this asset for some reason
+		return false;
+	}
+
+	// Don't create the Actors in the level, just read the Assets
+	DatasmithImportContext.Options->BaseOptions.SceneHandling = EDatasmithImportScene::AssetsOnly;
+
+	constexpr EObjectFlags NewObjectFlags = RF_Public | RF_Transactional | RF_Transient | RF_Standalone;
+	const TSharedPtr<FJsonObject> ImportSettingsJson;
+	constexpr bool bIsSilent = true;
+
+	if (!DatasmithImportContext.Init(DestinationPackage->GetPathName(), NewObjectFlags, GWarn, ImportSettingsJson, bIsSilent))
+	{
+		return false;
+	}
+
+	if (const TSharedPtr<IDatasmithScene> LoadedScene = ExternalSource->TryLoad())
+	{
+		DatasmithImportContext.InitScene(LoadedScene.ToSharedRef());
+	}
+	else
+	{
+		return false;
+	}
+
+	bool bUserCancelled = false;
+	bool bImportSucceed = DatasmithImportFactoryImpl::ImportDatasmithScene(DatasmithImportContext, bUserCancelled);
+	bImportSucceed &= !bUserCancelled;
+
+	if (bImportSucceed && DatasmithImportContext.ImportedClothes.Num() > 0)
+	{
+		UObject* const ClothObject = DatasmithImportContext.ImportedClothes.CreateIterator()->Value;
+		UChaosClothAsset* const DatasmithClothAsset = Cast<UChaosClothAsset>(ClothObject);
+		if (ensure(DatasmithClothAsset))
+		{
+			DatasmithClothAsset->GetClothCollection()->CopyTo(&OutCollection);
+			return true;
+		}
+		DatasmithClothAsset->ClearFlags(RF_Standalone);
+	}
+
+	return false;
+}
+
+void FChaosClothAssetDatasmithImportNode::Evaluate(Dataflow::FContext& Context, const FDataflowOutput* Out) const
+{
+	FManagedArrayCollection OutCollection;
+	const bool bSuccess = EvaluateImpl(Context, OutCollection);
+
+	if (bSuccess)
+	{
+		SetValue<FManagedArrayCollection>(Context, OutCollection, &Collection);
+	}
+	else
+	{
+		using namespace UE::Chaos::ClothAsset;
+
+		// Init with an empty cloth collection
+		const TSharedRef<FManagedArrayCollection> ClothCollection = MakeShared<FManagedArrayCollection>();
+		FCollectionClothFacade CollectionClothFacade(ClothCollection);
+		CollectionClothFacade.DefineSchema();
+		CollectionClothFacade.AddLod();
+
+		SetValue<FManagedArrayCollection>(Context, *ClothCollection, &Collection);
+	}
+}
+
+#undef LOCTEXT_NAMESPACE
diff --git a/Engine/Plugins/Experimental/ChaosClothAssetEditor/Source/ChaosClothAssetDataflowNodes/Private/ChaosClothAsset/DeleteRenderMeshNode.cpp b/Engine/Plugins/Experimental/ChaosClothAssetEditor/Source/ChaosClothAssetDataflowNodes/Private/ChaosClothAsset/DeleteRenderMeshNode.cpp
new file mode 100644
--- /dev/null
+++ b/Engine/Plugins/Experimental/ChaosClothAssetEditor/Source/ChaosClothAssetDataflowNodes/Private/ChaosClothAsset/DeleteRenderMeshNode.cpp
@@ -0,0 +1,42 @@
+// Copyright Epic Games, Inc. All Rights Reserved.
+
+#include "ChaosClothAsset/DeleteRenderMeshNode.h"
+#include "ChaosClothAsset/DataflowNodes.h"
+#include "ChaosClothAsset/ClothGeometryTools.h"
+
+#include UE_INLINE_GENERATED_CPP_BY_NAME(DeleteRenderMeshNode)
+
+#define LOCTEXT_NAMESPACE "ChaosClothAssetDeleteRenderMeshNode"
+
+FChaosClothAssetDeleteRenderMeshNode::FChaosClothAssetDeleteRenderMeshNode(const Dataflow::FNodeParameters& InParam, FGuid InGuid)
+	: FDataflowNode(InParam, InGuid)
+{
+	RegisterInputConnection(&Collection);
+	RegisterInputConnection(&Patterns);
+	RegisterOutputConnection(&Collection, &Collection);
+}
+
+void FChaosClothAssetDeleteRenderMeshNode::Evaluate(Dataflow::FContext& Context, const FDataflowOutput* Out) const
+{
+	if (Out->IsA<FManagedArrayCollection>(&Collection))
+	{
+		using namespace UE::Chaos::ClothAsset;
+
+		// Evaluate in collection
+		FManagedArrayCollection InCollection = GetValue<FManagedArrayCollection>(Context, &Collection);
+		const TSharedRef<FManagedArrayCollection> ClothCollection = MakeShared<FManagedArrayCollection>(MoveTemp(InCollection));
+
+		if (Patterns.IsEmpty())
+		{
+			FClothGeometryTools::DeleteRenderMesh(ClothCollection);
+		}
+		else
+		{
+			DataflowNodes::LogAndToastWarning(LOCTEXT("DeletePerPatternNotImplemented", "FClothAssetDeleteRenderMeshNode: Delete per patterns not implemented."));
+		}
+
+		SetValue<FManagedArrayCollection>(Context, *ClothCollection, &Collection);
+	}
+}
+
+#undef LOCTEXT_NAMESPACE
diff --git a/Engine/Plugins/Experimental/ChaosClothAssetEditor/Source/ChaosClothAssetDataflowNodes/Private/ChaosClothAsset/ImportNode.cpp b/Engine/Plugins/Experimental/ChaosClothAssetEditor/Source/ChaosClothAssetDataflowNodes/Private/ChaosClothAsset/ImportNode.cpp
new file mode 100644
--- /dev/null
+++ b/Engine/Plugins/Experimental/ChaosClothAssetEditor/Source/ChaosClothAssetDataflowNodes/Private/ChaosClothAsset/ImportNode.cpp
@@ -0,0 +1,84 @@
+// Copyright Epic Games, Inc. All Rights Reserved.
+
+#include "ChaosClothAsset/ImportNode.h"
+#include "ChaosClothAsset/DataflowNodes.h"
+#include "ChaosClothAsset/ClothAsset.h"
+#include "ChaosClothAsset/CollectionClothFacade.h"
+#include "Engine/SkinnedAssetCommon.h"
+#include "Materials/Material.h"
+#include "PhysicsEngine/PhysicsAsset.h"
+
+#include UE_INLINE_GENERATED_CPP_BY_NAME(ImportNode)
+
+#define LOCTEXT_NAMESPACE "ChaosClothAssetImportNode"
+
+FChaosClothAssetImportNode::FChaosClothAssetImportNode(const Dataflow::FNodeParameters& InParam, FGuid InGuid)
+	: FDataflowNode(InParam, InGuid)
+{
+	RegisterInputConnection(&ClothAsset);
+	RegisterInputConnection(&ImportLod);
+	RegisterOutputConnection(&Collection);
+}
+
+void FChaosClothAssetImportNode::Evaluate(Dataflow::FContext& Context, const FDataflowOutput* Out) const
+{
+	using namespace UE::Chaos::ClothAsset;
+	using namespace UE::Chaos::ClothAsset::DataflowNodes;
+
+	if (Out->IsA<FManagedArrayCollection>(&Collection))
+	{
+		bool bSetValue = false;
+		if (ClothAsset)
+		{
+			// Copy the main cloth asset details to this dataflow's owner if any
+			if (const Dataflow::FEngineContext* EngineContext = Context.AsType<Dataflow::FEngineContext>())
+			{
+				if (const UChaosClothAsset* const OwnerClothAsset = Cast<UChaosClothAsset>(EngineContext->Owner))
+				{
+					if (OwnerClothAsset == ClothAsset)
+					{
+						// Can't create a loop
+						bSetValue = false;
+						LogAndToastWarning(LOCTEXT("RecursiveAssetLoop", "FClothAssetNode: The source asset cannot be the terminal asset."));
+					}
+					else
+					{
+						bSetValue = true;
+					}
+				}
+			}
+			else
+			{
+				// No terminal asset, this is a stray dataflow and it is safe to set the value without fear of a loop
+				bSetValue = true;
+			}
+		}
+
+		// Create a new cloth collection with its LOD 0
+		const TSharedRef<FManagedArrayCollection> ClothCollection = MakeShared<FManagedArrayCollection>();
+		FCollectionClothFacade ClothFacade(ClothCollection);
+		ClothFacade.DefineSchema();
+		FCollectionClothLodFacade ClothLodFacade = ClothFacade.AddGetLod();  // Always output a single LOD at LOD 0, even if empty
+
+		// Copy the cloth asset to this node's output collection
+		if (bSetValue)
+		{
+			const TObjectPtr<const UChaosClothAsset>& InClothAsset = GetValue<TObjectPtr<const UChaosClothAsset>>(Context, &ClothAsset);
+			if (InClothAsset->GetClothCollection())
+			{
+				const int32& InImportLod = GetValue<int32>(Context, &ImportLod);
+				const FCollectionClothConstFacade InClothFacade(InClothAsset->GetClothCollection());
+
+				if (InImportLod >= 0 && InImportLod < InClothFacade.GetNumLods())
+				{
+					// Copy specified LOD from the input asset to the LOD 0 of the output collection
+					const FCollectionClothLodConstFacade InClothLodFacade = InClothFacade.GetLod(InImportLod);
+					ClothLodFacade.Initialize(InClothLodFacade);
+				}
+			}
+		}
+		SetValue<FManagedArrayCollection>(Context, *ClothCollection, &Collection);
+}
+}
+
+#undef LOCTEXT_NAMESPACE
diff --git a/Engine/Plugins/Experimental/ChaosClothAssetEditor/Source/ChaosClothAssetDataflowNodes/Private/ChaosClothAsset/PatternSelectionNode.cpp b/Engine/Plugins/Experimental/ChaosClothAssetEditor/Source/ChaosClothAssetDataflowNodes/Private/ChaosClothAsset/PatternSelectionNode.cpp
new file mode 100644
--- /dev/null
+++ b/Engine/Plugins/Experimental/ChaosClothAssetEditor/Source/ChaosClothAssetDataflowNodes/Private/ChaosClothAsset/PatternSelectionNode.cpp
@@ -0,0 +1,26 @@
+// Copyright Epic Games, Inc. All Rights Reserved.
+
+#include "ChaosClothAsset/PatternSelectionNode.h"
+#include "ChaosClothAsset/DataflowNodes.h"
+
+#include UE_INLINE_GENERATED_CPP_BY_NAME(PatternSelectionNode)
+
+#define LOCTEXT_NAMESPACE "ChaosClothAssetPatternSelectionNode"
+
+FChaosClothAssetPatternSelectionNode::FChaosClothAssetPatternSelectionNode(const Dataflow::FNodeParameters& InParam, FGuid InGuid)
+	: FDataflowNode(InParam, InGuid)
+{
+	RegisterOutputConnection(&Patterns);
+}
+
+void FChaosClothAssetPatternSelectionNode::Evaluate(Dataflow::FContext& Context, const FDataflowOutput* Out) const
+{
+	if (Out->IsA<TArray<int32>>(&Patterns))
+	{
+		// Init with empty pattern indices
+		const TArray<int32> Selection;
+		SetValue<TArray<int32>>(Context, Selection, &Patterns);
+	}
+}
+
+#undef LOCTEXT_NAMESPACE
diff --git a/Engine/Plugins/Experimental/ChaosClothAssetEditor/Source/ChaosClothAssetDataflowNodes/Private/ChaosClothAsset/ReverseNormalsNode.cpp b/Engine/Plugins/Experimental/ChaosClothAssetEditor/Source/ChaosClothAssetDataflowNodes/Private/ChaosClothAsset/ReverseNormalsNode.cpp
new file mode 100644
--- /dev/null
+++ b/Engine/Plugins/Experimental/ChaosClothAssetEditor/Source/ChaosClothAssetDataflowNodes/Private/ChaosClothAsset/ReverseNormalsNode.cpp
@@ -0,0 +1,35 @@
+// Copyright Epic Games, Inc. All Rights Reserved.
+
+#include "ChaosClothAsset/ReverseNormalsNode.h"
+#include "ChaosClothAsset/DataflowNodes.h"
+#include "ChaosClothAsset/ClothGeometryTools.h"
+
+#include UE_INLINE_GENERATED_CPP_BY_NAME(ReverseNormalsNode)
+
+#define LOCTEXT_NAMESPACE "ChaosClothAssetReverseNormalsNode"
+
+FChaosClothAssetReverseNormalsNode::FChaosClothAssetReverseNormalsNode(const Dataflow::FNodeParameters& InParam, FGuid InGuid)
+	: FDataflowNode(InParam, InGuid)
+{
+	RegisterInputConnection(&Collection);
+	RegisterInputConnection(&Patterns);
+	RegisterOutputConnection(&Collection, &Collection);
+}
+
+void FChaosClothAssetReverseNormalsNode::Evaluate(Dataflow::FContext& Context, const FDataflowOutput* Out) const
+{
+	if (Out->IsA<FManagedArrayCollection>(&Collection))
+	{
+		using namespace UE::Chaos::ClothAsset;
+
+		// Evaluate in collection
+		FManagedArrayCollection InCollection = GetValue<FManagedArrayCollection>(Context, &Collection);
+		const TSharedRef<FManagedArrayCollection> ClothCollection = MakeShared<FManagedArrayCollection>(MoveTemp(InCollection));
+
+		FClothGeometryTools::ReverseNormals(ClothCollection, bReverseSimMeshNormals, bReverseRenderMeshNormals, Patterns);
+
+		SetValue<FManagedArrayCollection>(Context, *ClothCollection, &Collection);
+	}
+}
+
+#undef LOCTEXT_NAMESPACE
diff --git a/Engine/Plugins/Experimental/ChaosClothAssetEditor/Source/ChaosClothAssetDataflowNodes/Private/ChaosClothAsset/SimulationDefaultConfigNode.cpp b/Engine/Plugins/Experimental/ChaosClothAssetEditor/Source/ChaosClothAssetDataflowNodes/Private/ChaosClothAsset/SimulationDefaultConfigNode.cpp
new file mode 100644
--- /dev/null
+++ b/Engine/Plugins/Experimental/ChaosClothAssetEditor/Source/ChaosClothAssetDataflowNodes/Private/ChaosClothAsset/SimulationDefaultConfigNode.cpp
@@ -0,0 +1,66 @@
+// Copyright Epic Games, Inc. All Rights Reserved.
+
+#include "ChaosClothAsset/SimulationDefaultConfigNode.h"
+#include "ChaosClothAsset/DataflowNodes.h"
+#include "ChaosCloth/ChaosClothConfig.h"
+#include "ChaosCloth/ChaosClothingSimulationConfig.h"
+#include "Chaos/CollectionPropertyFacade.h"
+
+#include UE_INLINE_GENERATED_CPP_BY_NAME(SimulationDefaultConfigNode)
+
+#define LOCTEXT_NAMESPACE "ChaosClothAssetSimulationDefaultConfigNode"
+
+FChaosClothAssetSimulationDefaultConfigNode::FChaosClothAssetSimulationDefaultConfigNode(const Dataflow::FNodeParameters& InParam, FGuid InGuid)
+	: FDataflowNode(InParam, InGuid)
+	, SimulationConfig(NewObject<UChaosClothConfig>())
+	, SharedSimulationConfig(NewObject<UChaosClothSharedSimConfig>())
+{
+	RegisterInputConnection(&Collection);
+	RegisterOutputConnection(&Collection, &Collection);
+}
+
+void FChaosClothAssetSimulationDefaultConfigNode::Serialize(FArchive& Ar)
+{
+	if (!SimulationConfig)
+	{
+		SimulationConfig = NewObject<UChaosClothConfig>();
+	}
+	SimulationConfig->Serialize(Ar);
+
+	if (!SharedSimulationConfig)
+	{
+		SharedSimulationConfig = NewObject<UChaosClothSharedSimConfig>();
+	}
+	SharedSimulationConfig->Serialize(Ar);
+}
+
+void FChaosClothAssetSimulationDefaultConfigNode::Evaluate(Dataflow::FContext& Context, const FDataflowOutput* Out) const
+{
+	using namespace Chaos;
+	using namespace Chaos::Softs;
+
+	if (Out->IsA<FManagedArrayCollection>(&Collection))
+	{
+		FManagedArrayCollection InCollection = GetValue<FManagedArrayCollection>(Context, &Collection);
+
+		FClothingSimulationConfig ClothingSimulationConfig;
+		ClothingSimulationConfig.Initialize(SimulationConfig.Get(), SharedSimulationConfig.Get());
+
+		ClothingSimulationConfig.GetPropertyCollection()->CopyTo(&InCollection);
+
+		SetValue<FManagedArrayCollection>(Context, InCollection, &Collection);
+	}
+}
+
+void FChaosClothAssetSimulationDefaultConfigNode::AddReferencedObjects(FReferenceCollector& Collector)
+{
+	Collector.AddReferencedObject(SimulationConfig);
+	Collector.AddReferencedObject(SharedSimulationConfig);
+}
+
+FString FChaosClothAssetSimulationDefaultConfigNode::GetReferencerName() const
+{
+	return TEXT("FChaosClothAssetSimulationDefaultConfigNode");
+}
+
+#undef LOCTEXT_NAMESPACE
diff --git a/Engine/Plugins/Experimental/ChaosClothAssetEditor/Source/ChaosClothAssetDataflowNodes/Private/ChaosClothAsset/SimulationXPBDConfigNode.cpp b/Engine/Plugins/Experimental/ChaosClothAssetEditor/Source/ChaosClothAssetDataflowNodes/Private/ChaosClothAsset/SimulationXPBDConfigNode.cpp
new file mode 100644
--- /dev/null
+++ b/Engine/Plugins/Experimental/ChaosClothAssetEditor/Source/ChaosClothAssetDataflowNodes/Private/ChaosClothAsset/SimulationXPBDConfigNode.cpp
@@ -0,0 +1,185 @@
+// Copyright Epic Games, Inc. All Rights Reserved.
+
+#include "ChaosClothAsset/SimulationXPBDConfigNode.h"
+#include "ChaosClothAsset/DataflowNodes.h"
+#include "ChaosCloth/ChaosClothConfig.h"
+#include "ChaosCloth/ChaosClothingSimulationConfig.h"
+#include "Chaos/CollectionPropertyFacade.h"
+
+#include UE_INLINE_GENERATED_CPP_BY_NAME(SimulationXPBDConfigNode)
+
+#define LOCTEXT_NAMESPACE "ChaosClothAssetSimulationXPBDConfigNode"
+
+FChaosClothAssetSimulationXPBDConfigNode::FChaosClothAssetSimulationXPBDConfigNode(const Dataflow::FNodeParameters& InParam, FGuid InGuid)
+	: FDataflowNode(InParam, InGuid)
+	, SharedSimulationConfig(NewObject<UChaosClothSharedSimConfig>())
+{
+	RegisterInputConnection(&Collection);
+
+	RegisterInputConnection(&bEnableXPBDStretchConstraints);
+	RegisterInputConnection(&StretchStiffness);
+	RegisterInputConnection(&StretchDampingRatio);
+
+	RegisterInputConnection(&bEnableXPBDBendConstraints);
+	RegisterInputConnection(&bEnableBendAnisotropy);
+	RegisterInputConnection(&BendingStiffness);
+	RegisterInputConnection(&BendingStiffnessWeft);
+	RegisterInputConnection(&BendingStiffnessBias);
+	RegisterInputConnection(&BendingDampingRatio);
+	RegisterInputConnection(&BucklingRatio);
+	RegisterInputConnection(&BucklingStiffness);
+	RegisterInputConnection(&BucklingStiffnessWeft);
+	RegisterInputConnection(&BucklingStiffnessBias);
+
+	RegisterInputConnection(&bEnableXPBDAreaConstraints);
+	RegisterInputConnection(&AreaStiffness);
+
+	RegisterOutputConnection(&Collection, &Collection);
+}
+
+void FChaosClothAssetSimulationXPBDConfigNode::Serialize(FArchive& Ar)
+{
+	if (!SharedSimulationConfig)
+	{
+		SharedSimulationConfig = NewObject<UChaosClothSharedSimConfig>();
+	}
+	SharedSimulationConfig->Serialize(Ar);
+}
+
+void FChaosClothAssetSimulationXPBDConfigNode::Evaluate(Dataflow::FContext& Context, const FDataflowOutput* Out) const
+{
+	using namespace Chaos;
+	using namespace Chaos::Softs;
+
+	if (Out->IsA<FManagedArrayCollection>(&Collection))
+	{
+		FManagedArrayCollection InCollection = GetValue<FManagedArrayCollection>(Context, &Collection);
+		const TSharedRef<FManagedArrayCollection> ClothCollection = MakeShared<FManagedArrayCollection>(MoveTemp(InCollection));
+
+		FCollectionPropertyMutableFacade Properties(ClothCollection);
+		Properties.DefineSchema();
+
+		// Set xpbd properties
+		SetStretchProperties(Properties);
+		SetBendingProperties(Properties);
+		SetAreaProperties(Properties);
+		
+		// Copy shared config properties
+		// Solver properties
+		if (SharedSimulationConfig)
+		{
+			constexpr bool bEnable = true;
+			constexpr bool bAnimatable = true;
+			Properties.AddValue(TEXT("NumIterations"), SharedSimulationConfig->IterationCount, bEnable, bAnimatable);
+			Properties.AddValue(TEXT("MaxNumIterations"), SharedSimulationConfig->MaxIterationCount, bEnable, bAnimatable);
+			Properties.AddValue(TEXT("NumSubsteps"), SharedSimulationConfig->SubdivisionCount, bEnable, bAnimatable);
+		}
+
+
+		SetValue<FManagedArrayCollection>(Context, *ClothCollection, &Collection);
+	}
+}
+
+void FChaosClothAssetSimulationXPBDConfigNode::AddReferencedObjects(FReferenceCollector& Collector)
+{
+	Collector.AddReferencedObject(SharedSimulationConfig);
+}
+
+FString FChaosClothAssetSimulationXPBDConfigNode::GetReferencerName() const
+{
+	return TEXT("FChaosClothAssetSimulationXPBDConfigNode");
+}
+
+void FChaosClothAssetSimulationXPBDConfigNode::SetStretchProperties(Chaos::Softs::FCollectionPropertyMutableFacade& Properties) const
+{
+	if (bEnableXPBDStretchConstraints)
+	{
+		const int32 StretchStiffnessIndex = Properties.AddProperty(TEXT("XPBDEdgeSpringStiffness"));
+		Properties.SetWeightedValue(StretchStiffnessIndex, StretchStiffness.Low, StretchStiffness.High);
+		Properties.SetStringValue(StretchStiffnessIndex, TEXT("EdgeStiffness"));
+
+		const int32 StretchDampingRatioIndex = Properties.AddProperty(TEXT("XPBDEdgeSpringDamping"));
+		Properties.SetWeightedValue(StretchDampingRatioIndex, StretchDampingRatio.Low, StretchDampingRatio.High);
+		Properties.SetStringValue(StretchDampingRatioIndex, TEXT("EdgeDamping"));
+	}
+	else
+	{
+		Properties.SetEnabled(TEXT("XPBDEdgeSpringStiffness"), false);
+	}
+}
+
+void FChaosClothAssetSimulationXPBDConfigNode::SetBendingProperties(Chaos::Softs::FCollectionPropertyMutableFacade& Properties) const
+{
+	if(bEnableXPBDBendConstraints)
+	{
+		if (bEnableBendAnisotropy)
+		{
+			const int32 BendingStiffnessWarpIndex = Properties.AddProperty(TEXT("XPBDBendingElementStiffnessWarp"));
+			Properties.SetWeightedValue(BendingStiffnessWarpIndex, BendingStiffness.Low, BendingStiffness.High);
+			Properties.SetStringValue(BendingStiffnessWarpIndex, TEXT("BendingStiffnessWarp"));
+			const int32 BendingStiffnessWeftIndex = Properties.AddProperty(TEXT("XPBDBendingElementStiffnessWeft"));
+			Properties.SetWeightedValue(BendingStiffnessWeftIndex, BendingStiffnessWeft.Low, BendingStiffnessWeft.High);
+			Properties.SetStringValue(BendingStiffnessWeftIndex, TEXT("BendingStiffnessWeft"));
+			const int32 BendingStiffnessBiasIndex = Properties.AddProperty(TEXT("XPBDBendingElementStiffnessBias"));
+			Properties.SetWeightedValue(BendingStiffnessBiasIndex, BendingStiffnessBias.Low, BendingStiffnessBias.High);
+			Properties.SetStringValue(BendingStiffnessBiasIndex, TEXT("BendingStiffnessBias"));
+
+			const int32 BendingDampingRatioIndex = Properties.AddProperty(TEXT("XPBDBendingElementDamping"));
+			Properties.SetWeightedValue(BendingDampingRatioIndex, BendingDampingRatio.Low, BendingDampingRatio.High);
+			Properties.SetStringValue(BendingDampingRatioIndex, TEXT("BendingDampingRatio"));
+
+			Properties.AddValue(TEXT("XPBDBucklingRatio"), BucklingRatio, true, true);
+
+			const int32 BucklingStiffnessWarpIndex = Properties.AddProperty(TEXT("XPBDBucklingStiffnessWarp"));
+			Properties.SetWeightedValue(BucklingStiffnessWarpIndex, BucklingStiffness.Low, BucklingStiffness.High);
+			Properties.SetStringValue(BucklingStiffnessWarpIndex, TEXT("BucklingStiffnessWarp"));
+			const int32 BucklingStiffnessWeftIndex = Properties.AddProperty(TEXT("XPBDBucklingStiffnessWeft"));
+			Properties.SetWeightedValue(BucklingStiffnessWeftIndex, BucklingStiffnessWeft.Low, BucklingStiffnessWeft.High);
+			Properties.SetStringValue(BucklingStiffnessWeftIndex, TEXT("BucklingStiffnessWeft"));
+			const int32 BucklingStiffnessBiasIndex = Properties.AddProperty(TEXT("XPBDBucklingStiffnessBias"));
+			Properties.SetWeightedValue(BucklingStiffnessBiasIndex, BucklingStiffnessBias.Low, BucklingStiffnessBias.High);
+			Properties.SetStringValue(BucklingStiffnessBiasIndex, TEXT("BucklingStiffnessBias"));
+
+			Properties.SetEnabled(TEXT("XPBDBendingElementStiffness"), false);
+		}
+		else
+		{
+			const int32 BendingStiffnessIndex = Properties.AddProperty(TEXT("XPBDBendingElementStiffness"));
+			Properties.SetWeightedValue(BendingStiffnessIndex, BendingStiffness.Low, BendingStiffness.High);
+			Properties.SetStringValue(BendingStiffnessIndex, TEXT("BendingStiffness"));
+
+			const int32 BendingDampingRatioIndex = Properties.AddProperty(TEXT("XPBDBendingElementDamping"));
+			Properties.SetWeightedValue(BendingDampingRatioIndex, BendingDampingRatio.Low, BendingDampingRatio.High);
+			Properties.SetStringValue(BendingDampingRatioIndex, TEXT("BendingDampingRatio"));
+
+			Properties.AddValue(TEXT("XPBDBucklingRatio"), BucklingRatio, true, true);
+
+			const int32 BucklingStiffnessIndex = Properties.AddProperty(TEXT("XPBDBucklingStiffness"));
+			Properties.SetWeightedValue(BucklingStiffnessIndex, BucklingStiffness.Low, BucklingStiffness.High);
+			Properties.SetStringValue(BucklingStiffnessIndex, TEXT("BucklingStiffness"));
+
+			Properties.SetEnabled(TEXT("XPBDBendingElementStiffnessWarp"), false);
+		}
+	}
+	else
+	{
+		Properties.SetEnabled(TEXT("XPBDBendingElementStiffness"), false);
+		Properties.SetEnabled(TEXT("XPBDBendingElementStiffnessWarp"), false);
+	}
+}
+
+void FChaosClothAssetSimulationXPBDConfigNode::SetAreaProperties(Chaos::Softs::FCollectionPropertyMutableFacade& Properties) const
+{
+	if (bEnableXPBDAreaConstraints)
+	{
+		const int32 AreaStiffnessIndex = Properties.AddProperty(TEXT("XPBDAreaSpringStiffness"));
+		Properties.SetWeightedValue(AreaStiffnessIndex, AreaStiffness.Low, AreaStiffness.High);
+		Properties.SetStringValue(AreaStiffnessIndex, TEXT("AreaStiffness"));
+	}
+	else
+	{
+		Properties.SetEnabled(TEXT("XPBDAreaSpringStiffness"), false);
+	}
+}
+
+#undef LOCTEXT_NAMESPACE
diff --git a/Engine/Plugins/Experimental/ChaosClothAssetEditor/Source/ChaosClothAssetDataflowNodes/Private/ChaosClothAsset/StaticMeshImportNode.cpp b/Engine/Plugins/Experimental/ChaosClothAssetEditor/Source/ChaosClothAssetDataflowNodes/Private/ChaosClothAsset/StaticMeshImportNode.cpp
new file mode 100644
--- /dev/null
+++ b/Engine/Plugins/Experimental/ChaosClothAssetEditor/Source/ChaosClothAssetDataflowNodes/Private/ChaosClothAsset/StaticMeshImportNode.cpp
@@ -0,0 +1,65 @@
+// Copyright Epic Games, Inc. All Rights Reserved.
+
+#include "ChaosClothAsset/StaticMeshImportNode.h"
+#include "ChaosClothAsset/DataflowNodes.h"
+#include "ChaosClothAsset/CollectionClothFacade.h"
+#include "ChaosClothAsset/ClothAsset.h"
+#include "ChaosClothAsset/ClothCollection.h"
+#include "Engine/StaticMesh.h"
+#include "MeshDescriptionToDynamicMesh.h"
+
+#include UE_INLINE_GENERATED_CPP_BY_NAME(StaticMeshImportNode)
+
+#define LOCTEXT_NAMESPACE "ChaosClothAssetStaticMeshImportNode"
+
+FChaosClothAssetStaticMeshImportNode::FChaosClothAssetStaticMeshImportNode(const Dataflow::FNodeParameters& InParam, FGuid InGuid)
+	: FDataflowNode(InParam, InGuid)
+{
+	RegisterOutputConnection(&Collection);
+}
+
+void FChaosClothAssetStaticMeshImportNode::Evaluate(Dataflow::FContext& Context, const FDataflowOutput* Out) const
+{
+	using namespace UE::Chaos::ClothAsset;
+	using namespace UE::Chaos::ClothAsset::DataflowNodes;
+
+	if (Out->IsA<FManagedArrayCollection>(&Collection))
+	{
+
+		if (StaticMesh)
+		{
+			// Evaluate in collection
+			const TSharedRef<FManagedArrayCollection> ClothCollection = MakeShared<FManagedArrayCollection>();
+			FCollectionClothFacade ClothFacade(ClothCollection);
+			ClothFacade.DefineSchema();
+
+			const int32 NumLods = StaticMesh->GetNumSourceModels();
+			for (int32 LodIndex = 0; LodIndex < NumLods; ++LodIndex)
+			{
+				const FMeshDescription* MeshDescription = StaticMesh->GetMeshDescription(LodIndex);
+				check(MeshDescription);
+
+				FMeshDescriptionToDynamicMesh Converter;
+				Converter.bPrintDebugMessages = true; // TODO: testing only
+				Converter.bEnableOutputGroups = false;
+				Converter.bVIDsFromNonManifoldMeshDescriptionAttr = true;
+				UE::Geometry::FDynamicMesh3 DynamicMesh;
+
+				Converter.Convert(MeshDescription, DynamicMesh);
+
+				FCollectionClothLodFacade ClothLodFacade = ClothFacade.AddGetLod();
+				ClothLodFacade.Initialize(DynamicMesh, UVChannel);
+			}
+			// Make sure that whatever happens there is always at least one empty LOD to avoid crashing the render data
+			if (ClothFacade.GetNumLods() < 1)
+			{
+				ClothFacade.AddLod();
+			}
+			SetValue<FManagedArrayCollection>(Context, *ClothCollection, &Collection);
+		}
+
+
+	}
+}
+
+#undef LOCTEXT_NAMESPACE
diff --git a/Engine/Plugins/Experimental/ChaosClothAssetEditor/Source/ChaosClothAssetDataflowNodes/Private/ChaosClothAsset/TerminalNode.cpp b/Engine/Plugins/Experimental/ChaosClothAssetEditor/Source/ChaosClothAssetDataflowNodes/Private/ChaosClothAsset/TerminalNode.cpp
new file mode 100644
--- /dev/null
+++ b/Engine/Plugins/Experimental/ChaosClothAssetEditor/Source/ChaosClothAssetDataflowNodes/Private/ChaosClothAsset/TerminalNode.cpp
@@ -0,0 +1,201 @@
+// Copyright Epic Games, Inc. All Rights Reserved.
+
+#include "ChaosClothAsset/TerminalNode.h"
+#include "ChaosClothAsset/DataflowNodes.h"
+#include "ChaosClothAsset/ClothAsset.h"
+#include "ChaosClothAsset/CollectionClothFacade.h"
+#include "Chaos/CollectionPropertyFacade.h"
+#include "Animation/Skeleton.h"
+#include "Engine/SkinnedAssetCommon.h"
+#include "Materials/Material.h"
+#include "PhysicsEngine/PhysicsAsset.h"
+
+#include UE_INLINE_GENERATED_CPP_BY_NAME(TerminalNode)
+
+#define LOCTEXT_NAMESPACE "ChaosClothAssetTerminalNode"
+
+FChaosClothAssetTerminalNode::FChaosClothAssetTerminalNode(const Dataflow::FNodeParameters& InParam, FGuid InGuid)
+	: FDataflowTerminalNode(InParam, InGuid)
+{
+	RegisterInputConnection(&CollectionLod0);
+}
+
+void FChaosClothAssetTerminalNode::SetAssetValue(TObjectPtr<UObject> Asset, Dataflow::FContext& Context) const
+{
+	if (UChaosClothAsset* ClothAsset = Cast<UChaosClothAsset>(Asset.Get()))
+	{
+		using namespace UE::Chaos::ClothAsset;
+
+		// Reset the asset's collection
+		const TSharedPtr<FManagedArrayCollection> ClothCollection = ClothAsset->GetClothCollection();
+		ClothCollection->Reset();
+
+		FCollectionClothFacade ClothFacade(ClothCollection);
+		ClothFacade.DefineSchema();
+
+		// Iterate through the LODs
+		FString SkeletonPathName;
+		FString PhysicsAssetPathName;
+
+		const TArray<const FManagedArrayCollection*> CollectionLods = GetCollectionLods();
+		for (int32 LodIndex = 0; LodIndex < CollectionLods.Num(); ++LodIndex)
+		{
+			// New LOD
+			const int32 NewLodIndex = ClothFacade.AddLod();
+			check(NewLodIndex == LodIndex);
+
+			// Retrieve input LOD
+			const FManagedArrayCollection& InCollectionLod = GetValue<FManagedArrayCollection>(Context, CollectionLods[LodIndex]);
+			const TSharedRef<const FManagedArrayCollection> InClothCollection = MakeShared<const FManagedArrayCollection>(InCollectionLod);
+
+			// Only use LOD 0 from input
+			const FCollectionClothConstFacade InClothFacade(InClothCollection);
+			if (InClothFacade.GetNumLods() >= 1)
+			{
+				const FCollectionClothLodConstFacade InClothLodFacade = InClothFacade.GetLod(0);
+
+				// Copy input LOD 0 to current output LOD
+				FCollectionClothLodFacade ClothLodFacade = ClothFacade.GetLod(LodIndex);
+				ClothLodFacade.Initialize(InClothLodFacade);
+
+				// Add this LOD's materials to the asset
+				const int32 NumLodMaterials = InClothLodFacade.GetNumMaterials();
+
+				TArray<FSkeletalMaterial>& Materials = ClothAsset->GetMaterials();
+				Materials.Reserve(Materials.Num() + NumLodMaterials);
+
+				const TConstArrayView<FString> LodRenderMaterialPathName = ClothLodFacade.GetRenderMaterialPathName();
+				for (int32 LodMaterialIndex = 0; LodMaterialIndex < NumLodMaterials; ++LodMaterialIndex)
+				{
+					const FString& RenderMaterialPathName = LodRenderMaterialPathName[LodMaterialIndex];
+
+					auto Predicate = [&RenderMaterialPathName](const FSkeletalMaterial& SkeletalMaterial)
+						{
+							return SkeletalMaterial.MaterialInterface && SkeletalMaterial.MaterialInterface->GetPathName() == RenderMaterialPathName;
+						};
+
+					if (!Materials.FindByPredicate(Predicate))
+					{
+						if (UMaterial* const Material = LoadObject<UMaterial>(ClothAsset, *RenderMaterialPathName, nullptr, LOAD_None, nullptr))
+						{
+							Materials.Emplace(Material, true, false, Material->GetFName());
+						}
+					}
+				}
+
+				const TArray<FName> MapNames = InClothFacade.GetWeightMapNames();
+				for (const FName& MapName : MapNames)
+				{
+					ClothFacade.AddWeightMap(MapName);
+					TArrayView<float> WeightMap = ClothLodFacade.GetWeightMap(MapName);
+					TConstArrayView<float> InWeightMap = InClothLodFacade.GetWeightMap(MapName);
+					for (int32 i = 0; i < WeightMap.Num(); ++i)
+					{
+						WeightMap[i] = InWeightMap[i];
+					}
+				}
+
+				// Set properties, skeleton, and physics asset only with LOD 0 at the moment
+				if (LodIndex == 0)
+				{
+					using namespace ::Chaos::Softs;
+
+					SkeletonPathName = InClothLodFacade.GetSkeletonAssetPathName();
+					PhysicsAssetPathName = InClothLodFacade.GetPhysicsAssetPathName();
+				
+					FCollectionPropertyMutableFacade(ClothCollection).Append(*InClothCollection);
+				}
+			}
+		}
+
+		// Make sure that whatever happens there is always at least one empty LOD to avoid crashing the render data
+		if (ClothFacade.GetNumLods() < 1)
+		{
+			ClothFacade.AddLod();
+		}
+
+		// Set reference skeleton
+		constexpr bool bRebuildModels = false;  // Avoid rebuilding the asset twice
+		USkeleton* const Skeleton = !SkeletonPathName.IsEmpty() ?
+			LoadObject<USkeleton>(nullptr, *SkeletonPathName, nullptr, LOAD_None, nullptr) :
+			nullptr;
+		ClothAsset->SetSkeleton(Skeleton, bRebuildModels);
+
+		// Set physics asset (note: the cloth asset's physics asset is only replaced if a collection path name is found valid)
+		UPhysicsAsset* const PhysicsAsset = !PhysicsAssetPathName.IsEmpty() ?
+			LoadObject<UPhysicsAsset>(ClothAsset, *PhysicsAssetPathName, nullptr, LOAD_None, nullptr) :
+			nullptr;
+		ClothAsset->SetPhysicsAsset(PhysicsAsset);
+
+		// Rebuild the asset static data
+		ClothAsset->Build();
+	}
+}
+
+Dataflow::FPin FChaosClothAssetTerminalNode::AddPin()
+{
+	auto AddInput = [this](const FManagedArrayCollection* Collection) -> Dataflow::FPin
+		{
+			RegisterInputConnection(Collection);
+			const FDataflowInput* const Input = FindInput(Collection);
+			return { Dataflow::FPin::EDirection::INPUT, Input->GetType(), Input->GetName() };
+		};
+
+	switch (NumLods)
+	{
+	case 1: ++NumLods; return AddInput(&CollectionLod1);
+	case 2: ++NumLods; return AddInput(&CollectionLod2);
+	case 3: ++NumLods; return AddInput(&CollectionLod3);
+	case 4: ++NumLods; return AddInput(&CollectionLod4);
+	case 5: ++NumLods; return AddInput(&CollectionLod5);
+	default: break;
+	}
+
+	return Super::AddPin();
+}
+
+Dataflow::FPin FChaosClothAssetTerminalNode::RemovePin()
+{
+	auto RemoveInput = [this](const FManagedArrayCollection* Collection) -> Dataflow::FPin
+		{
+			const FDataflowInput* const Input = FindInput(Collection);
+			check(Input);
+			Dataflow::FPin Pin = { Dataflow::FPin::EDirection::INPUT, Input->GetType(), Input->GetName() };
+			UnregisterInputConnection(Collection);  // This will delete the input, so set the pin before that
+			return Pin;
+		};
+
+	switch (NumLods - 1)
+	{
+	case 1: --NumLods; return RemoveInput(&CollectionLod1);
+	case 2: --NumLods; return RemoveInput(&CollectionLod2);
+	case 3: --NumLods; return RemoveInput(&CollectionLod3);
+	case 4: --NumLods; return RemoveInput(&CollectionLod4);
+	case 5: --NumLods; return RemoveInput(&CollectionLod5);
+	default: break;
+	}
+	return Super::AddPin();
+}
+
+TArray<const FManagedArrayCollection*> FChaosClothAssetTerminalNode::GetCollectionLods() const
+{
+	TArray<const FManagedArrayCollection*> CollectionLods;
+	CollectionLods.SetNumUninitialized(NumLods);
+
+	for (int32 LodIndex = 0; LodIndex < NumLods; ++LodIndex)
+	{
+		switch (LodIndex)
+		{
+		case 0: CollectionLods[LodIndex] = &CollectionLod0; break;
+		case 1: CollectionLods[LodIndex] = &CollectionLod1; break;
+		case 2: CollectionLods[LodIndex] = &CollectionLod2; break;
+		case 3: CollectionLods[LodIndex] = &CollectionLod3; break;
+		case 4: CollectionLods[LodIndex] = &CollectionLod4; break;
+		case 5: CollectionLods[LodIndex] = &CollectionLod5; break;
+		default: check(false); break;
+		}
+	}
+	return CollectionLods;
+}
+
+#undef LOCTEXT_NAMESPACE
diff --git a/Engine/Plugins/Experimental/ChaosClothAssetEditor/Source/ChaosClothAssetDataflowNodes/Public/ChaosClothAsset/BindToRootBoneNode.h b/Engine/Plugins/Experimental/ChaosClothAssetEditor/Source/ChaosClothAssetDataflowNodes/Public/ChaosClothAsset/BindToRootBoneNode.h
new file mode 100644
--- /dev/null
+++ b/Engine/Plugins/Experimental/ChaosClothAssetEditor/Source/ChaosClothAssetDataflowNodes/Public/ChaosClothAsset/BindToRootBoneNode.h
@@ -0,0 +1,32 @@
+// Copyright Epic Games, Inc. All Rights Reserved.
+
+#pragma once 
+
+#include "Dataflow/DataflowNode.h"
+#include "GeometryCollection/ManagedArrayCollection.h"
+#include "BindToRootBoneNode.generated.h"
+
+USTRUCT(Meta = (DataflowCloth))
+struct FChaosClothAssetBindToRootBoneNode : public FDataflowNode
+{
+	GENERATED_USTRUCT_BODY()
+	DATAFLOW_NODE_DEFINE_INTERNAL(FChaosClothAssetBindToRootBoneNode, "BindToRootBone", "Cloth", "Cloth Bind Skinning Weights To Root Bone")
+	//DATAFLOW_NODE_RENDER_TYPE(FGeometryCollection::StaticType(), "Collection")  // TODO: Leave out the render type until there is something to render
+
+public:
+
+	UPROPERTY(Meta = (Dataflowinput, DataflowOutput, DisplayName = "Collection", DataflowPassthrough = "Collection"))
+	FManagedArrayCollection Collection;
+
+	/** Whether to bind the simulation mesh. */
+	UPROPERTY(EditAnywhere, Category = "Bind To Root Bone")
+	bool bBindSimMesh = true;
+
+	/** Whether to bind the render mesh. */
+	UPROPERTY(EditAnywhere, Category = "Bind To Root Bone")
+	bool bBindRenderMesh = true;
+
+	FChaosClothAssetBindToRootBoneNode(const Dataflow::FNodeParameters& InParam, FGuid InGuid = FGuid::NewGuid());
+
+	virtual void Evaluate(Dataflow::FContext& Context, const FDataflowOutput* Out) const override;
+};
diff --git a/Engine/Plugins/Experimental/ChaosClothAssetEditor/Source/ChaosClothAssetDataflowNodes/Public/ChaosClothAsset/CopySimulationToRenderMeshNode.h b/Engine/Plugins/Experimental/ChaosClothAssetEditor/Source/ChaosClothAssetDataflowNodes/Public/ChaosClothAsset/CopySimulationToRenderMeshNode.h
new file mode 100644
--- /dev/null
+++ b/Engine/Plugins/Experimental/ChaosClothAssetEditor/Source/ChaosClothAssetDataflowNodes/Public/ChaosClothAsset/CopySimulationToRenderMeshNode.h
@@ -0,0 +1,34 @@
+// Copyright Epic Games, Inc. All Rights Reserved.
+
+#pragma once 
+
+#include "Dataflow/DataflowNode.h"
+#include "GeometryCollection/ManagedArrayCollection.h"
+#include "CopySimulationToRenderMeshNode.generated.h"
+
+class UMaterialInterface;
+
+USTRUCT(Meta = (DataflowCloth))
+struct FChaosClothAssetCopySimulationToRenderMeshNode : public FDataflowNode
+{
+	GENERATED_USTRUCT_BODY()
+	DATAFLOW_NODE_DEFINE_INTERNAL(FChaosClothAssetCopySimulationToRenderMeshNode, "CopySimulationToRenderMesh", "Cloth", "Cloth Simulation Render Mesh")
+	//DATAFLOW_NODE_RENDER_TYPE(FGeometryCollection::StaticType(), "Collection")  // TODO: Leave out the render type until there is something to render
+
+public:
+
+	UPROPERTY(Meta = (Dataflowinput, DataflowOutput, DisplayName = "Collection", DataflowPassthrough = "Collection"))
+	FManagedArrayCollection Collection;
+
+	/** List of patterns to apply the operation on. All patterns will be used if left empty. */
+	UPROPERTY(EditAnywhere, Category = "Pattern Selection")
+	TArray<int32> Patterns;
+
+	/** New material for the render mesh. */
+	UPROPERTY(EditAnywhere, Category = "Render Mesh")
+	TObjectPtr<const UMaterialInterface> Material;
+
+	FChaosClothAssetCopySimulationToRenderMeshNode(const Dataflow::FNodeParameters& InParam, FGuid InGuid = FGuid::NewGuid());
+
+	virtual void Evaluate(Dataflow::FContext& Context, const FDataflowOutput* Out) const override;
+};
diff --git a/Engine/Plugins/Experimental/ChaosClothAssetEditor/Source/ChaosClothAssetDataflowNodes/Public/ChaosClothAsset/DataflowNodes.h b/Engine/Plugins/Experimental/ChaosClothAssetEditor/Source/ChaosClothAssetDataflowNodes/Public/ChaosClothAsset/DataflowNodes.h
new file mode 100644
--- /dev/null
+++ b/Engine/Plugins/Experimental/ChaosClothAssetEditor/Source/ChaosClothAssetDataflowNodes/Public/ChaosClothAsset/DataflowNodes.h
@@ -0,0 +1,16 @@
+// Copyright Epic Games, Inc. All Rights Reserved.
+
+#pragma once
+
+#include "Dataflow/DataflowNodeParameters.h"
+#include "Dataflow/DataflowInputOutput.h"
+#include "Dataflow/DataflowObjectInterface.h"
+#include "Logging/LogMacros.h"
+
+namespace UE::Chaos::ClothAsset::DataflowNodes
+{
+	void Register();
+	void LogAndToastWarning(const FText& Error);
+}
+
+DECLARE_LOG_CATEGORY_EXTERN(LogChaosClothAssetDataflowNodes, Log, All);
diff --git a/Engine/Plugins/Experimental/ChaosClothAssetEditor/Source/ChaosClothAssetDataflowNodes/Public/ChaosClothAsset/DatasmithImportNode.h b/Engine/Plugins/Experimental/ChaosClothAssetEditor/Source/ChaosClothAssetDataflowNodes/Public/ChaosClothAsset/DatasmithImportNode.h
new file mode 100644
--- /dev/null
+++ b/Engine/Plugins/Experimental/ChaosClothAssetEditor/Source/ChaosClothAssetDataflowNodes/Public/ChaosClothAsset/DatasmithImportNode.h
@@ -0,0 +1,32 @@
+// Copyright Epic Games, Inc. All Rights Reserved.
+
+#pragma once 
+
+#include "Dataflow/DataflowNode.h"
+#include "GeometryCollection/ManagedArrayCollection.h"
+#include "DatasmithImportNode.generated.h"
+
+USTRUCT(meta = (DataflowCloth))
+struct FChaosClothAssetDatasmithImportNode : public FDataflowNode
+{
+	GENERATED_USTRUCT_BODY()
+	DATAFLOW_NODE_DEFINE_INTERNAL(FChaosClothAssetDatasmithImportNode, "DatasmithImport", "Cloth", "Cloth Datasmith Import")
+	//DATAFLOW_NODE_RENDER_TYPE(FGeometryCollection::StaticType(), "Collection")
+
+public:
+	
+	/** Datasmith file to read from. */
+	UPROPERTY(EditAnywhere, Category = "Datasmith Importer")
+	FFilePath DatasmithFile;
+
+	UPROPERTY(meta = (DataflowOutput, DisplayName = "Collection"))
+	FManagedArrayCollection Collection;
+
+	FChaosClothAssetDatasmithImportNode(const Dataflow::FNodeParameters& InParam, FGuid InGuid = FGuid::NewGuid());
+
+	virtual void Evaluate(Dataflow::FContext& Context, const FDataflowOutput* Out) const override;
+
+private:
+
+	bool EvaluateImpl(Dataflow::FContext& Context, FManagedArrayCollection& OutCollection) const;
+};
diff --git a/Engine/Plugins/Experimental/ChaosClothAssetEditor/Source/ChaosClothAssetDataflowNodes/Public/ChaosClothAsset/DeleteRenderMeshNode.h b/Engine/Plugins/Experimental/ChaosClothAssetEditor/Source/ChaosClothAssetDataflowNodes/Public/ChaosClothAsset/DeleteRenderMeshNode.h
new file mode 100644
--- /dev/null
+++ b/Engine/Plugins/Experimental/ChaosClothAssetEditor/Source/ChaosClothAssetDataflowNodes/Public/ChaosClothAsset/DeleteRenderMeshNode.h
@@ -0,0 +1,28 @@
+// Copyright Epic Games, Inc. All Rights Reserved.
+
+#pragma once 
+
+#include "Dataflow/DataflowNode.h"
+#include "GeometryCollection/ManagedArrayCollection.h"
+#include "DeleteRenderMeshNode.generated.h"
+
+USTRUCT(Meta = (DataflowCloth))
+struct FChaosClothAssetDeleteRenderMeshNode : public FDataflowNode
+{
+	GENERATED_USTRUCT_BODY()
+	DATAFLOW_NODE_DEFINE_INTERNAL(FChaosClothAssetDeleteRenderMeshNode, "DeleteRenderMesh", "Cloth", "Cloth Delete Render Mesh")
+	//DATAFLOW_NODE_RENDER_TYPE(FGeometryCollection::StaticType(), "Collection")  // TODO: Leave out the render type until there is something to render
+
+public:
+
+	UPROPERTY(Meta = (Dataflowinput, DataflowOutput, DisplayName = "Collection", DataflowPassthrough = "Collection"))
+	FManagedArrayCollection Collection;
+
+	/** List of patterns to apply the operation on. All patterns will be used if left empty. */
+	UPROPERTY(EditAnywhere, Category = "Pattern Selection")
+	TArray<int32> Patterns;
+
+	FChaosClothAssetDeleteRenderMeshNode(const Dataflow::FNodeParameters& InParam, FGuid InGuid = FGuid::NewGuid());
+
+	virtual void Evaluate(Dataflow::FContext& Context, const FDataflowOutput* Out) const override;
+};
diff --git a/Engine/Plugins/Experimental/ChaosClothAssetEditor/Source/ChaosClothAssetDataflowNodes/Public/ChaosClothAsset/ImportNode.h b/Engine/Plugins/Experimental/ChaosClothAssetEditor/Source/ChaosClothAssetDataflowNodes/Public/ChaosClothAsset/ImportNode.h
new file mode 100644
--- /dev/null
+++ b/Engine/Plugins/Experimental/ChaosClothAssetEditor/Source/ChaosClothAssetDataflowNodes/Public/ChaosClothAsset/ImportNode.h
@@ -0,0 +1,34 @@
+// Copyright Epic Games, Inc. All Rights Reserved.
+
+#pragma once 
+
+#include "Dataflow/DataflowNode.h"
+#include "GeometryCollection/ManagedArrayCollection.h"
+#include "ImportNode.generated.h"
+
+class UChaosClothAsset;
+
+USTRUCT(Meta = (DataflowCloth))
+struct FChaosClothAssetImportNode : public FDataflowNode
+{
+	GENERATED_USTRUCT_BODY()
+	DATAFLOW_NODE_DEFINE_INTERNAL(FChaosClothAssetImportNode, "ClothAssetImport", "Cloth", "Cloth Asset Import")
+	//DATAFLOW_NODE_RENDER_TYPE(FGeometryCollection::StaticType(), "Collection")  // TODO: Leave out the render type until there is something to render
+
+public:
+
+	UPROPERTY(Meta = (DataflowOutput, DisplayName = "Collection"))
+	FManagedArrayCollection Collection;
+
+	/** The Cloth Asset to import into a collection. */
+	UPROPERTY(EditAnywhere, Category = "ClothAsset")
+	TObjectPtr<const UChaosClothAsset> ClothAsset;
+
+	/** The LOD to import into the collection. Only one LOD can be imported at a time. */
+	UPROPERTY(EditAnywhere, Category = "ClothAsset", Meta = (DisplayName = "Import LOD"))
+	int32 ImportLod;
+
+	FChaosClothAssetImportNode(const Dataflow::FNodeParameters& InParam, FGuid InGuid = FGuid::NewGuid());
+
+	virtual void Evaluate(Dataflow::FContext& Context, const FDataflowOutput* Out) const override;
+};
diff --git a/Engine/Plugins/Experimental/ChaosClothAssetEditor/Source/ChaosClothAssetDataflowNodes/Public/ChaosClothAsset/PatternSelectionNode.h b/Engine/Plugins/Experimental/ChaosClothAssetEditor/Source/ChaosClothAssetDataflowNodes/Public/ChaosClothAsset/PatternSelectionNode.h
new file mode 100644
--- /dev/null
+++ b/Engine/Plugins/Experimental/ChaosClothAssetEditor/Source/ChaosClothAssetDataflowNodes/Public/ChaosClothAsset/PatternSelectionNode.h
@@ -0,0 +1,23 @@
+// Copyright Epic Games, Inc. All Rights Reserved.
+
+#pragma once 
+
+#include "Dataflow/DataflowNode.h"
+#include "PatternSelectionNode.generated.h"
+
+USTRUCT(Meta = (DataflowCloth))
+struct FChaosClothAssetPatternSelectionNode : public FDataflowNode
+{
+	GENERATED_USTRUCT_BODY()
+	DATAFLOW_NODE_DEFINE_INTERNAL(FChaosClothAssetPatternSelectionNode, "PatternSelection", "Cloth", "Cloth Pattern Selection")
+	//DATAFLOW_NODE_RENDER_TYPE(FGeometryCollection::StaticType(), "Collection")  // TODO: Leave out the render type until there is something to render
+
+public:
+
+	UPROPERTY(EditAnywhere, Category = "Pattern", Meta = (DataflowOutput, DisplayName = "Patterns"))
+	TArray<int32> Patterns;
+
+	FChaosClothAssetPatternSelectionNode(const Dataflow::FNodeParameters& InParam, FGuid InGuid = FGuid::NewGuid());
+
+	virtual void Evaluate(Dataflow::FContext& Context, const FDataflowOutput* Out) const override;
+};
diff --git a/Engine/Plugins/Experimental/ChaosClothAssetEditor/Source/ChaosClothAssetDataflowNodes/Public/ChaosClothAsset/ReverseNormalsNode.h b/Engine/Plugins/Experimental/ChaosClothAssetEditor/Source/ChaosClothAssetDataflowNodes/Public/ChaosClothAsset/ReverseNormalsNode.h
new file mode 100644
--- /dev/null
+++ b/Engine/Plugins/Experimental/ChaosClothAssetEditor/Source/ChaosClothAssetDataflowNodes/Public/ChaosClothAsset/ReverseNormalsNode.h
@@ -0,0 +1,36 @@
+// Copyright Epic Games, Inc. All Rights Reserved.
+
+#pragma once 
+
+#include "Dataflow/DataflowNode.h"
+#include "GeometryCollection/ManagedArrayCollection.h"
+#include "ReverseNormalsNode.generated.h"
+
+USTRUCT(Meta = (DataflowCloth))
+struct FChaosClothAssetReverseNormalsNode : public FDataflowNode
+{
+	GENERATED_USTRUCT_BODY()
+	DATAFLOW_NODE_DEFINE_INTERNAL(FChaosClothAssetReverseNormalsNode, "ReverseNormals", "Cloth", "Cloth Reverse Simulation Render Mesh Normals")
+	//DATAFLOW_NODE_RENDER_TYPE(FGeometryCollection::StaticType(), "Collection")  // TODO: Leave out the render type until there is something to render
+
+public:
+
+	UPROPERTY(Meta = (Dataflowinput, DataflowOutput, DisplayName = "Collection", DataflowPassthrough = "Collection"))
+	FManagedArrayCollection Collection;
+
+	/** List of patterns to apply the operation on. All patterns will be used if left empty. */
+	UPROPERTY(EditAnywhere, Category = "Pattern Selection")
+	TArray<int32> Patterns;
+
+	/** Whether to reverse the simulation mesh normals. */
+	UPROPERTY(EditAnywhere, Category = "Reverse Normals")
+	bool bReverseSimMeshNormals = true;
+
+	/** Whether to reverse the render mesh normals. */
+	UPROPERTY(EditAnywhere, Category = "Reverse Normals")
+	bool bReverseRenderMeshNormals = true;
+
+	FChaosClothAssetReverseNormalsNode(const Dataflow::FNodeParameters& InParam, FGuid InGuid = FGuid::NewGuid());
+
+	virtual void Evaluate(Dataflow::FContext& Context, const FDataflowOutput* Out) const override;
+};
diff --git a/Engine/Plugins/Experimental/ChaosClothAssetEditor/Source/ChaosClothAssetDataflowNodes/Public/ChaosClothAsset/SimulationDefaultConfigNode.h b/Engine/Plugins/Experimental/ChaosClothAssetEditor/Source/ChaosClothAssetDataflowNodes/Public/ChaosClothAsset/SimulationDefaultConfigNode.h
new file mode 100644
--- /dev/null
+++ b/Engine/Plugins/Experimental/ChaosClothAssetEditor/Source/ChaosClothAssetDataflowNodes/Public/ChaosClothAsset/SimulationDefaultConfigNode.h
@@ -0,0 +1,43 @@
+// Copyright Epic Games, Inc. All Rights Reserved.
+
+#pragma once 
+
+#include "Dataflow/DataflowNode.h"
+#include "GeometryCollection/ManagedArrayCollection.h"
+#include "SimulationDefaultConfigNode.generated.h"
+
+class UChaosClothConfig;
+class UChaosClothSharedSimConfig;
+class UMaterialInterface;
+
+USTRUCT(Meta = (DataflowCloth))
+struct FChaosClothAssetSimulationDefaultConfigNode : public FDataflowNode, public FGCObject
+{
+	GENERATED_USTRUCT_BODY()
+	DATAFLOW_NODE_DEFINE_INTERNAL(FChaosClothAssetSimulationDefaultConfigNode, "SimulationDefaultConfig", "Cloth", "Cloth Simulation Default Config")
+	//DATAFLOW_NODE_RENDER_TYPE(FManagedArrayCollection::StaticType(), "Collection")  // TODO: Leave out the render type until there is something to render
+
+public:
+
+	UPROPERTY(Meta = (Dataflowinput, DataflowOutput, DisplayName = "Collection", DataflowPassthrough = "Collection"))
+	FManagedArrayCollection Collection;
+
+	/** Cloth Simulation Properties. */
+	UPROPERTY(EditAnywhere, Instanced, NoClear, Category = "Dataflow")
+	TObjectPtr<UChaosClothConfig> SimulationConfig;
+
+	/** Cloth Shared Simulation Properties. */
+	UPROPERTY(EditAnywhere, Instanced, NoClear, Category = "Dataflow")
+	TObjectPtr<UChaosClothSharedSimConfig> SharedSimulationConfig;
+
+	FChaosClothAssetSimulationDefaultConfigNode(const Dataflow::FNodeParameters& InParam, FGuid InGuid = FGuid::NewGuid());
+
+	virtual void Serialize(FArchive& Ar) override;
+
+	virtual void Evaluate(Dataflow::FContext& Context, const FDataflowOutput* Out) const override;
+
+	//~ Begin FGCObject Interface
+	virtual void AddReferencedObjects(FReferenceCollector& Collector) override;
+	virtual FString GetReferencerName() const override;
+	//~ End FGCObject Interface
+};
diff --git a/Engine/Plugins/Experimental/ChaosClothAssetEditor/Source/ChaosClothAssetDataflowNodes/Public/ChaosClothAsset/SimulationXPBDConfigNode.h b/Engine/Plugins/Experimental/ChaosClothAssetEditor/Source/ChaosClothAssetDataflowNodes/Public/ChaosClothAsset/SimulationXPBDConfigNode.h
new file mode 100644
--- /dev/null
+++ b/Engine/Plugins/Experimental/ChaosClothAssetEditor/Source/ChaosClothAssetDataflowNodes/Public/ChaosClothAsset/SimulationXPBDConfigNode.h
@@ -0,0 +1,93 @@
+// Copyright Epic Games, Inc. All Rights Reserved.
+
+#pragma once 
+
+#include "Dataflow/DataflowNode.h"
+#include "GeometryCollection/ManagedArrayCollection.h"
+#include "ChaosCloth/ChaosClothConfig.h"
+#include "SimulationXPBDConfigNode.generated.h"
+
+class UChaosClothSharedSimConfig;
+class UMaterialInterface;
+namespace Chaos::Softs
+{
+class FCollectionPropertyMutableFacade;
+}
+
+USTRUCT(Meta = (DataflowCloth))
+struct FChaosClothAssetSimulationXPBDConfigNode : public FDataflowNode, public FGCObject
+{
+	GENERATED_USTRUCT_BODY()
+	DATAFLOW_NODE_DEFINE_INTERNAL(FChaosClothAssetSimulationXPBDConfigNode, "SimulationXPBDConfig", "Cloth", "Cloth Simulation XPBD Config")
+	//DATAFLOW_NODE_RENDER_TYPE(FManagedArrayCollection::StaticType(), "Collection")  // TODO: Leave out the render type until there is something to render
+
+public:
+
+	UPROPERTY(Meta = (Dataflowinput, DataflowOutput, DisplayName = "Collection", DataflowPassthrough = "Collection"))
+	FManagedArrayCollection Collection;
+
+	/** Cloth Simulation Properties. */
+	UPROPERTY(EditAnywhere, Category = "Stretch")
+	bool bEnableXPBDStretchConstraints = true;
+
+	UPROPERTY(EditAnywhere, Category = "Stretch", DisplayName = "Stretch Stiffness [kg/s^2]", meta = (UIMin = "0", UIMax = "10000", ClampMin = "0", ClampMax = "10000000", EditCondition = "bEnableXPBDStretchConstraints"))
+	FChaosClothWeightedValue StretchStiffness = { 100.f, 100.f };
+	UPROPERTY(EditAnywhere, Category = "Stretch", DisplayName = "Stretch Damping Ratio", meta = (UIMin = "0", UIMax = "10", ClampMin = "0", ClampMax = "1000", EditCondition = "bEnableXPBDStretchConstraints"))
+	FChaosClothWeightedValue StretchDampingRatio = { 1.f, 1.f };
+
+	UPROPERTY(EditAnywhere, Category = "Bending")
+	bool bEnableXPBDBendConstraints = true;
+	UPROPERTY(EditAnywhere, Category = "Bending")
+	bool bEnableBendAnisotropy = false;
+
+	UPROPERTY(EditAnywhere, Category = "Bending", DisplayName = "Bending Stiffness (Warp) [kg cm/s^2 rad]", meta = (UIMin = "0", UIMax = "10000", ClampMin = "0", ClampMax = "10000000", EditCondition = "bEnableXPBDBendConstraints"))
+	FChaosClothWeightedValue BendingStiffness = { 100.f, 100.f };
+
+	UPROPERTY(EditAnywhere, Category = "Bending", DisplayName = "Bending Stiffness (Weft) [kg cm/s^2 rad]", meta = (UIMin = "0", UIMax = "10000", ClampMin = "0", ClampMax = "10000000", EditCondition = "bEnableXPBDBendConstraints && bEnableBendAnisotropy"))
+	FChaosClothWeightedValue BendingStiffnessWeft = { 100.f, 100.f };
+
+	UPROPERTY(EditAnywhere, Category = "Bending", DisplayName = "Bending Stiffness (Bias) [kg cm/s^2 rad]", meta = (UIMin = "0", UIMax = "10000", ClampMin = "0", ClampMax = "10000000", EditCondition = "bEnableXPBDBendConstraints && bEnableBendAnisotropy"))
+	FChaosClothWeightedValue BendingStiffnessBias = { 100.f, 100.f };
+
+	UPROPERTY(EditAnywhere, Category = "Bending", DisplayName = "Bending Damping Ratio", meta = (UIMin = "0", UIMax = "10", ClampMin = "0", ClampMax = "1000", EditCondition = "bEnableXPBDBendConstraints"))
+	FChaosClothWeightedValue BendingDampingRatio = { 1.f, 1.f };
+		
+	UPROPERTY(EditAnywhere, Category = "Bending", meta = (UIMin = "0", UIMax = "1", ClampMin = "0", ClampMax = "1", EditCondition = "bEnableXPBDBendConstraints"))
+	float BucklingRatio = 0.f;
+
+	UPROPERTY(EditAnywhere, Category = "Bending", DisplayName = "Buckling Stiffness (Warp) [kg cm/s^2 rad]", meta = (UIMin = "0", UIMax = "10000", ClampMin = "0", ClampMax = "10000000", EditCondition = "bEnableXPBDBendConstraints && bEnableXPBDBendConstraints"))
+	FChaosClothWeightedValue BucklingStiffness = { 100.f, 100.f };
+
+	UPROPERTY(EditAnywhere, Category = "Bending", DisplayName = "Buckling Stiffness (Weft) [kg cm/s^2 rad]", meta = (UIMin = "0", UIMax = "10000", ClampMin = "0", ClampMax = "10000000", EditCondition = "bEnableXPBDBendConstraints && bEnableBendAnisotropy"))
+	FChaosClothWeightedValue BucklingStiffnessWeft = { 100.f, 100.f };
+
+	UPROPERTY(EditAnywhere, Category = "Bending", DisplayName = "Buckling Stiffness (Bias) [kg cm/s^2 rad]", meta = (UIMin = "0", UIMax = "10000", ClampMin = "0", ClampMax = "10000000", EditCondition = "bEnableXPBDBendConstraints && bEnableBendAnisotropy"))
+	FChaosClothWeightedValue BucklingStiffnessBias = { 100.f, 100.f };
+
+	UPROPERTY(EditAnywhere, Category = "Area")
+	bool bEnableXPBDAreaConstraints = true;
+	
+	UPROPERTY(EditAnywhere, Category = "Area", DisplayName = "Area Stiffness [kg/s^2]", meta = (UIMin = "0", UIMax = "10000", ClampMin = "0", ClampMax = "10000000", EditCondition = "bEnableXPBDAreaConstraints"))
+	FChaosClothWeightedValue AreaStiffness = { 100.f, 100.f };
+
+	/** Cloth Shared Simulation Properties. */
+	UPROPERTY(EditAnywhere, Instanced, NoClear, Category = "Dataflow")
+	TObjectPtr<UChaosClothSharedSimConfig> SharedSimulationConfig;
+
+	FChaosClothAssetSimulationXPBDConfigNode(const Dataflow::FNodeParameters& InParam, FGuid InGuid = FGuid::NewGuid());
+
+	virtual void Serialize(FArchive& Ar) override;
+
+	virtual void Evaluate(Dataflow::FContext& Context, const FDataflowOutput* Out) const override;
+
+	//~ Begin FGCObject Interface
+	virtual void AddReferencedObjects(FReferenceCollector& Collector) override;
+	virtual FString GetReferencerName() const override;
+	//~ End FGCObject Interface
+
+private:
+	void SetStretchProperties(Chaos::Softs::FCollectionPropertyMutableFacade& Properties) const;
+	void SetBendingProperties(Chaos::Softs::FCollectionPropertyMutableFacade& Properties) const;
+	void SetAreaProperties(Chaos::Softs::FCollectionPropertyMutableFacade& Properties) const;
+
+};
diff --git a/Engine/Plugins/Experimental/ChaosClothAssetEditor/Source/ChaosClothAssetDataflowNodes/Public/ChaosClothAsset/StaticMeshImportNode.h b/Engine/Plugins/Experimental/ChaosClothAssetEditor/Source/ChaosClothAssetDataflowNodes/Public/ChaosClothAsset/StaticMeshImportNode.h
new file mode 100644
--- /dev/null
+++ b/Engine/Plugins/Experimental/ChaosClothAssetEditor/Source/ChaosClothAssetDataflowNodes/Public/ChaosClothAsset/StaticMeshImportNode.h
@@ -0,0 +1,30 @@
+// Copyright Epic Games, Inc. All Rights Reserved.
+
+#pragma once 
+
+#include "Dataflow/DataflowNode.h"
+#include "GeometryCollection/ManagedArrayCollection.h"
+#include "StaticMeshImportNode.generated.h"
+
+class UStaticMesh;
+
+USTRUCT(Meta = (DataflowCloth))
+struct FChaosClothAssetStaticMeshImportNode : public FDataflowNode
+{
+	GENERATED_USTRUCT_BODY()
+	DATAFLOW_NODE_DEFINE_INTERNAL(FChaosClothAssetStaticMeshImportNode, "StaticMeshImport", "Cloth", "Cloth Static Mesh Import")
+public:
+
+	UPROPERTY(Meta = (DataflowOutput, DisplayName = "Collection"))
+	FManagedArrayCollection Collection;
+
+	UPROPERTY(EditAnywhere, Category = "Static Mesh")
+	TObjectPtr<const UStaticMesh> StaticMesh;
+
+	UPROPERTY(EditAnywhere, Category = "Static Mesh")
+	int32 UVChannel = INDEX_NONE;
+
+	FChaosClothAssetStaticMeshImportNode(const Dataflow::FNodeParameters& InParam, FGuid InGuid = FGuid::NewGuid());
+
+	virtual void Evaluate(Dataflow::FContext& Context, const FDataflowOutput* Out) const override;
+};
diff --git a/Engine/Plugins/Experimental/ChaosClothAssetEditor/Source/ChaosClothAssetDataflowNodes/Public/ChaosClothAsset/TerminalNode.h b/Engine/Plugins/Experimental/ChaosClothAssetEditor/Source/ChaosClothAssetDataflowNodes/Public/ChaosClothAsset/TerminalNode.h
new file mode 100644
--- /dev/null
+++ b/Engine/Plugins/Experimental/ChaosClothAssetEditor/Source/ChaosClothAssetDataflowNodes/Public/ChaosClothAsset/TerminalNode.h
@@ -0,0 +1,54 @@
+// Copyright Epic Games, Inc. All Rights Reserved.
+
+#pragma once 
+
+#include "Dataflow/DataflowTerminalNode.h"
+#include "GeometryCollection/ManagedArrayCollection.h"
+#include "TerminalNode.generated.h"
+
+USTRUCT(Meta = (DataflowCloth, DataflowTerminal))
+struct FChaosClothAssetTerminalNode : public FDataflowTerminalNode
+{
+	GENERATED_USTRUCT_BODY()
+	DATAFLOW_NODE_DEFINE_INTERNAL(FChaosClothAssetTerminalNode, "ClothAssetTerminal", "Cloth", "Cloth Terminal")  // TODO: Should the category be Terminal instead like all other terminal nodes
+	//DATAFLOW_NODE_RENDER_TYPE(FGeometryCollection::StaticType(), "Collection")  // TODO: Leave out the render type until there is something to render
+
+public:
+	static constexpr int32 MaxLods = 6;  // Hardcoded number of LODs since it is currently not possible to use arrays for optional inputs
+
+	/** LOD 0 input, right click on the node and add pins to add more LODs. */
+	UPROPERTY(meta = (DataflowInput, DisplayName = "Collection LOD 0"))
+	FManagedArrayCollection CollectionLod0;
+	/** LOD 1 input, right click on the node and add pins to add more LODs. */
+	UPROPERTY(meta = (DisplayName = "Collection LOD 1"))
+	FManagedArrayCollection CollectionLod1;
+	/** LOD 2 input, right click on the node and add pins to add more LODs. */
+	UPROPERTY(meta = (DisplayName = "Collection LOD 2"))
+	FManagedArrayCollection CollectionLod2;
+	/** LOD 3 input, right click on the node and add pins to add more LODs. */
+	UPROPERTY(meta = (DisplayName = "Collection LOD 3"))
+	FManagedArrayCollection CollectionLod3;
+	/** LOD 4 input, right click on the node and add pins to add more LODs. */
+	UPROPERTY(meta = (DisplayName = "Collection LOD 4"))
+	FManagedArrayCollection CollectionLod4;
+	/** LOD 5 input, right click on the node and add pins to add more LODs. */
+	UPROPERTY(meta = (DisplayName = "Collection LOD 5"))
+	FManagedArrayCollection CollectionLod5;
+	/** The number of LODs currently exposed to the node UI. */
+	UPROPERTY()
+	int32 NumLods = 1;
+
+	FChaosClothAssetTerminalNode(const Dataflow::FNodeParameters& InParam, FGuid InGuid = FGuid::NewGuid());
+
+	virtual void SetAssetValue(TObjectPtr<UObject> Asset, Dataflow::FContext& Context) const override;
+
+	virtual void Evaluate(Dataflow::FContext& Context, const FDataflowOutput* Out) const override {}
+
+	virtual Dataflow::FPin AddPin() override;
+	virtual bool CanAddPin() const override { return NumLods < MaxLods; }
+	virtual bool CanRemovePin() const override { return NumLods > 1; }
+	virtual Dataflow::FPin RemovePin() override;
+
+private:
+	TArray<const FManagedArrayCollection*> GetCollectionLods() const;
+};
diff --git a/Engine/Plugins/Experimental/ChaosClothAssetEditor/Source/ChaosClothAssetEditor/ChaosClothAssetEditor.Build.cs b/Engine/Plugins/Experimental/ChaosClothAssetEditor/Source/ChaosClothAssetEditor/ChaosClothAssetEditor.Build.cs
--- a/Engine/Plugins/Experimental/ChaosClothAssetEditor/Source/ChaosClothAssetEditor/ChaosClothAssetEditor.Build.cs
+++ b/Engine/Plugins/Experimental/ChaosClothAssetEditor/Source/ChaosClothAssetEditor/ChaosClothAssetEditor.Build.cs
@@ -55,6 +55,7 @@ public class ChaosClothAssetEditor : ModuleRules
 				"ChaosCloth",
 				"ChaosClothAsset",
 				"ChaosClothAssetEngine",
+				"ChaosClothAssetDataflowNodes",
 				"ClothingSystemEditor",
 				"ChaosClothAssetTools",
 				"ChaosClothAssetEditorTools",
@@ -66,8 +67,6 @@ public class ChaosClothAssetEditor : ModuleRules
 				"TypedElementFramework",
 				"TypedElementRuntime",
 				"AssetDefinition",
-				"DatasmithImporter",
-				"ExternalSource"
 				// ... add private dependencies that you statically link with here ...	
 			}
 			);
diff --git a/Engine/Plugins/Experimental/ChaosClothAssetEditor/Source/ChaosClothAssetEditor/Private/ChaosClothAsset/ChaosClothAssetEditorModule.cpp b/Engine/Plugins/Experimental/ChaosClothAssetEditor/Source/ChaosClothAssetEditor/Private/ChaosClothAsset/ChaosClothAssetEditorModule.cpp
--- a/Engine/Plugins/Experimental/ChaosClothAssetEditor/Source/ChaosClothAssetEditor/Private/ChaosClothAsset/ChaosClothAssetEditorModule.cpp
+++ b/Engine/Plugins/Experimental/ChaosClothAssetEditor/Source/ChaosClothAssetEditor/Private/ChaosClothAsset/ChaosClothAssetEditorModule.cpp
@@ -4,7 +4,6 @@
 #include "ChaosClothAsset/ClothEditorMode.h"
 #include "ChaosClothAsset/ClothEditorCommands.h"
 #include "ChaosClothAsset/ClothEditorStyle.h"
-#include "ChaosClothAsset/DataflowNodes/DataflowNodes.h"
 #include "ContentBrowserMenuContexts.h"
 #include "EditorModeRegistry.h"
 #include "Selection.h"
@@ -27,8 +26,6 @@ void FChaosClothAssetEditorModule::StartupModule()
 	FAssetToolsModule& AssetToolsModule = FAssetToolsModule::GetModule();
 	IAssetTools& AssetTools = AssetToolsModule.Get();
 
-	DataflowNodes::Register();
-
 	// TODO: Register details view customizations
 }
 
diff --git a/Engine/Plugins/Experimental/ChaosClothAssetEditor/Source/ChaosClothAssetEditor/Private/ChaosClothAsset/DataflowNodes/BindToRootBoneNode.cpp b/Engine/Plugins/Experimental/ChaosClothAssetEditor/Source/ChaosClothAssetEditor/Private/ChaosClothAsset/DataflowNodes/BindToRootBoneNode.cpp
deleted file mode 100644
diff --git a/Engine/Plugins/Experimental/ChaosClothAssetEditor/Source/ChaosClothAssetEditor/Private/ChaosClothAsset/DataflowNodes/BindToRootBoneNode.h b/Engine/Plugins/Experimental/ChaosClothAssetEditor/Source/ChaosClothAssetEditor/Private/ChaosClothAsset/DataflowNodes/BindToRootBoneNode.h
deleted file mode 100644
diff --git a/Engine/Plugins/Experimental/ChaosClothAssetEditor/Source/ChaosClothAssetEditor/Private/ChaosClothAsset/DataflowNodes/CopySimulationToRenderMeshNode.cpp b/Engine/Plugins/Experimental/ChaosClothAssetEditor/Source/ChaosClothAssetEditor/Private/ChaosClothAsset/DataflowNodes/CopySimulationToRenderMeshNode.cpp
deleted file mode 100644
diff --git a/Engine/Plugins/Experimental/ChaosClothAssetEditor/Source/ChaosClothAssetEditor/Private/ChaosClothAsset/DataflowNodes/CopySimulationToRenderMeshNode.h b/Engine/Plugins/Experimental/ChaosClothAssetEditor/Source/ChaosClothAssetEditor/Private/ChaosClothAsset/DataflowNodes/CopySimulationToRenderMeshNode.h
deleted file mode 100644
diff --git a/Engine/Plugins/Experimental/ChaosClothAssetEditor/Source/ChaosClothAssetEditor/Private/ChaosClothAsset/DataflowNodes/DataflowNodes.cpp b/Engine/Plugins/Experimental/ChaosClothAssetEditor/Source/ChaosClothAssetEditor/Private/ChaosClothAsset/DataflowNodes/DataflowNodes.cpp
deleted file mode 100644
diff --git a/Engine/Plugins/Experimental/ChaosClothAssetEditor/Source/ChaosClothAssetEditor/Private/ChaosClothAsset/DataflowNodes/DataflowNodes.h b/Engine/Plugins/Experimental/ChaosClothAssetEditor/Source/ChaosClothAssetEditor/Private/ChaosClothAsset/DataflowNodes/DataflowNodes.h
deleted file mode 100644
diff --git a/Engine/Plugins/Experimental/ChaosClothAssetEditor/Source/ChaosClothAssetEditor/Private/ChaosClothAsset/DataflowNodes/DatasmithImportNode.cpp b/Engine/Plugins/Experimental/ChaosClothAssetEditor/Source/ChaosClothAssetEditor/Private/ChaosClothAsset/DataflowNodes/DatasmithImportNode.cpp
deleted file mode 100644
diff --git a/Engine/Plugins/Experimental/ChaosClothAssetEditor/Source/ChaosClothAssetEditor/Private/ChaosClothAsset/DataflowNodes/DatasmithImportNode.h b/Engine/Plugins/Experimental/ChaosClothAssetEditor/Source/ChaosClothAssetEditor/Private/ChaosClothAsset/DataflowNodes/DatasmithImportNode.h
deleted file mode 100644
diff --git a/Engine/Plugins/Experimental/ChaosClothAssetEditor/Source/ChaosClothAssetEditor/Private/ChaosClothAsset/DataflowNodes/DeleteRenderMeshNode.cpp b/Engine/Plugins/Experimental/ChaosClothAssetEditor/Source/ChaosClothAssetEditor/Private/ChaosClothAsset/DataflowNodes/DeleteRenderMeshNode.cpp
deleted file mode 100644
diff --git a/Engine/Plugins/Experimental/ChaosClothAssetEditor/Source/ChaosClothAssetEditor/Private/ChaosClothAsset/DataflowNodes/DeleteRenderMeshNode.h b/Engine/Plugins/Experimental/ChaosClothAssetEditor/Source/ChaosClothAssetEditor/Private/ChaosClothAsset/DataflowNodes/DeleteRenderMeshNode.h
deleted file mode 100644
diff --git a/Engine/Plugins/Experimental/ChaosClothAssetEditor/Source/ChaosClothAssetEditor/Private/ChaosClothAsset/DataflowNodes/ImportNode.cpp b/Engine/Plugins/Experimental/ChaosClothAssetEditor/Source/ChaosClothAssetEditor/Private/ChaosClothAsset/DataflowNodes/ImportNode.cpp
deleted file mode 100644
diff --git a/Engine/Plugins/Experimental/ChaosClothAssetEditor/Source/ChaosClothAssetEditor/Private/ChaosClothAsset/DataflowNodes/ImportNode.h b/Engine/Plugins/Experimental/ChaosClothAssetEditor/Source/ChaosClothAssetEditor/Private/ChaosClothAsset/DataflowNodes/ImportNode.h
deleted file mode 100644
diff --git a/Engine/Plugins/Experimental/ChaosClothAssetEditor/Source/ChaosClothAssetEditor/Private/ChaosClothAsset/DataflowNodes/PatternSelectionNode.cpp b/Engine/Plugins/Experimental/ChaosClothAssetEditor/Source/ChaosClothAssetEditor/Private/ChaosClothAsset/DataflowNodes/PatternSelectionNode.cpp
deleted file mode 100644
diff --git a/Engine/Plugins/Experimental/ChaosClothAssetEditor/Source/ChaosClothAssetEditor/Private/ChaosClothAsset/DataflowNodes/PatternSelectionNode.h b/Engine/Plugins/Experimental/ChaosClothAssetEditor/Source/ChaosClothAssetEditor/Private/ChaosClothAsset/DataflowNodes/PatternSelectionNode.h
deleted file mode 100644
diff --git a/Engine/Plugins/Experimental/ChaosClothAssetEditor/Source/ChaosClothAssetEditor/Private/ChaosClothAsset/DataflowNodes/ReverseNormalsNode.cpp b/Engine/Plugins/Experimental/ChaosClothAssetEditor/Source/ChaosClothAssetEditor/Private/ChaosClothAsset/DataflowNodes/ReverseNormalsNode.cpp
deleted file mode 100644
diff --git a/Engine/Plugins/Experimental/ChaosClothAssetEditor/Source/ChaosClothAssetEditor/Private/ChaosClothAsset/DataflowNodes/ReverseNormalsNode.h b/Engine/Plugins/Experimental/ChaosClothAssetEditor/Source/ChaosClothAssetEditor/Private/ChaosClothAsset/DataflowNodes/ReverseNormalsNode.h
deleted file mode 100644
diff --git a/Engine/Plugins/Experimental/ChaosClothAssetEditor/Source/ChaosClothAssetEditor/Private/ChaosClothAsset/DataflowNodes/SimulationDefaultConfigNode.cpp b/Engine/Plugins/Experimental/ChaosClothAssetEditor/Source/ChaosClothAssetEditor/Private/ChaosClothAsset/DataflowNodes/SimulationDefaultConfigNode.cpp
deleted file mode 100644
diff --git a/Engine/Plugins/Experimental/ChaosClothAssetEditor/Source/ChaosClothAssetEditor/Private/ChaosClothAsset/DataflowNodes/SimulationDefaultConfigNode.h b/Engine/Plugins/Experimental/ChaosClothAssetEditor/Source/ChaosClothAssetEditor/Private/ChaosClothAsset/DataflowNodes/SimulationDefaultConfigNode.h
deleted file mode 100644
diff --git a/Engine/Plugins/Experimental/ChaosClothAssetEditor/Source/ChaosClothAssetEditor/Private/ChaosClothAsset/DataflowNodes/SimulationXPBDConfigNode.cpp b/Engine/Plugins/Experimental/ChaosClothAssetEditor/Source/ChaosClothAssetEditor/Private/ChaosClothAsset/DataflowNodes/SimulationXPBDConfigNode.cpp
deleted file mode 100644
diff --git a/Engine/Plugins/Experimental/ChaosClothAssetEditor/Source/ChaosClothAssetEditor/Private/ChaosClothAsset/DataflowNodes/SimulationXPBDConfigNode.h b/Engine/Plugins/Experimental/ChaosClothAssetEditor/Source/ChaosClothAssetEditor/Private/ChaosClothAsset/DataflowNodes/SimulationXPBDConfigNode.h
deleted file mode 100644
diff --git a/Engine/Plugins/Experimental/ChaosClothAssetEditor/Source/ChaosClothAssetEditor/Private/ChaosClothAsset/DataflowNodes/StaticMeshImportNode.cpp b/Engine/Plugins/Experimental/ChaosClothAssetEditor/Source/ChaosClothAssetEditor/Private/ChaosClothAsset/DataflowNodes/StaticMeshImportNode.cpp
deleted file mode 100644
diff --git a/Engine/Plugins/Experimental/ChaosClothAssetEditor/Source/ChaosClothAssetEditor/Private/ChaosClothAsset/DataflowNodes/StaticMeshImportNode.h b/Engine/Plugins/Experimental/ChaosClothAssetEditor/Source/ChaosClothAssetEditor/Private/ChaosClothAsset/DataflowNodes/StaticMeshImportNode.h
deleted file mode 100644
diff --git a/Engine/Plugins/Experimental/ChaosClothAssetEditor/Source/ChaosClothAssetEditor/Private/ChaosClothAsset/DataflowNodes/TerminalNode.cpp b/Engine/Plugins/Experimental/ChaosClothAssetEditor/Source/ChaosClothAssetEditor/Private/ChaosClothAsset/DataflowNodes/TerminalNode.cpp
deleted file mode 100644
diff --git a/Engine/Plugins/Experimental/ChaosClothAssetEditor/Source/ChaosClothAssetEditor/Private/ChaosClothAsset/DataflowNodes/TerminalNode.h b/Engine/Plugins/Experimental/ChaosClothAssetEditor/Source/ChaosClothAssetEditor/Private/ChaosClothAsset/DataflowNodes/TerminalNode.h
deleted file mode 100644
