ClothTrainingTool: add skinning to training data generation

diff --git a/Engine/Plugins/Experimental/ChaosClothAsset/Source/ChaosClothAssetEngine/Private/ChaosClothAsset/ClothComponent.cpp b/Engine/Plugins/Experimental/ChaosClothAsset/Source/ChaosClothAssetEngine/Private/ChaosClothAsset/ClothComponent.cpp
--- a/Engine/Plugins/Experimental/ChaosClothAsset/Source/ChaosClothAssetEngine/Private/ChaosClothAsset/ClothComponent.cpp
+++ b/Engine/Plugins/Experimental/ChaosClothAsset/Source/ChaosClothAssetEngine/Private/ChaosClothAsset/ClothComponent.cpp
@@ -369,3 +369,20 @@ bool UChaosClothComponent::ShouldWaitForParallelSimulationInTickComponent() cons
 
 	return bWaitForParallelTask || (CVarClothPhysicsWaitForParallelClothTask && CVarClothPhysicsWaitForParallelClothTask->GetBool());
 }
+
+void UChaosClothComponent::Pose(const TArray<FTransform>& InComponentSpaceTransforms)
+{
+	const bool bIsInputValid = InComponentSpaceTransforms.Num() == GetComponentSpaceTransforms().Num();
+	ensure(bIsInputValid);
+	if (!bIsInputValid)
+	{
+		return;
+	}
+	GetEditableComponentSpaceTransforms() = InComponentSpaceTransforms;
+	bNeedToFlipSpaceBaseBuffers = true;
+	FinalizeBoneTransform();
+
+	UpdateBounds();
+	MarkRenderTransformDirty();
+	MarkRenderDynamicDataDirty();
+}
\ No newline at end of file
diff --git a/Engine/Plugins/Experimental/ChaosClothAsset/Source/ChaosClothAssetEngine/Public/ChaosClothAsset/ClothComponent.h b/Engine/Plugins/Experimental/ChaosClothAsset/Source/ChaosClothAssetEngine/Public/ChaosClothAsset/ClothComponent.h
--- a/Engine/Plugins/Experimental/ChaosClothAsset/Source/ChaosClothAssetEngine/Public/ChaosClothAsset/ClothComponent.h
+++ b/Engine/Plugins/Experimental/ChaosClothAsset/Source/ChaosClothAssetEngine/Public/ChaosClothAsset/ClothComponent.h
@@ -85,6 +85,9 @@ public:
 	 */
 	TSharedPtr<const FManagedArrayCollection> GetPropertyCollection() const { return TSharedPtr<const FManagedArrayCollection>(PropertyCollection); }
 
+	/** Pose the cloth component using component space transforms. */
+	void Pose(const TArray<FTransform>& InComponentSpaceTransforms);
+
 	friend UE::Chaos::ClothAsset::FClothComponentCacheAdapter;
 
 protected:
diff --git a/Engine/Plugins/Experimental/ChaosClothAssetEditor/Source/ChaosClothAssetEditorTools/Private/ChaosClothAsset/ClothTrainingTool.cpp b/Engine/Plugins/Experimental/ChaosClothAssetEditor/Source/ChaosClothAssetEditorTools/Private/ChaosClothAsset/ClothTrainingTool.cpp
--- a/Engine/Plugins/Experimental/ChaosClothAssetEditor/Source/ChaosClothAssetEditorTools/Private/ChaosClothAsset/ClothTrainingTool.cpp
+++ b/Engine/Plugins/Experimental/ChaosClothAssetEditor/Source/ChaosClothAssetEditorTools/Private/ChaosClothAsset/ClothTrainingTool.cpp
@@ -2,6 +2,9 @@
 
 #include "ChaosClothAsset/ClothTrainingTool.h"
 #include "CoreMinimal.h"
+
+#include "Animation/AttributesRuntime.h"
+#include "BonePose.h"
 #include "ChaosCloth/ChaosClothingSimulationConfig.h"
 #include "ChaosClothAsset/ClothComponent.h"
 #include "ChaosClothAsset/ClothComponentToolTarget.h"
@@ -19,6 +22,7 @@
 #include "UObject/SavePackage.h"
 
 #define LOCTEXT_NAMESPACE "ClothTrainingTool"
+DEFINE_LOG_CATEGORY_STATIC(LogClothTrainingTool, Log, All);
 
 class UClothTrainingTool::FClothSimulationDataGenerationProxy : public UE::Chaos::ClothAsset::FClothSimulationProxy
 {
@@ -56,8 +60,16 @@ public:
 
 private:
 	void Simulate(float DeltaTime, int32 NumSteps);
+	void PrepareAnimSequence();
+	TArray<FTransform> GetBoneTransforms(int32 Frame) const;
+	bool IsClothComponentValid() const;
+	bool IsTransformsValid(const TArray<FTransform> &Transforms) const;
+	bool GetSimPositions(TArray<FVector3f>& OutPositions) const;
+	void AddToCache(UChaosCache* InCache, int32 Frame, const TArray<FVector3f>& Positions) const;
+
 	void BackupClothComponentState();
 	void RestoreClothComponentState();
+
 	const TObjectPtr<UAnimSequence> AnimSequence;
 	TObjectPtr<UChaosCache> Cache;
 	TObjectPtr<UChaosClothComponent> ClothComponent;
@@ -66,25 +78,109 @@ private:
 	bool bIsSimulationSuspendedBackup = false;
 	bool bTeleportBackup = false;
 	bool bResetBackup = false;
+	TArray<FTransform> ComponentSpaceTransformsBackup;
+
+	inline static const FName PositionXName = TEXT("PositionX");
+	inline static const FName PositionYName = TEXT("PositionY");
+	inline static const FName PositionZName = TEXT("PositionZ");
 };
 
+void UClothTrainingTool::FGenerateClothOp::PrepareAnimSequence()
+{
+	if (AnimSequence)
+	{
+		AnimSequence->Interpolation = EAnimInterpolationType::Step;
+	}
+}
+
+TArray<FTransform> UClothTrainingTool::FGenerateClothOp::GetBoneTransforms(int32 Frame) const
+{
+	const double Time = AnimSequence->GetTimeAtFrame(Frame);
+	FAnimExtractContext ExtractionContext(Time);
+
+	UChaosClothAsset* const ClothAsset = ClothComponent->GetClothAsset();
+	const FReferenceSkeleton* const ReferenceSkeleton = ClothAsset ? &ClothAsset->GetRefSkeleton() : nullptr;
+	USkeleton* const Skeleton = ClothAsset ? ClothAsset->GetSkeleton() : nullptr;
+	const int32 NumBones = ReferenceSkeleton ? ReferenceSkeleton->GetNum() : 0;
+
+	TArray<uint16> BoneIndices;
+	BoneIndices.SetNumUninitialized(NumBones);
+	for (int32 Index = 0; Index < NumBones; ++Index)
+	{
+		BoneIndices[Index] = (uint16)Index;
+	}
+
+	FBoneContainer BoneContainer;
+	BoneContainer.SetUseRAWData(true);
+	BoneContainer.InitializeTo(BoneIndices, UE::Anim::FCurveFilterSettings(), *Skeleton);
+
+	FCompactPose OutPose;
+	OutPose.SetBoneContainer(&BoneContainer);
+	FBlendedCurve OutCurve;
+	OutCurve.InitFrom(BoneContainer);
+	UE::Anim::FStackAttributeContainer TempAttributes;
+
+	FAnimationPoseData AnimationPoseData(OutPose, OutCurve, TempAttributes);
+	AnimSequence->GetAnimationPose(AnimationPoseData, ExtractionContext);
+
+	TArray<FTransform> ComponentSpaceTransforms;
+	ComponentSpaceTransforms.SetNumUninitialized(NumBones);
+	for (int32 Index = 0; Index < NumBones; ++Index)
+	{
+		const FCompactPoseBoneIndex CompactIndex = BoneContainer.MakeCompactPoseIndex(FMeshPoseBoneIndex(Index));
+		const int32 ParentIndex = ReferenceSkeleton->GetParentIndex(Index);
+		ComponentSpaceTransforms[Index] = 
+			ComponentSpaceTransforms.IsValidIndex(ParentIndex) && ParentIndex < Index ? 
+			AnimationPoseData.GetPose()[CompactIndex] * ComponentSpaceTransforms[ParentIndex] : 
+			ReferenceSkeleton->GetRefBonePose()[Index];
+	}
+
+	return ComponentSpaceTransforms;
+}
+void UClothTrainingTool::FGenerateClothOp::AddToCache(UChaosCache* InCache, int32 Frame, const TArray<FVector3f>& Positions) const
+{
+	constexpr float CacheFPS = 30;
+	const float Time = Frame / CacheFPS;
+	FPendingFrameWrite NewFrame;
+	NewFrame.Time = Time;
+
+	const int32 NumParticles = Positions.Num();
+	TArray<int32>& PendingID = NewFrame.PendingChannelsIndices;
+	TArray<float> PendingPX, PendingPY, PendingPZ;
+	TArray<float> PendingVX, PendingVY, PendingVZ;
+	PendingID.SetNum(NumParticles);
+	PendingPX.SetNum(NumParticles);
+	PendingPY.SetNum(NumParticles);
+	PendingPZ.SetNum(NumParticles);
+
+	for (int32 ParticleIndex = 0; ParticleIndex < NumParticles; ++ParticleIndex)
+	{
+		const FVector3f& Position = Positions[ParticleIndex];
+		PendingID[ParticleIndex] = ParticleIndex;
+		PendingPX[ParticleIndex] = Position.X;
+		PendingPY[ParticleIndex] = Position.Y;
+		PendingPZ[ParticleIndex] = Position.Z;
+	}
+
+	NewFrame.PendingChannelsData.Add(PositionXName, PendingPX);
+	NewFrame.PendingChannelsData.Add(PositionYName, PendingPY);
+	NewFrame.PendingChannelsData.Add(PositionZName, PendingPZ);
+
+	Cache->AddFrame_Concurrent(MoveTemp(NewFrame));
+}
+
 void UClothTrainingTool::FGenerateClothOp::CalculateResult(FProgressCancel* Progress)
 {
-	if (ClothComponent == nullptr || Cache == nullptr || DataGenerationProxy == nullptr)
+	if (ClothComponent == nullptr || Cache == nullptr || DataGenerationProxy == nullptr || !IsClothComponentValid())
 	{
 		return;
 	}
-
-	static const FName PositionXName = TEXT("PositionX");
-	static const FName PositionYName = TEXT("PositionY");
-	static const FName PositionZName = TEXT("PositionZ");
 	
 	bool bCancelled = false;
-	// TODO: change this part to data from AnimSequence
 	constexpr float DeltaTime = 1e-3;
-	constexpr float CacheFPS = 30;
 	const int32 NumFrames = 10;
 
+	PrepareAnimSequence();
 	BackupClothComponentState();
 
 	ClothComponent->ResumeSimulation();
@@ -98,51 +194,24 @@ void UClothTrainingTool::FGenerateClothOp::CalculateResult(FProgressCancel* Prog
 				break;
 			}
 
-			FPendingFrameWrite NewFrame;
-			const float TotalTime = Frame * DeltaTime * 5;
-			Simulate(DeltaTime, (Frame + 1) * 5);
+			const float NumSteps = 200;
 
-			const TMap<int32, FClothSimulData> &SimulDataMap = DataGenerationProxy->GetCurrentSimulationData_AnyThread();
-			const FClothSimulData* const SimulData = SimulDataMap.Find(0);
-			if (SimulDataMap.Num() > 1)
-			{
-				ensureMsgf(false, TEXT("Only support single cloth for now."));
-				continue;
-			}
-			if (SimulData == nullptr)
+			TArray<FTransform> BoneTransforms = GetBoneTransforms(Frame * 50);
+			if (!IsTransformsValid(BoneTransforms))
 			{
-				ensureMsgf(false, TEXT("ClothSimulData is nullptr"));
-				continue;
+				break;
 			}
 
-			const TArray<FVector3f>& SimPositions = SimulData->Positions;
-			const int32 NumParticles = SimPositions.Num();
-
-			const float Time = Frame / CacheFPS;
-			NewFrame.Time = Time;
-
-			TArray<int32>& PendingID = NewFrame.PendingChannelsIndices;
-			TArray<float> PendingPX, PendingPY, PendingPZ;
-			TArray<float> PendingVX, PendingVY, PendingVZ;
-			PendingID.SetNum(NumParticles);
-			PendingPX.SetNum(NumParticles);
-			PendingPY.SetNum(NumParticles);
-			PendingPZ.SetNum(NumParticles);
+			ClothComponent->Pose(BoneTransforms);
+			Simulate(DeltaTime, NumSteps);
 
-			for (int32 ParticleIndex = 0; ParticleIndex < NumParticles; ++ParticleIndex)
+			TArray<FVector3f> SimPositions;
+			if (!GetSimPositions(SimPositions))
 			{
-				const FVector3f& Position = SimPositions[ParticleIndex];
-				PendingID[ParticleIndex] = ParticleIndex;
-				PendingPX[ParticleIndex] = Position.X;
-				PendingPY[ParticleIndex] = Position.Y;
-				PendingPZ[ParticleIndex] = Position.Z + Frame;
+				break;
 			}
+			AddToCache(Cache, Frame, SimPositions);
 
-			NewFrame.PendingChannelsData.Add(PositionXName, PendingPX);
-			NewFrame.PendingChannelsData.Add(PositionYName, PendingPY);
-			NewFrame.PendingChannelsData.Add(PositionZName, PendingPZ);
-
-			Cache->AddFrame_Concurrent(MoveTemp(NewFrame));
 			Progress->AdvanceCurrentScopeProgressBy(1.f / NumFrames);
 		}
 	}
@@ -150,7 +219,7 @@ void UClothTrainingTool::FGenerateClothOp::CalculateResult(FProgressCancel* Prog
 
 	if (!bCancelled)
 	{
-		UE_LOG(LogTemp, Display, TEXT("Data generation complete."));
+		UE_LOG(LogClothTrainingTool, Display, TEXT("Data generation complete."));
 	}
 }
 
@@ -173,11 +242,36 @@ void UClothTrainingTool::FGenerateClothOp::Simulate(float DeltaTime, int32 NumSt
 	DataGenerationProxy->WriteSimulationData();
 }
 
+bool UClothTrainingTool::FGenerateClothOp::GetSimPositions(TArray<FVector3f> &OutPositions) const
+{
+	const TMap<int32, FClothSimulData> &SimulDataMap = DataGenerationProxy->GetCurrentSimulationData_AnyThread();
+	const FClothSimulData* const SimulData = SimulDataMap.Find(0);
+	if (SimulDataMap.Num() > 1)
+	{
+		ensureMsgf(false, TEXT("Multiple cloth is not yet supported."));
+		return false;
+	}
+	if (SimulData == nullptr)
+	{
+		ensureMsgf(false, TEXT("ClothSimulData is nullptr"));
+		return false;
+	}
+
+	const TArray<FVector3f>& SimPositions = SimulData->Positions;
+	OutPositions.SetNum(SimPositions.Num());
+	for (int32 Index = 0; Index < SimPositions.Num(); ++Index)
+	{
+		OutPositions[Index] = FVector3f(SimulData->ComponentRelativeTransform.TransformPosition(FVector(SimPositions[Index])));
+	}
+	return true;
+}
+
 void UClothTrainingTool::FGenerateClothOp::BackupClothComponentState()
 {
 	bIsSimulationSuspendedBackup = ClothComponent->IsSimulationSuspended();
 	bTeleportBackup = ClothComponent->NeedsTeleport();
 	bResetBackup = ClothComponent->NeedsReset();
+	ComponentSpaceTransformsBackup = ClothComponent->GetComponentSpaceTransforms();
 }
 
 void UClothTrainingTool::FGenerateClothOp::RestoreClothComponentState()
@@ -195,8 +289,34 @@ void UClothTrainingTool::FGenerateClothOp::RestoreClothComponentState()
 	{
 		ClothComponent->ResetTeleportMode();
 	}
+	ClothComponent->Pose(ComponentSpaceTransformsBackup);
 }
 
+bool UClothTrainingTool::FGenerateClothOp::IsClothComponentValid() const
+{
+	if (USkinnedMeshComponent* const LeaderComponent = ClothComponent->LeaderPoseComponent.Get())
+	{
+		UE_LOG(LogClothTrainingTool, Error, TEXT("Leader pose component is not supported yet."));
+		return false;
+	}
+	else
+	{
+		return true;
+	}
+}
+
+bool UClothTrainingTool::FGenerateClothOp::IsTransformsValid(const TArray<FTransform>& Transforms) const
+{
+	if (Transforms.Num() != ComponentSpaceTransformsBackup.Num())
+	{
+		UE_LOG(LogClothTrainingTool, Error, TEXT("The number of bones in the AnimSequence does not match the number of bones in the ClothAsset. Try an AnimSequence with the same skeleton as the ClothAsset."));
+		return false;
+	}
+	else
+	{
+		return true;
+	}
+}
 
 // ------------------- Properties -------------------
 
@@ -208,7 +328,6 @@ void UClothTrainingToolActionProperties::PostAction(EClothTrainingToolActions Ac
 	}
 }
 
-
 // ------------------- Builder -------------------
 
 const FToolTargetTypeRequirements& UClothTrainingToolBuilder::GetTargetRequirements() const
@@ -284,6 +403,7 @@ void UClothTrainingTool::RunTraining()
 	{
 		DataGenerationProxy = MakeUnique<FClothSimulationDataGenerationProxy>(*ClothComponent);
 	}
+	FCacheUserToken CacheUserToken = Cache->BeginRecord(ClothComponent, FGuid(), FTransform::Identity);
 
 	TUniquePtr<FGenerateClothOp> NewOp = MakeUnique<FGenerateClothOp>(ToolProperties->AnimationSequence, Cache, ClothComponent, DataGenerationProxy.Get());
 
@@ -294,7 +414,6 @@ void UClothTrainingTool::RunTraining()
 	SlowTask.MakeDialog(true);
 
 	bool bSuccess = false;
-	FCacheUserToken CacheUserToken = Cache->BeginRecord(ClothComponent, FGuid(), FTransform::Identity);
 	while (true)
 	{
 		if (SlowTask.ShouldCancel())
