ClothingTrainingTool: add training data debug tool
add deltatime and numsteps options
save tool properties for easy iteration

diff --git a/Engine/Plugins/Experimental/ChaosClothAssetEditor/Source/ChaosClothAssetEditor/Private/ChaosClothAsset/ClothEditorCommands.cpp b/Engine/Plugins/Experimental/ChaosClothAssetEditor/Source/ChaosClothAssetEditor/Private/ChaosClothAsset/ClothEditorCommands.cpp
--- a/Engine/Plugins/Experimental/ChaosClothAssetEditor/Source/ChaosClothAssetEditor/Private/ChaosClothAsset/ClothEditorCommands.cpp
+++ b/Engine/Plugins/Experimental/ChaosClothAssetEditor/Source/ChaosClothAssetEditor/Private/ChaosClothAsset/ClothEditorCommands.cpp
@@ -38,7 +38,7 @@ void FChaosClothAssetEditorCommands::RegisterCommands()
 	UI_COMMAND(BeginAttributeEditorTool, "AttrEd", "Edit/configure mesh attributes", EUserInterfaceActionType::Button, FInputChord());
 	UI_COMMAND(BeginWeightMapPaintTool, "MapPnt", "Paint Weight Maps on the mesh", EUserInterfaceActionType::Button, FInputChord());
 
-	UI_COMMAND(BeginClothTrainingTool, "Train", "Launch Cloth Training tool", EUserInterfaceActionType::Button, FInputChord());
+	UI_COMMAND(BeginClothTrainingTool, "Generate Train Data", "Generate cloth training data", EUserInterfaceActionType::Button, FInputChord());
 
 	UI_COMMAND(BeginTransferSkinWeightsTool, "Transfer Skin Weights", "Launch Transfer Skin Weights tool", EUserInterfaceActionType::Button, FInputChord());
 
diff --git a/Engine/Plugins/Experimental/ChaosClothAssetEditor/Source/ChaosClothAssetEditorTools/Private/ChaosClothAsset/ClothTrainingTool.cpp b/Engine/Plugins/Experimental/ChaosClothAssetEditor/Source/ChaosClothAssetEditorTools/Private/ChaosClothAsset/ClothTrainingTool.cpp
--- a/Engine/Plugins/Experimental/ChaosClothAssetEditor/Source/ChaosClothAssetEditorTools/Private/ChaosClothAsset/ClothTrainingTool.cpp
+++ b/Engine/Plugins/Experimental/ChaosClothAssetEditor/Source/ChaosClothAssetEditorTools/Private/ChaosClothAsset/ClothTrainingTool.cpp
@@ -5,19 +5,22 @@
 
 #include "Animation/AttributesRuntime.h"
 #include "BonePose.h"
+#include "Chaos/ChaosCache.h"
+#include "Chaos/CacheCollection.h"
 #include "ChaosCloth/ChaosClothingSimulationConfig.h"
+#include "ChaosClothAsset/ClothAsset.h"
 #include "ChaosClothAsset/ClothComponent.h"
 #include "ChaosClothAsset/ClothComponentToolTarget.h"
 #include "ChaosClothAsset/ClothCollection.h"
 #include "ChaosClothAsset/ClothSimulationProxy.h"
 #include "ClothingSystemRuntimeTypes.h"
-#include "Chaos/ChaosCache.h"
-#include "Chaos/CacheCollection.h"
 #include "ComponentReregisterContext.h"
-#include "ModelingOperators.h"
-#include "Misc/ScopedSlowTask.h"
 #include "ContextObjectStore.h"
 #include "InteractiveToolManager.h"
+#include "Internationalization/Regex.h"
+#include "ModelingOperators.h"
+#include "Misc/ScopedSlowTask.h"
+#include "PhysicsEngine/PhysicsAsset.h"
 #include "ToolTargetManager.h"
 #include "UObject/SavePackage.h"
 
@@ -45,35 +48,106 @@ UClothTrainingTool::FClothSimulationDataGenerationProxy::~FClothSimulationDataGe
 {
 }
 
+namespace UE::ClothTrainingTool::Private
+{
+	UChaosCache* GetCache(UChaosCacheCollection* CacheCollection)
+	{
+		static const FName CacheName = FName("SimulatedCache");
+		return CacheCollection ? CacheCollection->FindOrAddCache(CacheName) : nullptr;
+	}
+
+	TArray<int32> ParseFrames(const FString& FramesString)
+	{
+		TArray<int32> Result;
+		static const FRegexPattern AllowedCharsPattern(TEXT("^[-,0-9\\s]+$"));
+
+		if (!FRegexMatcher(AllowedCharsPattern, FramesString).FindNext())
+		{
+		    UE_LOG(LogClothTrainingTool, Error, TEXT("Input contains invalid characters."));
+		    return Result;
+		}
+
+		static const FRegexPattern SingleNumberPattern(TEXT("^\\s*(\\d+)\\s*$"));
+		static const FRegexPattern RangePattern(TEXT("^\\s*(\\d+)\\s*-\\s*(\\d+)\\s*$"));
+
+		TArray<FString> Segments;
+	    FramesString.ParseIntoArray(Segments, TEXT(","), true);
+	    for (const FString& Segment : Segments)
+	    {
+	    	bool bSegmentValid = false;
+
+	    	FRegexMatcher SingleNumberMatcher(SingleNumberPattern, Segment);
+	    	if (SingleNumberMatcher.FindNext())
+	    	{
+	    	    const int32 SingleNumber = FCString::Atoi(*SingleNumberMatcher.GetCaptureGroup(1));
+	    	    Result.Add(SingleNumber);
+	    	    bSegmentValid = true;
+	    	}
+	    	else
+	    	{
+	    		FRegexMatcher RangeMatcher(RangePattern, Segment);
+	    		if (RangeMatcher.FindNext())
+	    		{
+	    		    const int32 RangeStart = FCString::Atoi(*RangeMatcher.GetCaptureGroup(1));
+	    		    const int32 RangeEnd = FCString::Atoi(*RangeMatcher.GetCaptureGroup(2));
+
+	    		    for (int32 i = RangeStart; i <= RangeEnd; ++i)
+	    		    {
+	    		        Result.Add(i);
+	    		    }
+	    		    bSegmentValid = true;
+	    		}
+	    	}
+	    	
+	    	if (!bSegmentValid)
+	    	{
+	    	    UE_LOG(LogClothTrainingTool, Error, TEXT("Invalid format in segment: %s"), *Segment);
+	    	}
+	    }
+
+		return Result;
+	}
+
+	TArray<int32> Range(int32 End)
+	{
+		TArray<int32> Result;
+		Result.Reserve(End);
+		for (int32 Index = 0; Index < End; Index++)
+		{
+			Result.Add(Index);
+		}
+		return Result;
+	}
+};
+
 class UClothTrainingTool::FGenerateClothOp : public UE::Geometry::TGenericDataOperator<FSkinnedMeshVertices>
 {
 public: 
-	FGenerateClothOp(TObjectPtr<UAnimSequence> InAnimSequence, TObjectPtr<UChaosCache> InCache, TObjectPtr<UChaosClothComponent> InClothComponent, FClothSimulationDataGenerationProxy* InDataGenerationProxy)
-		: AnimSequence(InAnimSequence)
-		, Cache(InCache)
-		, ClothComponent(InClothComponent)
+	FGenerateClothOp(TObjectPtr<UChaosClothComponent> InClothComponent, FClothSimulationDataGenerationProxy* InDataGenerationProxy, TObjectPtr<UClothTrainingToolProperties> InToolProperties)
+		: ClothComponent(InClothComponent)
 		, DataGenerationProxy(InDataGenerationProxy)
+		, ToolProperties(InToolProperties)
 	{
 	}
 
 	virtual void CalculateResult(FProgressCancel* Progress) override;
 
 private:
-	void Simulate(float DeltaTime, int32 NumSteps);
-	void PrepareAnimSequence();
+	void Simulate(float DeltaTime, int32 NumSteps, int32 PoseIndex, FProgressCancel* Progress, bool bDebug = false);
+	void PrepareAnimationSequence();
 	TArray<FTransform> GetBoneTransforms(int32 Frame) const;
 	bool IsClothComponentValid() const;
 	bool IsTransformsValid(const TArray<FTransform> &Transforms) const;
 	bool GetSimPositions(TArray<FVector3f>& OutPositions) const;
+	void AddToCache(UChaosCache* InCache, int32 Frame);
 	void AddToCache(UChaosCache* InCache, int32 Frame, const TArray<FVector3f>& Positions) const;
 
 	void BackupClothComponentState();
 	void RestoreClothComponentState();
 
-	const TObjectPtr<UAnimSequence> AnimSequence;
-	TObjectPtr<UChaosCache> Cache;
-	TObjectPtr<UChaosClothComponent> ClothComponent;
+	TObjectPtr<UChaosClothComponent> ClothComponent = nullptr;
 	FClothSimulationDataGenerationProxy *DataGenerationProxy = nullptr;
+	TObjectPtr<UClothTrainingToolProperties> ToolProperties = nullptr;
 
 	bool bIsSimulationSuspendedBackup = false;
 	bool bTeleportBackup = false;
@@ -85,17 +159,20 @@ private:
 	inline static const FName PositionZName = TEXT("PositionZ");
 };
 
-void UClothTrainingTool::FGenerateClothOp::PrepareAnimSequence()
+void UClothTrainingTool::FGenerateClothOp::PrepareAnimationSequence()
 {
-	if (AnimSequence)
+	TObjectPtr<UAnimSequence> AnimationSequence = ToolProperties->AnimationSequence;
+	if (AnimationSequence)
 	{
-		AnimSequence->Interpolation = EAnimInterpolationType::Step;
+		AnimationSequence->Interpolation = EAnimInterpolationType::Step;
 	}
 }
 
 TArray<FTransform> UClothTrainingTool::FGenerateClothOp::GetBoneTransforms(int32 Frame) const
 {
-	const double Time = AnimSequence->GetTimeAtFrame(Frame);
+	TObjectPtr<UAnimSequence> AnimationSequence = ToolProperties->AnimationSequence;
+
+	const double Time = AnimationSequence->GetTimeAtFrame(Frame);
 	FAnimExtractContext ExtractionContext(Time);
 
 	UChaosClothAsset* const ClothAsset = ClothComponent->GetClothAsset();
@@ -121,7 +198,7 @@ TArray<FTransform> UClothTrainingTool::FGenerateClothOp::GetBoneTransforms(int32
 	UE::Anim::FStackAttributeContainer TempAttributes;
 
 	FAnimationPoseData AnimationPoseData(OutPose, OutCurve, TempAttributes);
-	AnimSequence->GetAnimationPose(AnimationPoseData, ExtractionContext);
+	AnimationSequence->GetAnimationPose(AnimationPoseData, ExtractionContext);
 
 	TArray<FTransform> ComponentSpaceTransforms;
 	ComponentSpaceTransforms.SetNumUninitialized(NumBones);
@@ -137,6 +214,18 @@ TArray<FTransform> UClothTrainingTool::FGenerateClothOp::GetBoneTransforms(int32
 
 	return ComponentSpaceTransforms;
 }
+
+void UClothTrainingTool::FGenerateClothOp::AddToCache(UChaosCache* InCache, int32 Frame)
+{
+	TArray<FVector3f> SimPositions;
+	if (!GetSimPositions(SimPositions))
+	{
+		return;
+	}
+
+	AddToCache(InCache, Frame, SimPositions);
+}
+
 void UClothTrainingTool::FGenerateClothOp::AddToCache(UChaosCache* InCache, int32 Frame, const TArray<FVector3f>& Positions) const
 {
 	constexpr float CacheFPS = 30;
@@ -166,23 +255,33 @@ void UClothTrainingTool::FGenerateClothOp::AddToCache(UChaosCache* InCache, int3
 	NewFrame.PendingChannelsData.Add(PositionYName, PendingPY);
 	NewFrame.PendingChannelsData.Add(PositionZName, PendingPZ);
 
-	Cache->AddFrame_Concurrent(MoveTemp(NewFrame));
+	InCache->AddFrame_Concurrent(MoveTemp(NewFrame));
 }
 
 void UClothTrainingTool::FGenerateClothOp::CalculateResult(FProgressCancel* Progress)
 {
-	if (ClothComponent == nullptr || Cache == nullptr || DataGenerationProxy == nullptr || !IsClothComponentValid())
+	if (ClothComponent == nullptr || DataGenerationProxy == nullptr || !IsClothComponentValid())
 	{
 		return;
 	}
 	
+	using UE::ClothTrainingTool::Private::ParseFrames;
+	using UE::ClothTrainingTool::Private::Range;
+
 	bool bCancelled = false;
 	constexpr float DeltaTime = 1e-3;
-	const int32 NumFrames = 10;
+	const TArray<int32> FramesToSimulate = ToolProperties->FramesToSimulate.Len() > 0
+		? ParseFrames(ToolProperties->FramesToSimulate) 
+		: Range(ToolProperties->AnimationSequence->GetNumberOfSampledKeys());
 
-	PrepareAnimSequence();
-	BackupClothComponentState();
+	const int32 NumFrames = ToolProperties->bDebug ? 1 : FramesToSimulate.Num();
+	if (NumFrames == 0)
+	{
+		return;
+	}
 
+	PrepareAnimationSequence();
+	BackupClothComponentState();
 	ClothComponent->ResumeSimulation();
 	for (int32 Frame = 0; Frame < NumFrames; Frame++)
 	{
@@ -194,27 +293,20 @@ void UClothTrainingTool::FGenerateClothOp::CalculateResult(FProgressCancel* Prog
 				break;
 			}
 
-			const float NumSteps = 200;
-
-			TArray<FTransform> BoneTransforms = GetBoneTransforms(Frame * 50);
+			const int32 FrameToSimulate = ToolProperties->bDebug ? ToolProperties->DebugFrame : FramesToSimulate[Frame];
+			const TArray<FTransform> BoneTransforms = GetBoneTransforms(FrameToSimulate);
 			if (!IsTransformsValid(BoneTransforms))
 			{
 				break;
 			}
 
 			ClothComponent->Pose(BoneTransforms);
-			Simulate(DeltaTime, NumSteps);
-
-			TArray<FVector3f> SimPositions;
-			if (!GetSimPositions(SimPositions))
-			{
-				break;
-			}
-			AddToCache(Cache, Frame, SimPositions);
+			Simulate(ToolProperties->TimeStep, ToolProperties->NumSteps, FrameToSimulate, Progress, ToolProperties->bDebug);
 
 			Progress->AdvanceCurrentScopeProgressBy(1.f / NumFrames);
 		}
 	}
+
 	RestoreClothComponentState();
 
 	if (!bCancelled)
@@ -223,23 +315,51 @@ void UClothTrainingTool::FGenerateClothOp::CalculateResult(FProgressCancel* Prog
 	}
 }
 
-void UClothTrainingTool::FGenerateClothOp::Simulate(float DeltaTime, int32 NumSteps)
+void UClothTrainingTool::FGenerateClothOp::Simulate(float DeltaTime, int32 NumSteps, int32 Frame, FProgressCancel* Progress, bool bDebug)
 {
+	using UE::ClothTrainingTool::Private::GetCache;
+	UChaosCache* Cache = ToolProperties->bDebug ? GetCache(ToolProperties->DebugCacheCollection) : GetCache(ToolProperties->CacheCollection);
+	if (Cache == nullptr)
+	{
+		return;
+	}
+
 	ClothComponent->ForceNextUpdateTeleportAndReset();
 	DataGenerationProxy->FillSimulationContext(DeltaTime);	
 	DataGenerationProxy->InitializeConfigs();
+	bool bCancelled = false;
 	for (int32 Step = 0; Step < NumSteps; ++Step)
 	{
-		DataGenerationProxy->Tick();
-
-		// Clear any reset flags at the end of the first step
-		if (Step == 0 && NumSteps > 1)
+		if (Progress)
 		{
-			ClothComponent->ResetTeleportMode();
-			DataGenerationProxy->FillSimulationContext(DeltaTime);	
+			if (Progress->Cancelled())
+			{
+				bCancelled = true;
+				break;
+			}
+
+			DataGenerationProxy->Tick();
+
+			// Clear any reset flags at the end of the first step
+			if (Step == 0 && NumSteps > 1)
+			{
+				ClothComponent->ResetTeleportMode();
+				DataGenerationProxy->FillSimulationContext(DeltaTime);	
+			}
+
+			if (bDebug)
+			{
+				DataGenerationProxy->WriteSimulationData();
+				AddToCache(Cache, Step);
+			}
 		}
 	}
-	DataGenerationProxy->WriteSimulationData();
+
+	if (!bDebug && !bCancelled)
+	{
+		DataGenerationProxy->WriteSimulationData();
+		AddToCache(Cache, Frame);
+	}
 }
 
 bool UClothTrainingTool::FGenerateClothOp::GetSimPositions(TArray<FVector3f> &OutPositions) const
@@ -377,6 +497,7 @@ void UClothTrainingTool::Setup()
 
 	ToolProperties = NewObject<UClothTrainingToolProperties>(this);
 	AddToolPropertySource(ToolProperties);
+	ToolProperties->RestoreProperties(this);
 
 	ActionProperties = NewObject<UClothTrainingToolActionProperties>(this);
 	ActionProperties->ParentTool = this;
@@ -385,27 +506,40 @@ void UClothTrainingTool::Setup()
 
 void UClothTrainingTool::RunTraining()
 {
-	if (ClothComponent == nullptr || ToolProperties->CacheCollection == nullptr)
+	if (ClothComponent == nullptr)
 	{
 		return;
 	}
 
-	static const FName CacheName = FName("SimulatedCache");
-	UChaosCacheCollection* CacheCollection = ToolProperties->CacheCollection;
-	UChaosCache* Cache = CacheCollection->FindOrAddCache(CacheName);
 
 	const FText DefaultMessage(LOCTEXT("ClothTrainingMessage", "Generate training data..."));
 	
 	using FTaskType = UE::Geometry::TModelingOpTask<FGenerateClothOp>;
 	using FExecuterType = UE::Geometry::FAsyncTaskExecuterWithProgressCancel<FTaskType>;
+	using UE::ClothTrainingTool::Private::GetCache;
 
 	if (!DataGenerationProxy.IsValid())
 	{
 		DataGenerationProxy = MakeUnique<FClothSimulationDataGenerationProxy>(*ClothComponent);
 	}
+	UChaosCacheCollection* const CacheCollection = ToolProperties->bDebug ? ToolProperties->DebugCacheCollection : ToolProperties->CacheCollection;
+	if (CacheCollection == nullptr)
+	{
+		if (ToolProperties->bDebug)
+		{
+			UE_LOG(LogClothTrainingTool, Error, TEXT("Debug cache is None. Please select a valid cache for output."));
+		}
+		else
+		{
+			UE_LOG(LogClothTrainingTool, Error, TEXT("Generated Cache is None. Please select a valid cache for output."));
+		}
+		return;
+	}
+	UChaosCache* const Cache = GetCache(CacheCollection);
+
 	FCacheUserToken CacheUserToken = Cache->BeginRecord(ClothComponent, FGuid(), FTransform::Identity);
 
-	TUniquePtr<FGenerateClothOp> NewOp = MakeUnique<FGenerateClothOp>(ToolProperties->AnimationSequence, Cache, ClothComponent, DataGenerationProxy.Get());
+	TUniquePtr<FGenerateClothOp> NewOp = MakeUnique<FGenerateClothOp>(ClothComponent, DataGenerationProxy.Get(), ToolProperties);
 
 	TUniquePtr<FExecuterType> BackgroundTaskExecuter = MakeUnique<FExecuterType>(MoveTemp(NewOp));
 	BackgroundTaskExecuter->StartBackgroundTask();
@@ -460,7 +594,7 @@ void UClothTrainingTool::RunTraining()
 		}
 		FSavePackageArgs SaveArgs;
 		SaveArgs.SaveFlags = SAVE_NoError;
-		const bool bSaveSucced = UPackage::SavePackage(Package, CacheCollection, *SavePath, SaveArgs);
+		const bool bSaveSucced = UPackage::SavePackage(Package, ToolProperties->CacheCollection, *SavePath, SaveArgs);
 		if (!bSaveSucced)
 		{
 			UE_LOG(LogTemp, Display, TEXT("Failed to save cache collection"));
@@ -492,4 +626,10 @@ void UClothTrainingTool::RequestAction(EClothTrainingToolActions ActionType)
 	PendingAction = ActionType;
 }
 
+void UClothTrainingTool::Shutdown(EToolShutdownType ShutdownType)
+{
+	Super::Shutdown(ShutdownType);
+	ToolProperties->SaveProperties(this);
+}
+
 #undef LOCTEXT_NAMESPACE
diff --git a/Engine/Plugins/Experimental/ChaosClothAssetEditor/Source/ChaosClothAssetEditorTools/Public/ChaosClothAsset/ClothTrainingTool.h b/Engine/Plugins/Experimental/ChaosClothAssetEditor/Source/ChaosClothAssetEditorTools/Public/ChaosClothAsset/ClothTrainingTool.h
--- a/Engine/Plugins/Experimental/ChaosClothAssetEditor/Source/ChaosClothAssetEditorTools/Public/ChaosClothAsset/ClothTrainingTool.h
+++ b/Engine/Plugins/Experimental/ChaosClothAssetEditor/Source/ChaosClothAssetEditorTools/Public/ChaosClothAsset/ClothTrainingTool.h
@@ -12,8 +12,9 @@
 
 #include "ClothTrainingTool.generated.h"
 
-class UChaosClothComponent;
+class UChaosCache;
 class UChaosCacheCollection;
+class UChaosClothComponent;
 class UClothTrainingTool;
 
 
@@ -24,16 +25,35 @@ struct FSkinnedMeshVertices
 };
 
 UCLASS()
-class CHAOSCLOTHASSETEDITORTOOLS_API UClothTrainingToolProperties : public UObject
+class CHAOSCLOTHASSETEDITORTOOLS_API UClothTrainingToolProperties : public UInteractiveToolPropertySet
 {
 	GENERATED_BODY()
 public:
 
-	UPROPERTY(EditAnywhere, Category = Files)
+	UPROPERTY(EditAnywhere, Category = Input)
 	TObjectPtr<UAnimSequence> AnimationSequence;
 
-	UPROPERTY(EditAnywhere, Category = Files, meta = (DisplayName = "Generated Cache"))
+	/* e.g. "0, 2, 5-10, 12-15". If left empty, all frames will be used */
+	UPROPERTY(EditAnywhere, Category = Input)
+	FString FramesToSimulate;
+
+	UPROPERTY(EditAnywhere, Category = Output, meta = (DisplayName = "Generated Cache", EditCondition = "!bDebug"))
 	TObjectPtr<UChaosCacheCollection> CacheCollection;
+
+	UPROPERTY(EditAnywhere, Category = Output)
+	bool bDebug = false;
+
+	UPROPERTY(EditAnywhere, Category = Output, meta = (Min = 0, EditCondition = "bDebug"))
+	uint32 DebugFrame = 0;
+
+	UPROPERTY(EditAnywhere, Category = Output, meta = (EditCondition = "bDebug", DisplayName = "Debug Cache"))
+	TObjectPtr<UChaosCacheCollection> DebugCacheCollection;
+
+	UPROPERTY(EditAnywhere, Category = "Simulation Settings", meta = (Min = 0))
+	float TimeStep = 1.f / 30;
+
+	UPROPERTY(EditAnywhere, Category = "Simulation Settings", meta = (Min = 0))
+	int32 NumSteps = 200;
 };
 
 UENUM()
@@ -55,8 +75,8 @@ public:
 
 	void PostAction(EClothTrainingToolActions Action);
 
-	UFUNCTION(CallInEditor, Category = Actions, meta = (DisplayName = "Begin Training", DisplayPriority = 1))
-		void StartTraining()
+	UFUNCTION(CallInEditor, Category = Actions, meta = (DisplayName = "Begin Generating", DisplayPriority = 1))
+		void StartGenerating()
 	{
 		PostAction(EClothTrainingToolActions::Train);
 	}
@@ -92,6 +112,7 @@ public:
 	// UInteractiveTool
 	virtual void Setup() override;
 	virtual void OnTick(float DeltaTime) override;
+	virtual void Shutdown(EToolShutdownType ShutdownType) override;
 
 private:
 	class FClothSimulationDataGenerationProxy;
@@ -115,7 +136,6 @@ private:
 	void RequestAction(EClothTrainingToolActions ActionType);
 
 	void RunTraining();
-
 };
 
 
