Cloth Editor:
- Add a UClothEditorContextObject, which can be used to pass information into interactive tools. Currently we are passing in the current Dataflow graph and its associated Editor
- Change the DynamicMeshComponents that are rendered in the construction view to be the output of the currently selected Dataflow node, if any
- Change the construction view so that all patterns are combined into one DynamicMeshComponent
- Add a new AddWeightMapNode which adds a given vertex weight map to a ManagedArrayCollection
- Modify the ClothWeightMapPaintTool to create a new AddWeightMapNode, or modify an existing one, rather than adding the weight map directly to the cloth asset

diff --git a/Engine/Plugins/Experimental/ChaosClothAssetEditor/Source/ChaosClothAssetDataflowNodes/Private/ChaosClothAsset/AddWeightMapNode.cpp b/Engine/Plugins/Experimental/ChaosClothAssetEditor/Source/ChaosClothAssetDataflowNodes/Private/ChaosClothAsset/AddWeightMapNode.cpp
new file mode 100644
--- /dev/null
+++ b/Engine/Plugins/Experimental/ChaosClothAssetEditor/Source/ChaosClothAssetDataflowNodes/Private/ChaosClothAsset/AddWeightMapNode.cpp
@@ -0,0 +1,64 @@
+// Copyright Epic Games, Inc. All Rights Reserved.
+
+#include "ChaosClothAsset/AddWeightMapNode.h"
+#include "ChaosClothAsset/DataflowNodes.h"
+#include "ChaosClothAsset/CollectionClothFacade.h"
+
+#include UE_INLINE_GENERATED_CPP_BY_NAME(AddWeightMapNode)
+
+#define LOCTEXT_NAMESPACE "ChaosClothAssetAddWeightMapNode"
+
+FChaosClothAssetAddWeightMapNode::FChaosClothAssetAddWeightMapNode(const Dataflow::FNodeParameters& InParam, FGuid InGuid)
+	: FDataflowNode(InParam, InGuid)
+{
+	RegisterInputConnection(&Collection);
+	RegisterOutputConnection(&Collection, &Collection);
+	RegisterInputConnection(&Name);
+	RegisterOutputConnection(&Name, &Name);
+}
+
+void FChaosClothAssetAddWeightMapNode::Evaluate(Dataflow::FContext& Context, const FDataflowOutput* Out) const
+{
+	using namespace UE::Chaos::ClothAsset;
+
+	if (Out->IsA<FManagedArrayCollection>(&Collection))
+	{
+		// Evaluate in collection
+		FManagedArrayCollection InCollection = GetValue<FManagedArrayCollection>(Context, &Collection);
+		const TSharedRef<FManagedArrayCollection> ClothCollection = MakeShared<FManagedArrayCollection>(MoveTemp(InCollection));
+		FCollectionClothFacade ClothFacade(ClothCollection);
+
+		const FString InNameString = GetValue<FString>(Context, &Name);
+
+		// Note: If the input collection has no valid LODs, just pass it to the output as is (this can happen if the node is asked to Evaluate before the inputs are connected)
+
+		if (ClothFacade.GetNumLods() > 0)
+		{
+			const FName InName(InNameString);
+			ClothFacade.AddWeightMap(InName);		// Does nothing if weight map already exists
+
+			FCollectionClothLodFacade ClothLod = ClothFacade.GetLod(0);
+			TArrayView<float> LodWeights = ClothLod.GetWeightMap(InName);
+
+			const int32 MaxWeightIndex = FMath::Min(VertexWeights.Num(), LodWeights.Num());
+
+			if (VertexWeights.Num() != LodWeights.Num())
+			{
+				DataflowNodes::LogAndToastWarning(
+					FText::Format(LOCTEXT("WeightMapSize", "FChaosClothAssetAddWeightMapNode: Vertex count mismatch: vertex weights in the node: {0}; vertices in the specified LOD: {1}"),
+						VertexWeights.Num(),
+						LodWeights.Num()));
+			}
+				
+			for (int32 VertexID = 0; VertexID < MaxWeightIndex; ++VertexID)
+			{
+				LodWeights[VertexID] = VertexWeights[VertexID];
+			}
+		}
+
+		SetValue<FManagedArrayCollection>(Context, *ClothCollection, &Collection);
+		SetValue<FString>(Context, InNameString, &Name);
+	}
+}
+
+#undef LOCTEXT_NAMESPACE
diff --git a/Engine/Plugins/Experimental/ChaosClothAssetEditor/Source/ChaosClothAssetDataflowNodes/Private/ChaosClothAsset/DataflowNodes.cpp b/Engine/Plugins/Experimental/ChaosClothAssetEditor/Source/ChaosClothAssetDataflowNodes/Private/ChaosClothAsset/DataflowNodes.cpp
--- a/Engine/Plugins/Experimental/ChaosClothAssetEditor/Source/ChaosClothAssetDataflowNodes/Private/ChaosClothAsset/DataflowNodes.cpp
+++ b/Engine/Plugins/Experimental/ChaosClothAssetEditor/Source/ChaosClothAssetDataflowNodes/Private/ChaosClothAsset/DataflowNodes.cpp
@@ -12,6 +12,7 @@
 #include "ChaosClothAsset/SimulationDefaultConfigNode.h"
 #include "ChaosClothAsset/SimulationXPBDConfigNode.h"
 #include "ChaosClothAsset/BindToRootBoneNode.h"
+#include "ChaosClothAsset/AddWeightMapNode.h"
 #include "Dataflow/DataflowNodeFactory.h"
 #include "Framework/Notifications/NotificationManager.h"
 #include "Widgets/Notifications/SNotificationList.h"
@@ -33,6 +34,7 @@ namespace UE::Chaos::ClothAsset::DataflowNodes
 		DATAFLOW_NODE_REGISTER_CREATION_FACTORY(FChaosClothAssetSimulationDefaultConfigNode);
 		DATAFLOW_NODE_REGISTER_CREATION_FACTORY(FChaosClothAssetSimulationXPBDConfigNode);
 		DATAFLOW_NODE_REGISTER_CREATION_FACTORY(FChaosClothAssetBindToRootBoneNode);
+		DATAFLOW_NODE_REGISTER_CREATION_FACTORY(FChaosClothAssetAddWeightMapNode);
 	}
 
 	void LogAndToastWarning(const FText& Error)
diff --git a/Engine/Plugins/Experimental/ChaosClothAssetEditor/Source/ChaosClothAssetDataflowNodes/Public/ChaosClothAsset/AddWeightMapNode.h b/Engine/Plugins/Experimental/ChaosClothAssetEditor/Source/ChaosClothAssetDataflowNodes/Public/ChaosClothAsset/AddWeightMapNode.h
new file mode 100644
--- /dev/null
+++ b/Engine/Plugins/Experimental/ChaosClothAssetEditor/Source/ChaosClothAssetDataflowNodes/Public/ChaosClothAsset/AddWeightMapNode.h
@@ -0,0 +1,31 @@
+// Copyright Epic Games, Inc. All Rights Reserved.
+
+#pragma once 
+
+#include "Dataflow/DataflowNode.h"
+#include "GeometryCollection/ManagedArrayCollection.h"
+#include "AddWeightMapNode.generated.h"
+
+USTRUCT(Meta = (DataflowCloth))
+struct FChaosClothAssetAddWeightMapNode : public FDataflowNode
+{
+	GENERATED_USTRUCT_BODY()
+	DATAFLOW_NODE_DEFINE_INTERNAL(FChaosClothAssetAddWeightMapNode, "AddWeightMap", "Cloth", "Cloth Add Weight Map")
+	//DATAFLOW_NODE_RENDER_TYPE(FGeometryCollection::StaticType(), "Collection")  // TODO: Leave out the render type until there is something to render
+
+public:
+
+	UPROPERTY(Meta = (Dataflowinput, DataflowOutput, DisplayName = "Collection", DataflowPassthrough = "Collection"))
+	FManagedArrayCollection Collection;
+
+	UPROPERTY(EditAnywhere, Category = "Weight Map", Meta = (Dataflowinput, DataflowOutput, DisplayName = "Name", DataflowPassthrough = "Name"))
+	FString Name;
+
+	UPROPERTY(EditAnywhere, Category = "Weight Map")
+	TArray<float> VertexWeights;
+
+
+	FChaosClothAssetAddWeightMapNode(const Dataflow::FNodeParameters& InParam, FGuid InGuid = FGuid::NewGuid());
+
+	virtual void Evaluate(Dataflow::FContext& Context, const FDataflowOutput* Out) const override;
+};
diff --git a/Engine/Plugins/Experimental/ChaosClothAssetEditor/Source/ChaosClothAssetEditor/Private/ChaosClothAsset/ClothEditorMode.cpp b/Engine/Plugins/Experimental/ChaosClothAssetEditor/Source/ChaosClothAssetEditor/Private/ChaosClothAsset/ClothEditorMode.cpp
--- a/Engine/Plugins/Experimental/ChaosClothAssetEditor/Source/ChaosClothAssetEditor/Private/ChaosClothAsset/ClothEditorMode.cpp
+++ b/Engine/Plugins/Experimental/ChaosClothAssetEditor/Source/ChaosClothAssetEditor/Private/ChaosClothAsset/ClothEditorMode.cpp
@@ -9,6 +9,7 @@
 #include "ChaosClothAsset/ClothEditorRestSpaceViewportClient.h"
 #include "ChaosClothAsset/ClothPatternToDynamicMesh.h"
 #include "ChaosClothAsset/ClothEditorPreviewScene.h"
+#include "ChaosClothAsset/ClothEditorContextObject.h"
 #include "AssetEditorModeManager.h"
 #include "Drawing/MeshElementsVisualizer.h"
 #include "EditorViewportClient.h"
@@ -52,6 +53,9 @@
 #include "Dataflow/DataflowComponent.h"
 #include "Elements/Framework/EngineElementsLibrary.h"
 #include "ToolTargets/SkeletalMeshComponentToolTarget.h"
+#include "Dataflow/DataflowGraphEditor.h"
+#include "ContextObjectStore.h"
+
 
 #include UE_INLINE_GENERATED_CPP_BY_NAME(ClothEditorMode)
 
@@ -243,6 +247,12 @@ void UChaosClothAssetEditorMode::OnToolStarted(UInteractiveToolManager* Manager,
 	FChaosClothAssetEditorCommands::UpdateToolCommandBinding(Tool, ToolCommandList, false);
 
 	bCanTogglePattern2DMode = false;
+
+	// Don't let the Graph change during tool execution (or even which nodes are selected)
+	if (TSharedPtr<SDataflowGraphEditor> GraphEditor = DataflowGraphEditor.Pin())
+	{
+		GraphEditor->SetEnabled(false);
+	}
 }
 
 void UChaosClothAssetEditorMode::OnToolEnded(UInteractiveToolManager* Manager, UInteractiveTool* Tool)
@@ -252,8 +262,12 @@ void UChaosClothAssetEditorMode::OnToolEnded(UInteractiveToolManager* Manager, U
 	bCanTogglePattern2DMode = true;
 	ActiveToolsContext = nullptr;
 
-	UpdateSimulationMeshes();
 	ReinitializeDynamicMeshComponents();
+
+	if (TSharedPtr<SDataflowGraphEditor> GraphEditor = DataflowGraphEditor.Pin())
+	{
+		GraphEditor->SetEnabled(true);
+	}
 }
 
 void UChaosClothAssetEditorMode::PostUndo()
@@ -337,91 +351,6 @@ void UChaosClothAssetEditorMode::SetPreviewScene(FChaosClothPreviewScene* InPrev
 	RegisterPreviewTools();
 }
 
-void UChaosClothAssetEditorMode::UpdateSimulationMeshes()
-{
-	GetToolManager()->BeginUndoTransaction(LOCTEXT("ChaosClothAssetEditorApplyChangesTransaction", "Cloth Editor Apply Changes"));
-
-	UChaosClothAsset* ChaosClothAsset = PreviewScene->ClothComponent->GetClothAsset();
-	ChaosClothAsset->Modify();
-
-	UE::Chaos::ClothAsset::FCollectionClothFacade ClothFacade(ChaosClothAsset->GetClothCollection());
-
-	check(DynamicMeshSourceInfos.Num() == DynamicMeshComponents.Num());
-
-	// Search for weight map names that are in the original cloth asset, but are not in *all* dynamic mesh components. These are weight maps
-	// that have been removed by one of the mesh tools, and so should be removed from the asset.
-	const TSet<FName> ClothAssetWeightMapNames = TSet<FName>(ClothFacade.GetWeightMapNames());
-	TSet<FName> CommonDynamicMeshWeightMapNames;
-	bool bFirstIteration = true;
-
-	for (int DynamicMeshIndex = 0; DynamicMeshIndex < DynamicMeshComponents.Num(); ++DynamicMeshIndex)
-	{
-		const int32 LodIndex = DynamicMeshSourceInfos[DynamicMeshIndex].LodIndex;
-		const int32 PatternIndex = DynamicMeshSourceInfos[DynamicMeshIndex].PatternIndex;
-
-		UE::Chaos::ClothAsset::FCollectionClothLodFacade ClothLodFacade = ClothFacade.GetLod(LodIndex);
-		UE::Chaos::ClothAsset::FCollectionClothPatternFacade ClothPatternFacade = ClothLodFacade.GetPattern(PatternIndex);
-
-		const UDynamicMeshComponent& DynamicMeshComponent = *DynamicMeshComponents[DynamicMeshIndex];
-		const UE::Geometry::FDynamicMesh3* DynamicMesh = DynamicMeshComponent.GetMesh();
-
-		//
-		// TODO: Set vertices and triangles on the pattern, e.g.:
-		// void Initialize(const TArray<FVector2f>& Positions, const TArray<FVector3f>& RestPositions, const TArray<uint32>& Indices)
-		//
-
-		//
-		// Weight maps
-		//
-
-		if (!DynamicMesh->HasAttributes() || DynamicMesh->Attributes()->NumWeightLayers() == 0)
-		{
-			continue;
-		}
-		 
-		const TSet<FName> DynamicMeshWeightMapNames(ChaosClothAssetEditorModeHelpers::GetDynamicMeshWeightMapNames(*DynamicMesh));
-		if (bFirstIteration)
-		{
-			CommonDynamicMeshWeightMapNames = DynamicMeshWeightMapNames;
-			bFirstIteration = false;
-		}
-		else
-		{
-			CommonDynamicMeshWeightMapNames = CommonDynamicMeshWeightMapNames.Intersect(DynamicMeshWeightMapNames);
-		}
-
-		const int NumInputWeightElements = DynamicMesh->MaxVertexID();
-
-		for (int32 DynamicMeshWeightMapIndex = 0; DynamicMeshWeightMapIndex < DynamicMesh->Attributes()->NumWeightLayers(); ++DynamicMeshWeightMapIndex)
-		{
-			const UE::Geometry::FDynamicMeshWeightAttribute* WeightMapAttribute = DynamicMesh->Attributes()->GetWeightLayer(DynamicMeshWeightMapIndex);
-			const FName WeightMapName = WeightMapAttribute->GetName();
-
-			ClothFacade.AddWeightMap(WeightMapName);		// Does nothing if weight map already exists
-
-			TArrayView<float> PatternWeights = ClothPatternFacade.GetWeightMap(WeightMapName);
-			for (int VertexID = 0; VertexID < NumInputWeightElements; ++VertexID)
-			{
-				float Val;
-				WeightMapAttribute->GetValue(VertexID, &Val);
-				PatternWeights[VertexID] = Val;
-			}
-		}
-	}
-
-	const TSet<FName> WeightMapsToRemove = ClothAssetWeightMapNames.Difference(CommonDynamicMeshWeightMapNames);
-	ChaosClothAssetEditorModeHelpers::RemoveClothWeightMaps(ClothFacade, WeightMapsToRemove.Array());
-
-	ChaosClothAsset->Build();
-
-	// Reset cloth component
-	{
-		const FComponentReregisterContext Context(PreviewScene->ClothComponent);
-	}
-
-	GetToolManager()->EndUndoTransaction();
-}
-
 void UChaosClothAssetEditorMode::CreateToolTargets(const TArray<TObjectPtr<UObject>>& AssetsIn) 
 {
 	// TODO: When we have a cloth component tool target, create it here
@@ -447,8 +376,42 @@ bool UChaosClothAssetEditorMode::IsComponentSelected(const UPrimitiveComponent*
 }
 
 
+void UChaosClothAssetEditorMode::SetSelectedClothCollection(TSharedPtr<FManagedArrayCollection> Collection)
+{
+	SelectedClothCollection = Collection;
+	ReinitializeDynamicMeshComponents();
+}
+
+TSharedPtr<FManagedArrayCollection> UChaosClothAssetEditorMode::GetClothCollection()
+{
+	return SelectedClothCollection;
+
+	// TODO: If no cloth collection node is selected, show the ClothAsset's collection. In this case, also ensure that any interactive tools are disabled. (UE-181574)
+}
+
+
 void UChaosClothAssetEditorMode::ReinitializeDynamicMeshComponents()
 {
+	auto SetUpDynamicMeshComponent = [this](UDynamicMeshComponent& MeshComponent)
+	{
+		MeshComponent.SelectionOverrideDelegate = UPrimitiveComponent::FSelectionOverride::CreateUObject(this, &UChaosClothAssetEditorMode::IsComponentSelected);
+
+		UMaterialInterface* const Material = ToolSetupUtil::GetDefaultSculptMaterial(GetToolManager());
+		UMaterialInstanceDynamic* const MatInstance = UMaterialInstanceDynamic::Create(Material, GetToolManager());
+		MatInstance->SetVectorParameterValue(TEXT("Color"), FLinearColor::Gray);
+		MatInstance->TwoSided = 1;
+		MeshComponent.SetMaterial(0, MatInstance);
+
+		if (bPattern2DMode)
+		{
+			// The LVT_OrthoXY viewport transform results in the positive y axis pointing down in screen space. The following transform rotates 
+			// the mesh around the z axis so that increasing y value in the cloth asset corresponds to "up" in screen space
+			MeshComponent.SetWorldRotation(FQuat(FVector(0, 0, 1), FMathd::Pi));
+		}
+
+		MeshComponent.RegisterComponentWithWorld(this->GetWorld());
+	};
+	
 	// Clean up any existing DynamicMeshComponents
 	// Save indices of selected mesh components
 	TArray<int32> PreviouslySelectedDynamicMeshComponents;
@@ -481,19 +444,13 @@ void UChaosClothAssetEditorMode::ReinitializeDynamicMeshComponents()
 	WireframesToTick.Empty();
 	DynamicMeshSourceInfos.Empty();
 
-	if (!PreviewScene->ClothComponent)
+	TSharedPtr<FManagedArrayCollection> Collection = GetClothCollection();
+	if (!Collection)
 	{
 		return;
 	}
 
-	const UChaosClothAsset* ChaosClothAsset = PreviewScene->ClothComponent->GetClothAsset();
-
-	if (!ChaosClothAsset)
-	{
-		return;
-	}
-
-	const UE::Chaos::ClothAsset::FCollectionClothConstFacade ClothFacade(ChaosClothAsset->GetClothCollection());
+	const UE::Chaos::ClothAsset::FCollectionClothConstFacade ClothFacade(Collection);
 
 	for (int32 LodIndex = 0; LodIndex < ClothFacade.GetNumLods(); ++LodIndex)
 	{
@@ -501,8 +458,40 @@ void UChaosClothAssetEditorMode::ReinitializeDynamicMeshComponents()
 
 		if (bCombineAllPatterns)
 		{
-			// TODO: Enable a mode where all patterns are edited as one dynamic mesh component. Will require some fun bookkeeping.
-			check(0);
+			DynamicMeshComponents.SetNum(1);
+			DynamicMeshComponentParentActors.SetNum(1);
+
+			UE::Geometry::FDynamicMesh3 LodMesh;
+			UE::Geometry::FDynamicMeshEditor MeshEditor(&LodMesh);
+
+			for (int32 PatternIndex = 0; PatternIndex < ClothLodFacade.GetNumPatterns(); ++PatternIndex)
+			{
+				UE::Geometry::FDynamicMesh3 PatternMesh;
+				FClothPatternToDynamicMesh Converter;
+				Converter.Convert(Collection, LodIndex, PatternIndex, bPattern2DMode, PatternMesh);
+
+				if (bPattern2DMode)
+				{
+					// Use per-triangle normals for the 2D view
+					UE::Geometry::FMeshNormals::InitializeMeshToPerTriangleNormals(&PatternMesh);
+				}
+
+				UE::Geometry::FMeshIndexMappings IndexMapping;
+				MeshEditor.AppendMesh(&PatternMesh, IndexMapping);
+			}
+
+			// We only need an actor to allow use of HHitProxy for selection
+			const FRotator Rotation(0.0f, 0.0f, 0.0f);
+			const FActorSpawnParameters SpawnInfo;
+			const TObjectPtr<AActor> ParentActor = this->GetWorld()->SpawnActor<AActor>(FVector::ZeroVector, Rotation, SpawnInfo);
+			DynamicMeshComponentParentActors[0] = ParentActor;
+
+			DynamicMeshComponents[0] = NewObject<UDynamicMeshComponent>(ParentActor);
+			DynamicMeshComponents[0]->SetMesh(MoveTemp(LodMesh));
+			
+			SetUpDynamicMeshComponent(*DynamicMeshComponents[0]);
+
+			DynamicMeshSourceInfos.Add(FDynamicMeshSourceInfo{ LodIndex, -1 });
 		}
 		else
 		{
@@ -514,7 +503,7 @@ void UChaosClothAssetEditorMode::ReinitializeDynamicMeshComponents()
 			{
 				UE::Geometry::FDynamicMesh3 PatternMesh;
 				FClothPatternToDynamicMesh Converter;
-				Converter.Convert(ChaosClothAsset, LodIndex, PatternIndex, bPattern2DMode, PatternMesh);
+				Converter.Convert(Collection, LodIndex, PatternIndex, bPattern2DMode, PatternMesh);
 
 				if (bPattern2DMode)
 				{
@@ -531,22 +520,8 @@ void UChaosClothAssetEditorMode::ReinitializeDynamicMeshComponents()
 				TObjectPtr<UDynamicMeshComponent> PatternMeshComponent = NewObject<UDynamicMeshComponent>(ParentActor);
 				PatternMeshComponent->SetMesh(MoveTemp(PatternMesh));
 
-				PatternMeshComponent->SelectionOverrideDelegate = UPrimitiveComponent::FSelectionOverride::CreateUObject(this, &UChaosClothAssetEditorMode::IsComponentSelected);
-
-				UMaterialInterface* Material = ToolSetupUtil::GetDefaultSculptMaterial(GetToolManager());
-				UMaterialInstanceDynamic* MatInstance = UMaterialInstanceDynamic::Create(Material, GetToolManager());
-				MatInstance->SetVectorParameterValue(TEXT("Color"), FLinearColor::Gray);
-				MatInstance->TwoSided = 1;
-				PatternMeshComponent->SetMaterial(0, MatInstance);
+				SetUpDynamicMeshComponent(*PatternMeshComponent);
 
-				if (bPattern2DMode)
-				{
-					// The LVT_OrthoXY viewport transform results in the positive y axis pointing down in screen space. The following transform rotates 
-					// the mesh around the z axis so that increasing y value in the cloth asset corresponds to "up" in screen space
-					PatternMeshComponent->SetWorldRotation(FQuat(FVector(0,0,1), FMathd::Pi));
-				}
-
-				PatternMeshComponent->RegisterComponentWithWorld(this->GetWorld());
 				DynamicMeshComponents[PatternIndex] = PatternMeshComponent;
 				DynamicMeshSourceInfos.Add(FDynamicMeshSourceInfo{ LodIndex, PatternIndex });
 			}
@@ -922,5 +897,45 @@ void UChaosClothAssetEditorMode::SetRestSpaceViewportClient(TWeakPtr<FChaosCloth
 	}
 }
 
+
+void UChaosClothAssetEditorMode::InitializeContextObject()
+{
+	UEditorInteractiveToolsContext* const RestSpaceToolsContext = GetInteractiveToolsContext();
+
+	UClothEditorContextObject* EditorContextObject = RestSpaceToolsContext->ContextObjectStore->FindContext<UClothEditorContextObject>();
+	if (!EditorContextObject)
+	{
+		EditorContextObject = NewObject<UClothEditorContextObject>();
+		RestSpaceToolsContext->ContextObjectStore->AddContextObject(EditorContextObject);
+	}
+
+	EditorContextObject->Init(DataflowGraphEditor, DataflowGraph);
+
+	check(EditorContextObject);
+
+}
+
+void UChaosClothAssetEditorMode::DeleteContextObject()
+{
+	UEditorInteractiveToolsContext* const RestSpaceToolsContext = GetInteractiveToolsContext();
+	if (UClothEditorContextObject* ClothEditorContextObject = RestSpaceToolsContext->ContextObjectStore->FindContext<UClothEditorContextObject>())
+	{
+		RestSpaceToolsContext->ContextObjectStore->RemoveContextObject(ClothEditorContextObject);
+	}
+}
+
+void UChaosClothAssetEditorMode::SetDataflowGraphEditor(TSharedPtr<SDataflowGraphEditor> InGraphEditor)
+{
+	if (InGraphEditor)
+	{
+		DataflowGraphEditor = InGraphEditor;
+		InitializeContextObject();
+	}
+	else
+	{
+		DeleteContextObject();
+	}
+}
+
 #undef LOCTEXT_NAMESPACE
 
diff --git a/Engine/Plugins/Experimental/ChaosClothAssetEditor/Source/ChaosClothAssetEditor/Private/ChaosClothAsset/ClothEditorMode.h b/Engine/Plugins/Experimental/ChaosClothAssetEditor/Source/ChaosClothAssetEditor/Private/ChaosClothAsset/ClothEditorMode.h
--- a/Engine/Plugins/Experimental/ChaosClothAssetEditor/Source/ChaosClothAssetEditor/Private/ChaosClothAsset/ClothEditorMode.h
+++ b/Engine/Plugins/Experimental/ChaosClothAssetEditor/Source/ChaosClothAssetEditor/Private/ChaosClothAsset/ClothEditorMode.h
@@ -26,9 +26,12 @@ class UChaosClothComponent;
 class FEditorViewportClient;
 class FChaosClothEditorRestSpaceViewportClient;
 class FViewport;
+class UDataflow;
 class UDataflowComponent;
 class FChaosClothPreviewScene;
 class FChaosClothAssetEditor3DViewportClient;
+class SDataflowGraphEditor;
+struct FManagedArrayCollection;
 
 /**
  * The cloth editor mode is the mode used in the cloth asset editor. It holds most of the inter-tool state.
@@ -117,6 +120,16 @@ private:
 	void SetRestSpaceViewportClient(TWeakPtr<FChaosClothEditorRestSpaceViewportClient, ESPMode::ThreadSafe> ViewportClient);
 	void RefocusRestSpaceViewportClient();
 
+	// intended to be called by the toolkit when selected node in the Dataflow graph changes
+	void SetSelectedClothCollection(TSharedPtr<FManagedArrayCollection> Collection);
+
+	// gets the currently selected cloth collection, as specified by the toolkit
+	TSharedPtr<FManagedArrayCollection> GetClothCollection();
+
+	void SetDataflowGraphEditor(TSharedPtr<SDataflowGraphEditor> InGraphEditor);
+
+	void InitializeContextObject();
+	void DeleteContextObject();
 
 	// Rest space wireframes. They have to get ticked to be able to respond to setting changes. 
 	UPROPERTY()
@@ -162,16 +175,13 @@ private:
 
 	// Whether to combine all patterns into a single DynamicMeshComponent, or have separate components for each pattern
 	// TODO: Expose this to the user
-	bool bCombineAllPatterns = false;
+	bool bCombineAllPatterns = true;
 
 	bool IsComponentSelected(const UPrimitiveComponent* InComponent);
 
 	// Create dynamic mesh components from the cloth component's rest space info
 	void ReinitializeDynamicMeshComponents();
 
-	// Set up the preview simulation mesh from the given rest-space mesh
-	void UpdateSimulationMeshes();
-
 	// Simulation controls
 	bool bShouldResetSimulation = false;
 	bool bHardReset = false;
@@ -180,7 +190,15 @@ private:
 	UPROPERTY()
 	TObjectPtr<UDataflowComponent> DataflowComponent = nullptr;
 
+	UPROPERTY()
+	TObjectPtr<UDataflow> DataflowGraph = nullptr;
+
+	TWeakPtr<SDataflowGraphEditor> DataflowGraphEditor;
+
 	UPROPERTY()
 	TObjectPtr<UEditorInteractiveToolsContext> ActiveToolsContext = nullptr;
+
+
+	TSharedPtr<FManagedArrayCollection> SelectedClothCollection = nullptr;
 };
 
diff --git a/Engine/Plugins/Experimental/ChaosClothAssetEditor/Source/ChaosClothAssetEditor/Private/ChaosClothAsset/ClothEditorToolkit.cpp b/Engine/Plugins/Experimental/ChaosClothAssetEditor/Source/ChaosClothAssetEditor/Private/ChaosClothAsset/ClothEditorToolkit.cpp
--- a/Engine/Plugins/Experimental/ChaosClothAssetEditor/Source/ChaosClothAssetEditor/Private/ChaosClothAsset/ClothEditorToolkit.cpp
+++ b/Engine/Plugins/Experimental/ChaosClothAssetEditor/Source/ChaosClothAssetEditor/Private/ChaosClothAsset/ClothEditorToolkit.cpp
@@ -12,6 +12,8 @@
 #include "ChaosClothAsset/ClothComponent.h"
 #include "ChaosClothAsset/ClothAsset.h"
 #include "ChaosClothAsset/ClothEditorModeToolkit.h"
+#include "ChaosClothAsset/ClothCollection.h"
+#include "ChaosClothAsset/CollectionClothFacade.h"
 #include "Framework/Docking/LayoutExtender.h"
 #include "AssetEditorModeManager.h"
 #include "AdvancedPreviewScene.h"
@@ -33,6 +35,7 @@
 #include "Algo/RemoveIf.h"
 #include "AdvancedPreviewSceneModule.h"
 #include "Widgets/Layout/SSpacer.h"
+#include "DynamicMesh/DynamicMesh3.h"
 
 #define LOCTEXT_NAMESPACE "ChaosClothAssetEditorToolkit"
 
@@ -451,6 +454,9 @@ void FChaosClothAssetEditorToolkit::PostInitAssetEditor()
 	}
 
 	InitDetailsViewPanel();
+
+	ClothMode->DataflowGraph = Dataflow;
+	ClothMode->SetDataflowGraphEditor(GraphEditor);
 }
 
 //~ End FAssetEditorToolkit overrides
@@ -752,6 +758,10 @@ void FChaosClothAssetEditorToolkit::OnFinishedChangingAssetProperties(const FPro
 				GraphEditorTab.Get()->SetContent(SNew(SSpacer));
 				NodeDetailsTab.Get()->SetContent(SNew(SSpacer));
 			}
+
+			UChaosClothAssetEditorMode* const ClothMode = CastChecked<UChaosClothAssetEditorMode>(EditorModeManager->GetActiveScriptableMode(UChaosClothAssetEditorMode::EM_ChaosClothAssetEditorModeId));
+			ClothMode->DataflowGraph = Dataflow;
+			ClothMode->SetDataflowGraphEditor(GraphEditor);
 		}
 	}
 }
@@ -844,7 +854,7 @@ void FChaosClothAssetEditorToolkit::ReinitializeGraphEditorWidget()
 
 	GraphEditor->Construct(Args, ClothAsset);
 
-	GraphEditor->OnSelectionChangedMulticast.Clear();
+	GraphEditor->OnSelectionChangedMulticast.RemoveAll(this);
 	GraphEditor->OnSelectionChangedMulticast.AddSP(this, &FChaosClothAssetEditorToolkit::OnNodeSelectionChanged);
 }
 
@@ -899,24 +909,58 @@ void FChaosClothAssetEditorToolkit::OnNodeTitleCommitted(const FText& InNewText,
 
 void FChaosClothAssetEditorToolkit::OnNodeSelectionChanged(const TSet<UObject*>& NewSelection) const
 {
-	if (Dataflow)
+	auto GetClothCollectionIfPossible = [](const TSharedPtr<const FDataflowNode> DataflowNode, const TSharedPtr<Dataflow::FEngineContext> DataflowContext) -> TSharedPtr<FManagedArrayCollection>
 	{
-		// remove nodes if they don't have the render box checked
-		int32 NewLen = Algo::RemoveIf(Dataflow->RenderTargets, [](const UDataflowEdNode* Node)
+		for (const FDataflowOutput* const Output : DataflowNode->GetOutputs())
 		{
-			return !Node->bRenderInAssetEditor;
-		});
-		Dataflow->RenderTargets.SetNum(NewLen);
+			if (Output->GetType() == FName("FManagedArrayCollection"))
+			{
+				const FManagedArrayCollection DefaultValue;
+				TSharedPtr<FManagedArrayCollection> Collection = MakeShared<FManagedArrayCollection>(Output->GetValue<FManagedArrayCollection>(*DataflowContext, DefaultValue));
+
+				// see if the output collection is a ClothCollection
+				const UE::Chaos::ClothAsset::FCollectionClothConstFacade ClothFacade(Collection);
+				if (ClothFacade.IsValid())
+				{
+					return Collection;
+				}
+
+				return MakeShared<FManagedArrayCollection>();
+			}
+		}
+
+		return TSharedPtr<FManagedArrayCollection>();
+	};
+
 
-		for (UObject* Selected : NewSelection)
+	TSharedPtr<FManagedArrayCollection> Collection = nullptr;
+
+	if (Dataflow)
+	{
+		// Get any selected node with a ClothCollection output
+		for (const UObject* const Selected : NewSelection)
 		{
-			if (UDataflowEdNode* Node = Cast<UDataflowEdNode>(Selected))
+			if (const UDataflowEdNode* const Node = Cast<UDataflowEdNode>(Selected))
 			{
 				Dataflow->RenderTargets.Add(Node);
+
+				if (const TSharedPtr<const FDataflowNode> DataflowNode = Node->GetDataflowNode())
+				{
+					Collection = GetClothCollectionIfPossible(DataflowNode, this->DataflowContext);
+				}
 			}
 		}
+
 		Dataflow->LastModifiedRenderTarget = Dataflow::FTimestamp::Current();
 	}
+
+	UChaosClothAssetEditorMode* const ClothMode = CastChecked<UChaosClothAssetEditorMode>(EditorModeManager->GetActiveScriptableMode(UChaosClothAssetEditorMode::EM_ChaosClothAssetEditorModeId));
+	if (ClothMode)
+	{
+		ClothMode->SetSelectedClothCollection(Collection);
+		ClothMode->RefocusRestSpaceViewportClient();
+	}
+
 }
 
 //~ Ends DataflowEditorActions
diff --git a/Engine/Plugins/Experimental/ChaosClothAssetEditor/Source/ChaosClothAssetEditor/Private/ChaosClothAsset/ClothEditorToolkit.h b/Engine/Plugins/Experimental/ChaosClothAssetEditor/Source/ChaosClothAssetEditor/Private/ChaosClothAsset/ClothEditorToolkit.h
--- a/Engine/Plugins/Experimental/ChaosClothAssetEditor/Source/ChaosClothAssetEditor/Private/ChaosClothAsset/ClothEditorToolkit.h
+++ b/Engine/Plugins/Experimental/ChaosClothAssetEditor/Source/ChaosClothAssetEditor/Private/ChaosClothAsset/ClothEditorToolkit.h
@@ -104,6 +104,18 @@ private:
 	void PopulateOutliner();
 	void OnClothAssetChanged();
 
+	// Dataflow
+	void EvaluateNode(FDataflowNode* Node, FDataflowOutput* Out);
+	TSharedRef<SDataflowGraphEditor> CreateGraphEditorWidget();
+	void ReinitializeGraphEditorWidget();
+	TSharedPtr<IStructureDetailsView> CreateNodeDetailsEditorWidget(UObject* ObjectToEdit);
+
+	// DataflowEditorActions
+	void OnPropertyValueChanged(const FPropertyChangedEvent& PropertyChangedEvent);
+	bool OnNodeVerifyTitleCommit(const FText& NewText, UEdGraphNode* GraphNode, FText& OutErrorMessage) const;
+	void OnNodeTitleCommitted(const FText& InNewText, ETextCommit::Type InCommitType, UEdGraphNode* GraphNode) const;
+	void OnNodeSelectionChanged(const TSet<UObject*>& NewSelection) const;
+
 
 	/** Scene in which the 3D sim space preview meshes live. Ownership shared with AdvancedPreviewSettingsWidget*/
 	TSharedPtr<FChaosClothPreviewScene> ClothPreviewScene;
@@ -125,33 +137,18 @@ private:
 	TArray<FName> ClothCollectionGroupNames;		// Data source for SelectedGroupNameComboBox
 
 
-	//~ Begin Dataflow support
-	 
+	// Dataflow
 	UDataflow* Dataflow = nullptr;
-
-	void EvaluateNode(FDataflowNode* Node, FDataflowOutput* Out);
+	FString DataflowTerminalPath = "";
+	TSharedPtr<Dataflow::FEngineContext> DataflowContext;
+	Dataflow::FTimestamp LastDataflowNodeTimestamp = Dataflow::FTimestamp::Invalid;
 
 	static const FName GraphCanvasTabId;
 	TSharedPtr<SDockTab> GraphEditorTab;
 	TSharedPtr<SDataflowGraphEditor> GraphEditor;
-	TSharedRef<SDataflowGraphEditor> CreateGraphEditorWidget();
-	void ReinitializeGraphEditorWidget();
 
 	static const FName NodeDetailsTabId;
 	TSharedPtr<SDockTab> NodeDetailsTab;
 	TSharedPtr<IStructureDetailsView> NodeDetailsEditor;
-	TSharedPtr<IStructureDetailsView> CreateNodeDetailsEditorWidget(UObject* ObjectToEdit);
-
-	FString DataflowTerminalPath = "";
-	TSharedPtr<Dataflow::FEngineContext> DataflowContext;
-	Dataflow::FTimestamp LastDataflowNodeTimestamp = Dataflow::FTimestamp::Invalid;
-
-	// DataflowEditorActions
-	void OnPropertyValueChanged(const FPropertyChangedEvent& PropertyChangedEvent);
-	bool OnNodeVerifyTitleCommit(const FText& NewText, UEdGraphNode* GraphNode, FText& OutErrorMessage) const;
-	void OnNodeTitleCommitted(const FText& InNewText, ETextCommit::Type InCommitType, UEdGraphNode* GraphNode) const;
-	void OnNodeSelectionChanged(const TSet<UObject*>& NewSelection) const;
-
-	//~ End Dataflow support
 };
 
diff --git a/Engine/Plugins/Experimental/ChaosClothAssetEditor/Source/ChaosClothAssetEditorTools/ChaosClothAssetEditorTools.Build.cs b/Engine/Plugins/Experimental/ChaosClothAssetEditor/Source/ChaosClothAssetEditorTools/ChaosClothAssetEditorTools.Build.cs
--- a/Engine/Plugins/Experimental/ChaosClothAssetEditor/Source/ChaosClothAssetEditorTools/ChaosClothAssetEditorTools.Build.cs
+++ b/Engine/Plugins/Experimental/ChaosClothAssetEditor/Source/ChaosClothAssetEditorTools/ChaosClothAssetEditorTools.Build.cs
@@ -55,6 +55,10 @@ public class ChaosClothAssetEditorTools : ModuleRules
 				"SkeletalMeshDescription",
 				"ChaosClothAsset",
 				"ChaosCaching",
+				"ChaosClothAssetDataflowNodes",
+				"DataflowCore",
+				"DataflowEngine",
+				"DataflowEditor"
 				// ... add private dependencies that you statically link with here ...	
 			}
 			);
diff --git a/Engine/Plugins/Experimental/ChaosClothAssetEditor/Source/ChaosClothAssetEditorTools/Private/ChaosClothAsset/ClothEditorContextObject.cpp b/Engine/Plugins/Experimental/ChaosClothAssetEditor/Source/ChaosClothAssetEditorTools/Private/ChaosClothAsset/ClothEditorContextObject.cpp
new file mode 100644
--- /dev/null
+++ b/Engine/Plugins/Experimental/ChaosClothAssetEditor/Source/ChaosClothAssetEditorTools/Private/ChaosClothAsset/ClothEditorContextObject.cpp
@@ -0,0 +1,33 @@
+// Copyright Epic Games, Inc. All Rights Reserved.
+
+#include "ChaosClothAsset/ClothEditorContextObject.h"
+
+#include UE_INLINE_GENERATED_CPP_BY_NAME(ClothEditorContextObject)
+
+void UClothEditorContextObject::Init(TWeakPtr<SDataflowGraphEditor> InDataflowGraphEditor, TObjectPtr<UDataflow> InDataflowGraph)
+{
+	DataflowGraphEditor = InDataflowGraphEditor;
+	DataflowGraph = InDataflowGraph;
+}
+
+TWeakPtr<SDataflowGraphEditor> UClothEditorContextObject::GetDataflowGraphEditor()
+{
+	return DataflowGraphEditor;
+}
+
+const TWeakPtr<const SDataflowGraphEditor> UClothEditorContextObject::GetDataflowGraphEditor() const
+{
+	return DataflowGraphEditor;
+}
+
+TObjectPtr<UDataflow> UClothEditorContextObject::GetDataflowGraph()
+{
+	return DataflowGraph;
+}
+
+const TObjectPtr<const UDataflow> UClothEditorContextObject::GetDataflowGraph() const
+{
+	return DataflowGraph;
+}
+
+
diff --git a/Engine/Plugins/Experimental/ChaosClothAssetEditor/Source/ChaosClothAssetEditorTools/Private/ChaosClothAsset/ClothWeightMapPaintTool.cpp b/Engine/Plugins/Experimental/ChaosClothAssetEditor/Source/ChaosClothAssetEditorTools/Private/ChaosClothAsset/ClothWeightMapPaintTool.cpp
--- a/Engine/Plugins/Experimental/ChaosClothAssetEditor/Source/ChaosClothAssetEditorTools/Private/ChaosClothAsset/ClothWeightMapPaintTool.cpp
+++ b/Engine/Plugins/Experimental/ChaosClothAssetEditor/Source/ChaosClothAssetEditorTools/Private/ChaosClothAsset/ClothWeightMapPaintTool.cpp
@@ -30,6 +30,17 @@
 
 #include "CanvasTypes.h"
 #include "CanvasItem.h"
+#include "ChaosClothAsset/ClothEditorContextObject.h"
+#include "ContextObjectStore.h"
+#include "Dataflow/DataflowEdNode.h"
+#include "Dataflow/DataflowObject.h"
+#include "ChaosClothAsset/AddWeightMapNode.h"
+#include "ChaosClothAsset/DataflowNodes.h"
+#include "GraphEditor.h"
+#include "Dataflow/DataflowSNode.h"
+#include "Dataflow/DataflowGraphEditor.h"
+#include "Widgets/Notifications/SNotificationList.h"
+#include "Framework/Notifications/NotificationManager.h"
 
 #include UE_INLINE_GENERATED_CPP_BY_NAME(ClothWeightMapPaintTool)
 
@@ -37,7 +48,7 @@ using namespace UE::Geometry;
 
 #define LOCTEXT_NAMESPACE "UClothEditorWeightMapPaintTool"
 
-namespace
+namespace UE::Chaos::ClothAsset::Private
 {
 	// probably should be something defined for the whole tool framework...
 #if WITH_EDITOR
@@ -53,9 +64,15 @@ namespace
  */
 UMeshSurfacePointTool* UClothEditorWeightMapPaintToolBuilder::CreateNewTool(const FToolBuilderState& SceneState) const
 {
-	UClothEditorWeightMapPaintTool* SculptTool = NewObject<UClothEditorWeightMapPaintTool>(SceneState.ToolManager);
-	SculptTool->SetWorld(SceneState.World);
-	return SculptTool;
+	UClothEditorWeightMapPaintTool* PaintTool = NewObject<UClothEditorWeightMapPaintTool>(SceneState.ToolManager);
+	PaintTool->SetWorld(SceneState.World);
+
+	if (UClothEditorContextObject* ContextObject = SceneState.ToolManager->GetContextObjectStore()->FindContext<UClothEditorContextObject>())
+	{
+		PaintTool->SetClothEditorContextObject(ContextObject);
+	}
+
+	return PaintTool;
 }
 
 
@@ -104,12 +121,12 @@ void UClothEditorWeightMapPaintTool::Setup()
 	Mesh->EnableVertexColors(FVector3f::One());
 	FAxisAlignedBox3d Bounds = Mesh->GetBounds(true);
 
-	TFuture<void> PrecomputeFuture = Async(WeightPaintToolAsyncExecTarget, [&]()
+	TFuture<void> PrecomputeFuture = Async(UE::Chaos::ClothAsset::Private::WeightPaintToolAsyncExecTarget, [&]()
 	{
 		PrecomputeFilterData();
 	});
 
-	TFuture<void> OctreeFuture = Async(WeightPaintToolAsyncExecTarget, [&]()
+	TFuture<void> OctreeFuture = Async(UE::Chaos::ClothAsset::Private::WeightPaintToolAsyncExecTarget, [&]()
 	{
 		// initialize dynamic octree
 		if (Mesh->TriangleCount() > 100000)
@@ -141,18 +158,13 @@ void UClothEditorWeightMapPaintTool::Setup()
 	PolyLassoMechanic->SpacingTolerance = 10.0f;
 	PolyLassoMechanic->OnDrawPolyLassoFinished.AddUObject(this, &UClothEditorWeightMapPaintTool::OnPolyLassoFinished);
 
-	WeightMapSetProperties = NewObject<UWeightMapSetProperties>(this);
-	WeightMapSetProperties->RestoreProperties(this, TEXT("ClothEditorWeightMapPaintTool"));
+	AddWeightMapProperties = NewObject<UClothEditorAddWeightMapProperties>(this);
+	AddWeightMapProperties->Initialize(this);
+	AddToolPropertySource(AddWeightMapProperties);
 
-	InitializeWeightMapNames();
-
-	WeightMapSetProperties->WatchProperty(WeightMapSetProperties->WeightMap, [&](FName) { OnSelectedWeightMapChanged(); });
-	UpdateActiveWeightMap();
-	AddToolPropertySource(WeightMapSetProperties);
-
-	ClothEditorWeightMapActions = NewObject<UClothEditorWeightMapActions>(this);
-	ClothEditorWeightMapActions->Initialize(this);
-	AddToolPropertySource(ClothEditorWeightMapActions);
+	UpdateWeightMapProperties = NewObject<UClothEditorUpdateWeightMapProperties>(this);
+	UpdateWeightMapProperties->Initialize(this);
+	AddToolPropertySource(UpdateWeightMapProperties);
 
 	// initialize other properties
 	FilterProperties = NewObject<UClothEditorWeightMapPaintBrushFilterProperties>(this);
@@ -210,7 +222,8 @@ void UClothEditorWeightMapPaintTool::Setup()
 	MeshElementsDisplay->CreateInWorld(DynamicMeshComponent->GetWorld(), DynamicMeshComponent->GetComponentTransform());
 	if (ensure(MeshElementsDisplay->Settings))
 	{
-		MeshElementsDisplay->Settings->RestoreProperties(this, TEXT("ClothEditorWeightMapPaintTool"));
+		MeshElementsDisplay->Settings->bShowNormalSeams = false;
+		MeshElementsDisplay->Settings->RestoreProperties(this, TEXT("ClothEditorWeightMapPaintTool2"));
 		AddToolPropertySource(MeshElementsDisplay->Settings);
 	}
 	MeshElementsDisplay->SetMeshAccessFunction([this](UMeshElementsVisualizer::ProcessDynamicMeshFunc ProcessFunc) {
@@ -246,6 +259,50 @@ void UClothEditorWeightMapPaintTool::Setup()
 
 	PrecomputeFuture.Wait();
 	OctreeFuture.Wait();
+
+
+	// Create an attribute layer to temporarily paint into
+	const int NumAttributeLayers = Mesh->Attributes()->NumWeightLayers();
+	Mesh->Attributes()->SetNumWeightLayers(NumAttributeLayers + 1);
+	ActiveWeightMap = Mesh->Attributes()->GetWeightLayer(NumAttributeLayers);
+	ActiveWeightMap->SetName(FName("PaintLayer"));
+
+	// Check if there is a weight map node selected
+	FString CurrentMapName;
+	TArray<float> CurrentWeights;
+	const bool bWeightMapNodeIsSelected = GetSelectedNodeInfo(CurrentMapName, CurrentWeights);
+
+	bShouldUpdateExistingNode = bWeightMapNodeIsSelected && (Mesh->VertexCount() == CurrentWeights.Num());
+
+	if (bShouldUpdateExistingNode)
+	{
+		check(Mesh->VertexCount() == CurrentWeights.Num());		// true if bShouldUpdateExistingNode is true
+
+		for (int32 VertexID = 0; VertexID < Mesh->VertexCount(); ++VertexID)
+		{
+			ActiveWeightMap->SetValue(VertexID, &CurrentWeights[VertexID]);
+		}
+
+		UpdateWeightMapProperties->Name = CurrentMapName;
+		SetToolPropertySourceEnabled(AddWeightMapProperties, false);
+		SetToolPropertySourceEnabled(UpdateWeightMapProperties, true);
+	}
+	else
+	{
+		SetToolPropertySourceEnabled(AddWeightMapProperties, true);
+		SetToolPropertySourceEnabled(UpdateWeightMapProperties, false);
+	}
+
+	// update colors
+	DynamicMeshComponent->FastNotifyVertexAttributesUpdated(EMeshRenderAttributeFlags::VertexColors);
+	GetToolManager()->PostInvalidation();
+
+}
+
+
+void UClothEditorWeightMapPaintTool::SetClothEditorContextObject(TObjectPtr<UClothEditorContextObject> InClothEditorContextObject)
+{
+	ClothEditorContextObject = InClothEditorContextObject;
 }
 
 void UClothEditorWeightMapPaintTool::Shutdown(EToolShutdownType ShutdownType)
@@ -257,13 +314,11 @@ void UClothEditorWeightMapPaintTool::Shutdown(EToolShutdownType ShutdownType)
 
 	if (ensure(MeshElementsDisplay->Settings))
 	{
-		MeshElementsDisplay->Settings->SaveProperties(this, TEXT("ClothEditorWeightMapPaintTool"));
+		MeshElementsDisplay->Settings->SaveProperties(this, TEXT("ClothEditorWeightMapPaintTool2"));
 	}
 	MeshElementsDisplay->Disconnect();
 
 	FilterProperties->SaveProperties(this);
-	WeightMapSetProperties->SaveProperties(this, TEXT("ClothEditorWeightMapPaintTool"));
-
 	if (PreviewMeshActor != nullptr)
 	{
 		PreviewMeshActor->Destroy();
@@ -277,10 +332,16 @@ void UClothEditorWeightMapPaintTool::Shutdown(EToolShutdownType ShutdownType)
 void UClothEditorWeightMapPaintTool::CommitResult(UBaseDynamicMeshComponent* Component, bool bModifiedTopology)
 {
 	GetToolManager()->BeginUndoTransaction(LOCTEXT("WeightPaintToolTransactionName", "Paint Weights"));
-	Component->ProcessMesh([&](const FDynamicMesh3& CurMesh)
+
+	if (bShouldUpdateExistingNode)
 	{
-		UE::ToolTarget::CommitDynamicMeshUpdate(Target, CurMesh, true);
-	});
+		UpdateSelectedNode();
+	}
+	else
+	{
+		AddNewNode(AddWeightMapProperties->Name);
+	}
+
 	GetToolManager()->EndUndoTransaction();
 }
 
@@ -1119,7 +1180,7 @@ void UClothEditorWeightMapPaintTool::OnTick(float DeltaTime)
 
 			// append updated ROI to modified region (async)
 			FDynamicMesh3* Mesh = GetSculptMesh();
-			TFuture<void> AccumulateROI = Async(WeightPaintToolAsyncExecTarget, [&]()
+			TFuture<void> AccumulateROI = Async(UE::Chaos::ClothAsset::Private::WeightPaintToolAsyncExecTarget, [&]()
 			{
 				UE::Geometry::VertexToTriangleOneRing(Mesh, VertexROI, AccumulatedTriangleROI);
 			});
@@ -1141,6 +1202,17 @@ void UClothEditorWeightMapPaintTool::OnTick(float DeltaTime)
 
 }
 
+bool UClothEditorWeightMapPaintTool::CanAccept() const
+{
+	if (bShouldUpdateExistingNode)
+	{
+		return true;
+	}
+	else
+	{
+		return !(AddWeightMapProperties->Name.IsEmpty());
+	}
+}
 
 
 
@@ -1220,106 +1292,118 @@ void UClothEditorWeightMapPaintTool::ClearAllWeightsAction()
 }
 
 
-void UClothEditorWeightMapPaintTool::InitializeWeightMapNames()
+bool UClothEditorWeightMapPaintTool::GetSelectedNodeInfo(FString& OutMapName, TArray<float>& OutWeights)
 {
-	const FDynamicMesh3* const Mesh = GetSculptMesh();
-
-	if (!WeightMapSetProperties || !Mesh || !Mesh->HasAttributes())
-	{
-		return;
-	}
-
-	TArray<FString>& PropWeightMapNames = WeightMapSetProperties->WeightMapsList;
-	PropWeightMapNames.Reset();
-	PropWeightMapNames.Add(TEXT("None"));
-
-	// get names from the mesh attributes
-	for (int32 WeightMapIndex = 0; WeightMapIndex < Mesh->Attributes()->NumWeightLayers(); ++WeightMapIndex)
+	const TSharedPtr<const SDataflowGraphEditor> GraphEditor = ClothEditorContextObject->GetDataflowGraphEditor().Pin();
+	if (const UEdGraphNode* const SingleSelectedNode = GraphEditor->GetSingleSelectedNode())
 	{
-		const FName& MeshWeightMapName = Mesh->Attributes()->GetWeightLayer(WeightMapIndex)->GetName();
-		PropWeightMapNames.Add(MeshWeightMapName.ToString());
+		const UDataflowEdNode* const SelectedDataflowEdNode = Cast<UDataflowEdNode>(SingleSelectedNode);
+		const TSharedPtr<const FDataflowNode> DataflowNode = SelectedDataflowEdNode->GetDataflowNode();
+		if (const FChaosClothAssetAddWeightMapNode* const WeightMapNode = DataflowNode->AsType<FChaosClothAssetAddWeightMapNode>())
+		{
+			OutMapName = WeightMapNode->Name;
+			OutWeights = WeightMapNode->VertexWeights;
+			return true;
+		}
 	}
 
-	// if selected weight map name is no longer in the list of valid names, reset it
-	if (!PropWeightMapNames.Contains(WeightMapSetProperties->WeightMap.ToString()))
-	{
-		WeightMapSetProperties->WeightMap = FName(PropWeightMapNames[0]);
-	}
+	// No single node selected, or the selected node is not FChaosClothAssetAddWeightMapNode
+	return false;
 }
 
 
-void UClothEditorWeightMapPaintTool::AddWeightMapAction(const FName& NewWeightMapName)
+void UClothEditorWeightMapPaintTool::AddNewNode(const FString& NewMapName)
 {
-	if (DynamicMeshComponent && DynamicMeshComponent->GetMesh() && DynamicMeshComponent->GetMesh()->HasAttributes())
-	{
-		FDynamicMesh3* const Mesh = DynamicMeshComponent->GetMesh();
+	check(ActiveWeightMap);
+	TArray<float> CurrentWeights;
+	GetCurrentWeightMap(CurrentWeights);
 
-		FDynamicMeshAttributeSet* const MeshAttributes = Mesh->Attributes();
-		const int32 NumExistingWeights = MeshAttributes->NumWeightLayers();
-		
-		if (NewWeightMapName.IsNone())
-		{
-			GetToolManager()->DisplayMessage(LOCTEXT("InvalidAttributeName", "Invalid weight map name"), EToolMessageLevel::UserWarning);
-			return;
-		}
+	// TODO: If we are operating on a single pattern, expand the weight map to cover the whole LOD
 
-		// Check for matching name	
-		for (int32 WeightMapIndex = 0; WeightMapIndex < NumExistingWeights; ++WeightMapIndex)
-		{
-			const FDynamicMeshWeightAttribute* const ExistingWeightMap = MeshAttributes->GetWeightLayer(WeightMapIndex);
-			if (ExistingWeightMap->GetName() == NewWeightMapName)
-			{
-				GetToolManager()->DisplayMessage(LOCTEXT("ErrorAddingDuplicateNameMessage", "Weight map with this name already exists"), EToolMessageLevel::UserWarning);
-				return;
-			}
-		}
+	TSharedPtr<SDataflowGraphEditor> GraphEditor = ClothEditorContextObject->GetDataflowGraphEditor().Pin();
 
-		MeshAttributes->SetNumWeightLayers(NumExistingWeights + 1);
+	//
+	// We want to turn this:
+	// 
+	// [SelectedNode] ----> [ExistingDownstreamNode(s)]
+	// 
+	// to this:
+	//
+	// [SelectedNode] ----> [NewNode] ----> [ExistingDownstreamNode(s)]
+	//
 
-		const int32 NewWeightMapIndex = NumExistingWeights;
-		FDynamicMeshWeightAttribute* const NewWeightMap = MeshAttributes->GetWeightLayer(NewWeightMapIndex);
-		NewWeightMap->SetName(NewWeightMapName);
+	// First find the Collection output of the currently selected node, and any nodes it's connected to
 
-		// Refresh the drop down of available weight maps and select the new one
-		InitializeWeightMapNames();
-		WeightMapSetProperties->SetSelectedFromWeightMapIndex(NewWeightMapIndex);
+	const UEdGraphNode* const SingleSelectedNode = GraphEditor->GetSingleSelectedNode();
+	checkf(SingleSelectedNode, TEXT("There should be a single selected node in the Dataflow graph"));
 
-		// Clear the New Name text field
-		ClothEditorWeightMapActions->NewWeightMapName.Empty();
-	}
-}
-
-
-void UClothEditorWeightMapPaintTool::DeleteWeightMapAction(const FName& SelectedWeightMapName)
-{
-	FDynamicMeshAttributeSet* const MeshAttributes = DynamicMeshComponent->GetMesh()->Attributes();
-
-	bool bFoundSelectedWeightMap = false;
-	
-	const int32 NumWeightMaps = MeshAttributes->NumWeightLayers();
-	for (int32 WeightMapIndex = 0; WeightMapIndex < NumWeightMaps; ++WeightMapIndex)
+	const UDataflowEdNode* const SelectedDataflowEdNode = CastChecked<UDataflowEdNode>(SingleSelectedNode);
+	const TSharedPtr<const FDataflowNode> DataflowNode = SelectedDataflowEdNode->GetDataflowNode();
+	UEdGraphPin* SelectedNodeOutputPin = nullptr;
+	TArray<UEdGraphPin*> ExistingNodeInputPins;
+	for (const FDataflowOutput* Output : DataflowNode->GetOutputs())
 	{
-		const FDynamicMeshWeightAttribute* const WeightMap = MeshAttributes->GetWeightLayer(WeightMapIndex);
-		if (WeightMap->GetName() == SelectedWeightMapName)
+		if (Output->GetType() == FName("FManagedArrayCollection"))
 		{
-			bFoundSelectedWeightMap = true;
-			MeshAttributes->RemoveWeightLayer(WeightMapIndex);		// Hopefully nothing else is holding indices into this...
+			SelectedNodeOutputPin = SelectedDataflowEdNode->FindPin(*Output->GetName().ToString(), EGPD_Output);
+			ExistingNodeInputPins = SelectedNodeOutputPin->LinkedTo;
 			break;
 		}
 	}
 
-	if (!bFoundSelectedWeightMap)
+	// Now add the new node and re-wire the graph
+
+	const TObjectPtr<UDataflow> Dataflow = ClothEditorContextObject->GetDataflowGraph();
+	const TSharedPtr<FAssetSchemaAction_Dataflow_CreateNode_DataflowEdNode> NodeAction = FAssetSchemaAction_Dataflow_CreateNode_DataflowEdNode::CreateAction(Dataflow, "FChaosClothAssetAddWeightMapNode");
+
+	UEdGraphNode* const NewEdNode = NodeAction->PerformAction(Dataflow, nullptr, GraphEditor->GetPasteLocation(), true);
+	UDataflowEdNode* const NewDataflowEdNode = CastChecked<UDataflowEdNode>(NewEdNode);
+	const TSharedPtr<FDataflowNode> NewDataflowNode = NewDataflowEdNode->GetDataflowNode();
+	FChaosClothAssetAddWeightMapNode* const NewWeightMapNode = NewDataflowNode->AsType<FChaosClothAssetAddWeightMapNode>();
+
+	NewWeightMapNode->VertexWeights = CurrentWeights;
+	NewWeightMapNode->Name = AddWeightMapProperties->Name;
+	// NewWeightMapNode->SetName(FName(AddWeightMapProperties->NewMapName));   TODO: This doesn't seem to work to set the Node name
+
+	// Re-wire the graph
+	FDataflowInput* NewNodeCollectionInput = NewWeightMapNode->FindInput(&(NewWeightMapNode->Collection));
+	FDataflowOutput* NewNodeCollectionOutput = NewWeightMapNode->FindOutput(&(NewWeightMapNode->Collection));
+	UEdGraphPin* NewNodeInputPin = NewDataflowEdNode->FindPin(*NewNodeCollectionInput->GetName().ToString(), EGPD_Input);
+	UEdGraphPin* NewNodeOutputPin = NewDataflowEdNode->FindPin(*NewNodeCollectionInput->GetName().ToString(), EGPD_Output);
+
+	bool bWasAnythingChanged = Dataflow->GetSchema()->TryCreateConnection(SelectedNodeOutputPin, NewNodeInputPin);
+
+	for (UEdGraphPin* DownstreamInputPin : ExistingNodeInputPins)
 	{
-		GetToolManager()->DisplayMessage(LOCTEXT("InvalidExistingWeightMapName", "No weight map with the specified name exists"), EToolMessageLevel::UserWarning);
+		bWasAnythingChanged |= Dataflow->GetSchema()->TryCreateConnection(NewNodeOutputPin, DownstreamInputPin);
 	}
-	else
+
+	if (bWasAnythingChanged)
 	{
-		InitializeWeightMapNames();
-		UpdateActiveWeightMap();
-		GetToolManager()->DisplayMessage(FText(), EToolMessageLevel::UserWarning);
+		Dataflow->NotifyGraphChanged();
 	}
+
 }
 
+
+void UClothEditorWeightMapPaintTool::UpdateSelectedNode()
+{
+	check(ActiveWeightMap);
+	TArray<float> CurrentWeights;
+	GetCurrentWeightMap(CurrentWeights);
+
+	const TSharedPtr<SGraphEditor> GraphEditor = ClothEditorContextObject->GetDataflowGraphEditor().Pin();
+	UEdGraphNode* const SingleSelectedNode = GraphEditor->GetSingleSelectedNode();
+	checkf(SingleSelectedNode, TEXT("There should be a single selected node in the Dataflow graph"));
+
+	UDataflowEdNode* const SelectedDataflowEdNode = CastChecked<UDataflowEdNode>(SingleSelectedNode);
+	const TSharedPtr<FDataflowNode> DataflowNode = SelectedDataflowEdNode->GetDataflowNode();
+	FChaosClothAssetAddWeightMapNode* const WeightMapNode = DataflowNode->AsType<FChaosClothAssetAddWeightMapNode>();
+	WeightMapNode->VertexWeights = CurrentWeights;
+	WeightMapNode->Invalidate();
+}
+
+
 //
 // Change Tracking
 //
@@ -1510,66 +1594,18 @@ int32 UClothEditorWeightMapPaintTool::GetBrushNearestVertex() const
 	return Vertices[TriangleVertex];
 }
 
-
-void UClothEditorWeightMapPaintTool::OnSelectedWeightMapChanged()
+void UClothEditorWeightMapPaintTool::GetCurrentWeightMap(TArray<float>& OutWeights) const
 {
-	GetToolManager()->BeginUndoTransaction(LOCTEXT("ChangeActiveWeightLayer", "Change Weight Layer"));
-
-	FDynamicMeshWeightAttribute* CurrentActiveLayer = ActiveWeightMap;
-	UpdateActiveWeightMap();
-	FDynamicMeshWeightAttribute* NewActiveLayer = ActiveWeightMap; 
-
-	if (CurrentActiveLayer != NewActiveLayer)
+	if (ActiveWeightMap)
 	{
-		// A value of -1 matches the property set's "None" label.
-		int32 CurrentActiveLayerIndex = -1;
-		int32 NewActiveLayerIndex = -1;
-		if (CurrentActiveLayer) {
-			WeightMapSetProperties->GetWeightMapsFunc().Find(CurrentActiveLayer->GetName().ToString(), CurrentActiveLayerIndex);
-			CurrentActiveLayerIndex = CurrentActiveLayerIndex - 1;
-		}
-		if (NewActiveLayer)
+		const FDynamicMesh3* Mesh = GetSculptMesh();
+		const int32 NumVertices = Mesh->VertexCount();
+		OutWeights.SetNumUninitialized(NumVertices);
+		for (int32 VertexID = 0; VertexID < NumVertices; ++VertexID)
 		{
-			WeightMapSetProperties->GetWeightMapsFunc().Find(NewActiveLayer->GetName().ToString(), NewActiveLayerIndex);
-			NewActiveLayerIndex = NewActiveLayerIndex - 1;
+			ActiveWeightMap->GetValue(VertexID, &OutWeights[VertexID]);
 		}
-
-		TUniquePtr<TSimpleValueLambdaChange<int32>> WeightLayerChange = MakeUnique<TSimpleValueLambdaChange<int32>>();
-		WeightLayerChange->FromValue = CurrentActiveLayerIndex;
-		WeightLayerChange->ToValue = NewActiveLayerIndex;
-		WeightLayerChange->ValueChangeFunc = [this](UObject*, int32 FromIndex, int32 ToIndex, bool)
-		{
-			this->WeightMapSetProperties->SetSelectedFromWeightMapIndex(ToIndex);
-			this->WeightMapSetProperties->SilentUpdateWatched();		// to prevent OnSelectedWeightMapChanged() from being called immediately
-			this->UpdateActiveWeightMap();
-		};
-		GetToolManager()->EmitObjectChange(this, MoveTemp(WeightLayerChange), LOCTEXT("ChangeActiveWeightSet", "Change Weight Set"));
 	}
-
-	GetToolManager()->EndUndoTransaction();
-}
-
-
-void UClothEditorWeightMapPaintTool::UpdateActiveWeightMap()
-{
-	int32 SelectedWeightSetIndex;
-	FName SelectedName = WeightMapSetProperties->WeightMap;
-	WeightMapSetProperties->GetWeightMapsFunc().Find(SelectedName.ToString(), SelectedWeightSetIndex);
-
-	FDynamicMesh3* Mesh = GetSculptMesh();
-	// Properties are 1-indexed, due to an initial None option
-	if (Mesh->Attributes()->NumWeightLayers() >= SelectedWeightSetIndex && SelectedWeightSetIndex > 0)
-	{
-		ActiveWeightMap = Mesh->Attributes()->GetWeightLayer(SelectedWeightSetIndex-1);
-	}
-	else
-	{
-		ActiveWeightMap = nullptr;
-	}
-
-	// update colors
-	DynamicMeshComponent->FastNotifyVertexAttributesUpdated(EMeshRenderAttributeFlags::VertexColors);
-	GetToolManager()->PostInvalidation();
 }
 
 
@@ -1624,15 +1660,6 @@ void UClothEditorWeightMapPaintTool::ApplyAction(EClothEditorWeightMapPaintToolA
 	case EClothEditorWeightMapPaintToolActions::ClearAll:
 		ClearAllWeightsAction();
 		break;
-
-	case EClothEditorWeightMapPaintToolActions::AddWeightMap:
-		AddWeightMapAction(FName(ClothEditorWeightMapActions->NewWeightMapName));
-		break;
-
-	case EClothEditorWeightMapPaintToolActions::DeleteWeightMap:
-		DeleteWeightMapAction(WeightMapSetProperties->WeightMap);
-		break;
-
 	}
 }
 
diff --git a/Engine/Plugins/Experimental/ChaosClothAssetEditor/Source/ChaosClothAssetEditorTools/Public/ChaosClothAsset/ClothEditorContextObject.h b/Engine/Plugins/Experimental/ChaosClothAssetEditor/Source/ChaosClothAssetEditorTools/Public/ChaosClothAsset/ClothEditorContextObject.h
new file mode 100644
--- /dev/null
+++ b/Engine/Plugins/Experimental/ChaosClothAssetEditor/Source/ChaosClothAssetEditorTools/Public/ChaosClothAsset/ClothEditorContextObject.h
@@ -0,0 +1,33 @@
+// Copyright Epic Games, Inc. All Rights Reserved.
+
+#pragma once
+
+#include "UObject/Object.h"
+#include "ClothEditorContextObject.generated.h"
+
+class SDataflowGraphEditor;
+class UDataflow;
+
+UCLASS()
+class CHAOSCLOTHASSETEDITORTOOLS_API UClothEditorContextObject : public UObject
+{
+	GENERATED_BODY()
+
+public:
+
+	void Init(TWeakPtr<SDataflowGraphEditor> DataflowGraphEditor, TObjectPtr<UDataflow> DataflowGraph);
+
+	TWeakPtr<SDataflowGraphEditor> GetDataflowGraphEditor();
+	const TWeakPtr<const SDataflowGraphEditor> GetDataflowGraphEditor() const;
+
+	TObjectPtr<UDataflow> GetDataflowGraph();
+	const TObjectPtr<const UDataflow> GetDataflowGraph() const;
+
+private:
+
+	TWeakPtr<SDataflowGraphEditor> DataflowGraphEditor;
+
+	UPROPERTY()
+	TObjectPtr<UDataflow> DataflowGraph = nullptr;
+
+};
diff --git a/Engine/Plugins/Experimental/ChaosClothAssetEditor/Source/ChaosClothAssetEditorTools/Public/ChaosClothAsset/ClothWeightMapPaintTool.h b/Engine/Plugins/Experimental/ChaosClothAssetEditor/Source/ChaosClothAssetEditorTools/Public/ChaosClothAsset/ClothWeightMapPaintTool.h
--- a/Engine/Plugins/Experimental/ChaosClothAssetEditor/Source/ChaosClothAssetEditorTools/Public/ChaosClothAsset/ClothWeightMapPaintTool.h
+++ b/Engine/Plugins/Experimental/ChaosClothAssetEditor/Source/ChaosClothAssetEditorTools/Public/ChaosClothAsset/ClothWeightMapPaintTool.h
@@ -26,6 +26,7 @@
 class UMeshElementsVisualizer;
 class UWeightMapEraseBrushOpProps;
 class UWeightMapPaintBrushOpProps;
+class UClothEditorContextObject;
 
 DECLARE_STATS_GROUP(TEXT("WeightMapPaintTool"), STATGROUP_WeightMapPaintTool, STATCAT_Advanced);
 DECLARE_CYCLE_STAT(TEXT("WeightMapPaintTool_UpdateROI"), WeightMapPaintTool_UpdateROI, STATGROUP_WeightMapPaintTool);
@@ -163,9 +164,6 @@ enum class EClothEditorWeightMapPaintToolActions
 
 	FloodFillCurrent,
 	ClearAll,
-
-	AddWeightMap,
-	DeleteWeightMap
 };
 
 
@@ -205,26 +203,26 @@ public:
 
 
 UCLASS()
-class CHAOSCLOTHASSETEDITORTOOLS_API UClothEditorWeightMapActions : public UClothEditorWeightMapPaintToolActionPropertySet
+class CHAOSCLOTHASSETEDITORTOOLS_API UClothEditorAddWeightMapProperties : public UClothEditorWeightMapPaintToolActionPropertySet
 {
 	GENERATED_BODY()
 
 public:
 
-	UPROPERTY(EditAnywhere, Category = AddWeightMap, meta = (DisplayName = "New Weight Map Name"))
-	FString NewWeightMapName;
+	UPROPERTY(EditAnywhere, Category = AddNode, meta = (DisplayName = "Name"))
+	FString Name = "WeightMap";
+};
 
-	UFUNCTION(CallInEditor, Category = AddWeightMap, meta = (DisplayPriority = 2))
-	void AddNewWeightMap()
-	{
-		PostAction(EClothEditorWeightMapPaintToolActions::AddWeightMap);
-	}
 
-	UFUNCTION(CallInEditor, Category = DeleteWeightMap, meta = (DisplayPriority = 3))
-	void DeleteSelectedWeightMap()
-	{
-		PostAction(EClothEditorWeightMapPaintToolActions::DeleteWeightMap);
-	}
+UCLASS()
+class CHAOSCLOTHASSETEDITORTOOLS_API UClothEditorUpdateWeightMapProperties : public UClothEditorWeightMapPaintToolActionPropertySet
+{
+	GENERATED_BODY()
+
+public:
+
+	UPROPERTY(VisibleAnywhere, Category = UpdateNode, meta = (DisplayName = "Name"))
+	FString Name;
 };
 
 /**
@@ -246,7 +244,7 @@ public:
 
 	virtual bool HasCancel() const override { return true; }
 	virtual bool HasAccept() const override { return true; }
-	virtual bool CanAccept() const override { return true; }
+	virtual bool CanAccept() const override;
 
 	virtual bool OnUpdateHover(const FInputDeviceRay& DevicePos) override;
 
@@ -256,11 +254,10 @@ public:
 
 	virtual void CommitResult(UBaseDynamicMeshComponent* Component, bool bModifiedTopology) override;
 
+	void SetClothEditorContextObject(TObjectPtr<UClothEditorContextObject> InClothEditorContextObject);
 
 public:
 
-	UPROPERTY()
-	TObjectPtr<UWeightMapSetProperties> WeightMapSetProperties;
 
 	/** Filters on paint brush */
 	UPROPERTY()
@@ -278,9 +275,6 @@ public:
 	void FloodFillCurrentWeightAction();
 	void ClearAllWeightsAction();
 
-	void AddWeightMapAction(const FName& NewWeightMapName);
-	void DeleteWeightMapAction(const FName& SelectedWeightMapName);
-
 	void SetVerticesToWeightMap(const TSet<int32>& Vertices, double WeightValue, bool bIsErase);
 
 	bool HaveVisibilityFilter() const;
@@ -321,7 +315,10 @@ public:
 	TObjectPtr<UMeshWeightMapPaintToolActions> ActionsProps;
 
 	UPROPERTY()
-	TObjectPtr<UClothEditorWeightMapActions> ClothEditorWeightMapActions;
+	TObjectPtr<UClothEditorAddWeightMapProperties> AddWeightMapProperties;
+
+	UPROPERTY()
+	TObjectPtr<UClothEditorUpdateWeightMapProperties> UpdateWeightMapProperties;
 
 protected:
 	bool bHavePendingAction = false;
@@ -356,6 +353,9 @@ protected:
 	UPROPERTY()
 	TObjectPtr<UMeshElementsVisualizer> MeshElementsDisplay;
 
+	UPROPERTY()
+	TObjectPtr<UClothEditorContextObject> ClothEditorContextObject = nullptr;
+
 	// realtime visualization
 	void OnDynamicMeshComponentChanged(UDynamicMeshComponent* Component, const FMeshVertexChange* Change, bool bRevert);
 	FDelegateHandle OnDynamicMeshComponentChangedHandle;
@@ -365,8 +365,8 @@ protected:
 	double GetCurrentWeightValueUnderBrush() const;
 	FVector3d CurrentBaryCentricCoords;
 	int32 GetBrushNearestVertex() const;
-	void OnSelectedWeightMapChanged();
-	void UpdateActiveWeightMap();
+
+	void GetCurrentWeightMap(TArray<float>& OutWeights) const;
 
 	void UpdateSubToolType(EClothEditorWeightMapPaintInteractionType NewType);
 
@@ -418,13 +418,16 @@ protected:
 	TArray<int32> NormalSeamEdges;
 	void PrecomputeFilterData();
 
-	// Populate the drop-down WeightMapsList in WeightMapSetProperties
-	void InitializeWeightMapNames();
-
 protected:
 	virtual bool ShowWorkPlane() const override { return false; }
 
+	friend class UClothEditorWeightMapPaintToolBuilder;
 
+	// Node graph editor support
+	bool GetSelectedNodeInfo(FString& OutMapName, TArray<float>& OutWeights);
+	void AddNewNode(const FString& NewMapName);
+	void UpdateSelectedNode();
+	bool bShouldUpdateExistingNode = true;
 };
 
 
diff --git a/Engine/Plugins/Experimental/ChaosClothAssetEditor/Source/ChaosClothAssetTools/Private/ChaosClothAsset/ClothPatternToDynamicMesh.cpp b/Engine/Plugins/Experimental/ChaosClothAssetEditor/Source/ChaosClothAssetTools/Private/ChaosClothAsset/ClothPatternToDynamicMesh.cpp
--- a/Engine/Plugins/Experimental/ChaosClothAssetEditor/Source/ChaosClothAssetTools/Private/ChaosClothAsset/ClothPatternToDynamicMesh.cpp
+++ b/Engine/Plugins/Experimental/ChaosClothAssetEditor/Source/ChaosClothAssetTools/Private/ChaosClothAsset/ClothPatternToDynamicMesh.cpp
@@ -217,12 +217,10 @@ private:
 };
 
 
-void FClothPatternToDynamicMesh::Convert(const UChaosClothAsset* ClothAssetMeshIn, int32 LODIndex, int32 PatternIndex, bool bGet2DPattern, UE::Geometry::FDynamicMesh3& MeshOut)
+void FClothPatternToDynamicMesh::Convert(const TSharedPtr<const FManagedArrayCollection> ClothCollection, int32 LODIndex, int32 PatternIndex, bool bGet2DPattern, UE::Geometry::FDynamicMesh3& MeshOut)
 {
 	using namespace UE::Chaos::ClothAsset;
 
-	const TSharedPtr<const FManagedArrayCollection> ClothCollection = ClothAssetMeshIn->GetClothCollection();
-	check(ClothCollection.IsValid());
 	const FCollectionClothConstFacade ClothFacade(ClothCollection);
 
 	// Actual conversion
@@ -247,6 +245,13 @@ void FClothPatternToDynamicMesh::Convert(const UChaosClothAsset* ClothAssetMeshI
 	}
 }
 
+void FClothPatternToDynamicMesh::Convert(const UChaosClothAsset* ClothAssetMeshIn, int32 LODIndex, int32 PatternIndex, bool bGet2DPattern, UE::Geometry::FDynamicMesh3& MeshOut)
+{
+	const TSharedPtr<const FManagedArrayCollection> ClothCollection = ClothAssetMeshIn->GetClothCollection();
+	check(ClothCollection.IsValid());
+
+	Convert(ClothCollection, LODIndex, PatternIndex, bGet2DPattern, MeshOut);
+}
 
 #else
 
diff --git a/Engine/Plugins/Experimental/ChaosClothAssetEditor/Source/ChaosClothAssetTools/Public/ChaosClothAsset/ClothPatternToDynamicMesh.h b/Engine/Plugins/Experimental/ChaosClothAssetEditor/Source/ChaosClothAssetTools/Public/ChaosClothAsset/ClothPatternToDynamicMesh.h
--- a/Engine/Plugins/Experimental/ChaosClothAssetEditor/Source/ChaosClothAssetTools/Public/ChaosClothAsset/ClothPatternToDynamicMesh.h
+++ b/Engine/Plugins/Experimental/ChaosClothAssetEditor/Source/ChaosClothAssetTools/Public/ChaosClothAsset/ClothPatternToDynamicMesh.h
@@ -6,6 +6,7 @@
 
 namespace UE::Geometry { class FDynamicMesh3; }
 class UChaosClothAsset;
+struct FManagedArrayCollection;
 
 /**
  * Convert a single pattern from a ClothAsset to a FDynamicMesh3
@@ -14,5 +15,7 @@ class CHAOSCLOTHASSETTOOLS_API FClothPatternToDynamicMesh
 {
 public:
 
+	void Convert(const TSharedPtr<const FManagedArrayCollection> ClothCollection, int32 LODIndex, int32 PatternIndex, bool bGet2DPattern, UE::Geometry::FDynamicMesh3& MeshOut);
+
 	void Convert(const UChaosClothAsset* ClothAssetMeshIn, int32 LODIndex, int32 PatternIndex, bool bGet2DPattern, UE::Geometry::FDynamicMesh3& MeshOut);
 };
