Chaos Cloth:
- Fixed a bug in the clothing exporter that was adding an empty LOD to all exports.
- Fixed the uninitialized ImportNode's Import LOD property member.

#rnx

diff --git a/Engine/Plugins/Experimental/ChaosClothAsset/Source/ChaosClothAsset/Private/ChaosClothAsset/CollectionClothFacade.cpp b/Engine/Plugins/Experimental/ChaosClothAsset/Source/ChaosClothAsset/Private/ChaosClothAsset/CollectionClothFacade.cpp
--- a/Engine/Plugins/Experimental/ChaosClothAsset/Source/ChaosClothAsset/Private/ChaosClothAsset/CollectionClothFacade.cpp
+++ b/Engine/Plugins/Experimental/ChaosClothAsset/Source/ChaosClothAsset/Private/ChaosClothAsset/CollectionClothFacade.cpp
@@ -79,7 +79,7 @@ namespace UE::Chaos::ClothAsset
 		return FCollectionClothLodFacade(GetClothCollection(), LodIndex);
 	}
 
-	void FCollectionClothFacade::SetNumLod(int32 InNumLods)
+	void FCollectionClothFacade::SetNumLods(int32 InNumLods)
 	{
 		check(IsValid());
 
diff --git a/Engine/Plugins/Experimental/ChaosClothAsset/Source/ChaosClothAsset/Public/ChaosClothAsset/CollectionClothFacade.h b/Engine/Plugins/Experimental/ChaosClothAsset/Source/ChaosClothAsset/Public/ChaosClothAsset/CollectionClothFacade.h
--- a/Engine/Plugins/Experimental/ChaosClothAsset/Source/ChaosClothAsset/Public/ChaosClothAsset/CollectionClothFacade.h
+++ b/Engine/Plugins/Experimental/ChaosClothAsset/Source/ChaosClothAsset/Public/ChaosClothAsset/CollectionClothFacade.h
@@ -80,7 +80,7 @@ namespace UE::Chaos::ClothAsset
 		FCollectionClothLodFacade AddGetLod() { return GetLod(AddLod()); }
 
 		/** Set a new number of LODs for this cloth. */
-		void SetNumLod(int32 NumLods);
+		void SetNumLods(int32 NumLods);
 
 		/** Add a new weight map to this cloth. Access is then done per pattern. */
 		void AddWeightMap(const FName& Name);
diff --git a/Engine/Plugins/Experimental/ChaosClothAssetEditor/Source/ChaosClothAssetTools/Private/ChaosClothAsset/ClothingAssetToClothAssetExporter.cpp b/Engine/Plugins/Experimental/ChaosClothAssetEditor/Source/ChaosClothAssetTools/Private/ChaosClothAsset/ClothingAssetToClothAssetExporter.cpp
--- a/Engine/Plugins/Experimental/ChaosClothAssetEditor/Source/ChaosClothAssetTools/Private/ChaosClothAsset/ClothingAssetToClothAssetExporter.cpp
+++ b/Engine/Plugins/Experimental/ChaosClothAssetEditor/Source/ChaosClothAssetTools/Private/ChaosClothAsset/ClothingAssetToClothAssetExporter.cpp
@@ -36,17 +36,21 @@ void UClothingAssetToChaosClothAssetExporter::Export(const UClothingAssetBase* C
 	check(ClothAsset);
 	FCollectionClothFacade Cloth(ClothAsset->GetClothCollection());
 
+	const int32 NumLods = ClothingAssetCommon->LodData.Num();
+	Cloth.SetNumLods(NumLods);
+
 	// Create the LODs
-	for (const FClothLODDataCommon& ClothLODData : ClothingAssetCommon->LodData)
+	for (int32 LodIndex = 0; LodIndex < NumLods; ++LodIndex)
 	{
+		const FClothLODDataCommon& ClothLODData = ClothingAssetCommon->LodData[LodIndex];
 		const FClothPhysicalMeshData& PhysicalMeshData = ClothLODData.PhysicalMeshData;
 
 		// Unwrap the physical mesh data into the pattern and rest meshes
-		FCollectionClothLodFacade ClothLod = Cloth.AddGetLod();
+		FCollectionClothLodFacade ClothLod = Cloth.GetLod(LodIndex);
 		ClothLod.Initialize(PhysicalMeshData.Vertices, PhysicalMeshData.Indices);
 	}
 
-	if (Cloth.GetNumLods())
+	if (NumLods)
 	{
 		// Set the render mesh to duplicate the sim mesh
 		ClothAsset->CopySimMeshToRenderMesh();
