Fix lambda parameter name shadowing a class member variable name

CIS 282862

diff --git a/Engine/Plugins/Experimental/ChaosClothAssetEditor/Source/ChaosClothAssetEditor/Private/ChaosClothAsset/ClothEditorToolkit.cpp b/Engine/Plugins/Experimental/ChaosClothAssetEditor/Source/ChaosClothAssetEditor/Private/ChaosClothAsset/ClothEditorToolkit.cpp
--- a/Engine/Plugins/Experimental/ChaosClothAssetEditor/Source/ChaosClothAssetEditor/Private/ChaosClothAsset/ClothEditorToolkit.cpp
+++ b/Engine/Plugins/Experimental/ChaosClothAssetEditor/Source/ChaosClothAssetEditor/Private/ChaosClothAsset/ClothEditorToolkit.cpp
@@ -909,14 +909,14 @@ void FChaosClothAssetEditorToolkit::OnNodeTitleCommitted(const FText& InNewText,
 
 void FChaosClothAssetEditorToolkit::OnNodeSelectionChanged(const TSet<UObject*>& NewSelection) const
 {
-	auto GetClothCollectionIfPossible = [](const TSharedPtr<const FDataflowNode> DataflowNode, const TSharedPtr<Dataflow::FEngineContext> DataflowContext) -> TSharedPtr<FManagedArrayCollection>
+	auto GetClothCollectionIfPossible = [](const TSharedPtr<const FDataflowNode> DataflowNode, const TSharedPtr<Dataflow::FEngineContext> Context) -> TSharedPtr<FManagedArrayCollection>
 	{
 		for (const FDataflowOutput* const Output : DataflowNode->GetOutputs())
 		{
 			if (Output->GetType() == FName("FManagedArrayCollection"))
 			{
 				const FManagedArrayCollection DefaultValue;
-				TSharedPtr<FManagedArrayCollection> Collection = MakeShared<FManagedArrayCollection>(Output->GetValue<FManagedArrayCollection>(*DataflowContext, DefaultValue));
+				TSharedPtr<FManagedArrayCollection> Collection = MakeShared<FManagedArrayCollection>(Output->GetValue<FManagedArrayCollection>(*Context, DefaultValue));
 
 				// see if the output collection is a ClothCollection
 				const UE::Chaos::ClothAsset::FCollectionClothConstFacade ClothFacade(Collection);
