Chaos Cloth - Fixed the uninitialized ImportNode's Import LOD property member.

#rnx

diff --git a/Engine/Plugins/Experimental/ChaosClothAssetEditor/Source/ChaosClothAssetDataflowNodes/Public/ChaosClothAsset/ImportNode.h b/Engine/Plugins/Experimental/ChaosClothAssetEditor/Source/ChaosClothAssetDataflowNodes/Public/ChaosClothAsset/ImportNode.h
--- a/Engine/Plugins/Experimental/ChaosClothAssetEditor/Source/ChaosClothAssetDataflowNodes/Public/ChaosClothAsset/ImportNode.h
+++ b/Engine/Plugins/Experimental/ChaosClothAssetEditor/Source/ChaosClothAssetDataflowNodes/Public/ChaosClothAsset/ImportNode.h
@@ -26,7 +26,7 @@ public:
 
 	/** The LOD to import into the collection. Only one LOD can be imported at a time. */
 	UPROPERTY(EditAnywhere, Category = "ClothAsset", Meta = (DisplayName = "Import LOD"))
-	int32 ImportLod;
+	int32 ImportLod = 0;
 
 	FChaosClothAssetImportNode(const Dataflow::FNodeParameters& InParam, FGuid InGuid = FGuid::NewGuid());
 
