ClothTrainingTool: concurrent training data generation

diff --git a/Engine/Plugins/Experimental/ChaosCaching/Source/ChaosCaching/Private/Chaos/ChaosCache.cpp b/Engine/Plugins/Experimental/ChaosCaching/Source/ChaosCaching/Private/Chaos/ChaosCache.cpp
--- a/Engine/Plugins/Experimental/ChaosCaching/Source/ChaosCaching/Private/Chaos/ChaosCache.cpp
+++ b/Engine/Plugins/Experimental/ChaosCaching/Source/ChaosCaching/Private/Chaos/ChaosCache.cpp
@@ -378,7 +378,7 @@ void UChaosCache::FlushPendingFrames()
 	}
 }
 
-FCacheUserToken UChaosCache::BeginRecord(UPrimitiveComponent* InComponent, FGuid InAdapterId, const FTransform& SpaceTransform)
+FCacheUserToken UChaosCache::BeginRecord(const UPrimitiveComponent* InComponent, FGuid InAdapterId, const FTransform& SpaceTransform)
 {
 	// First make sure we're valid to record
 	int32 OtherRecordersCount = CurrentRecordCount.fetch_add(1);
@@ -714,7 +714,7 @@ FCacheEvaluationResult UChaosCache::Evaluate(const FCacheEvaluationContext& InCo
 	return Result;
 }
 
-void UChaosCache::BuildSpawnableFromComponent(UPrimitiveComponent* InComponent, const FTransform& SpaceTransform)
+void UChaosCache::BuildSpawnableFromComponent(const UPrimitiveComponent* InComponent, const FTransform& SpaceTransform)
 {
 	Spawnable.DuplicatedTemplate = StaticDuplicateObject(InComponent, this);
 	Spawnable.InitialTransform = InComponent->GetComponentToWorld();
diff --git a/Engine/Plugins/Experimental/ChaosCaching/Source/ChaosCaching/Public/Chaos/ChaosCache.h b/Engine/Plugins/Experimental/ChaosCaching/Source/ChaosCaching/Public/Chaos/ChaosCache.h
--- a/Engine/Plugins/Experimental/ChaosCaching/Source/ChaosCaching/Public/Chaos/ChaosCache.h
+++ b/Engine/Plugins/Experimental/ChaosCaching/Source/ChaosCaching/Public/Chaos/ChaosCache.h
@@ -311,7 +311,7 @@ public:
 	 * Reset and initialize a cache to make it ready to record the specified component
 	 * @param InComponent Component to prepare the cache for
 	 */
-	FCacheUserToken BeginRecord(UPrimitiveComponent* InComponent, FGuid InAdapterId, const FTransform& SpaceTransform);
+	FCacheUserToken BeginRecord(const UPrimitiveComponent* InComponent, FGuid InAdapterId, const FTransform& SpaceTransform);
 
 	/**
 	 * End the recording session for the cache. At this point the cache is deemed to now contain
@@ -359,7 +359,7 @@ public:
 	 * when a cache is dragged into the scene.
 	 * @param InComponent Component to build the spawnable template from
 	 */
-	void BuildSpawnableFromComponent(UPrimitiveComponent* InComponent, const FTransform& SpaceTransform);
+	void BuildSpawnableFromComponent(const UPrimitiveComponent* InComponent, const FTransform& SpaceTransform);
 
 	/**
 	 * Read access to the spawnable template stored in the cache
diff --git a/Engine/Plugins/Experimental/ChaosClothAssetEditor/Source/ChaosClothAssetEditorTools/Private/ChaosClothAsset/ClothTrainingTool.cpp b/Engine/Plugins/Experimental/ChaosClothAssetEditor/Source/ChaosClothAssetEditorTools/Private/ChaosClothAsset/ClothTrainingTool.cpp
--- a/Engine/Plugins/Experimental/ChaosClothAssetEditor/Source/ChaosClothAssetEditorTools/Private/ChaosClothAsset/ClothTrainingTool.cpp
+++ b/Engine/Plugins/Experimental/ChaosClothAssetEditor/Source/ChaosClothAssetEditorTools/Private/ChaosClothAsset/ClothTrainingTool.cpp
@@ -16,11 +16,14 @@
 #include "ClothingSystemRuntimeTypes.h"
 #include "ComponentReregisterContext.h"
 #include "ContextObjectStore.h"
+#include "GenericPlatform/GenericPlatformProcess.h"
 #include "InteractiveToolManager.h"
 #include "Internationalization/Regex.h"
 #include "ModelingOperators.h"
+#include "Misc/DateTime.h"
 #include "Misc/ScopedSlowTask.h"
 #include "PhysicsEngine/PhysicsAsset.h"
+#include "Tasks/Pipe.h"
 #include "ToolTargetManager.h"
 #include "UObject/SavePackage.h"
 
@@ -112,20 +115,45 @@ namespace UE::ClothTrainingTool::Private
 	{
 		TArray<int32> Result;
 		Result.Reserve(End);
-		for (int32 Index = 0; Index < End; Index++)
+		for (int32 Index = 0; Index < End; ++Index)
 		{
 			Result.Add(Index);
 		}
 		return Result;
 	}
+
+	class FTimeScope
+	{
+	public:
+		explicit FTimeScope(FString InName)
+			: Name(MoveTemp(InName))
+			, StartTime(FDateTime::UtcNow())
+		{
+		}
+		~FTimeScope()
+		{
+			const FTimespan Duration = FDateTime::UtcNow() - StartTime;
+			UE_LOG(LogClothTrainingTool, Log, TEXT("%s took %f secs"), *Name, Duration.GetTotalSeconds());
+		}
+	private:
+		FString Name;
+		FDateTime StartTime;
+	};
 };
 
-class UClothTrainingTool::FGenerateClothOp : public UE::Geometry::TGenericDataOperator<FSkinnedMeshVertices>
+struct UClothTrainingTool::FSimResource
+{
+	UChaosClothComponent* ClothComponent = nullptr;
+	TUniquePtr<FClothSimulationDataGenerationProxy> Proxy = nullptr;
+	TUniquePtr<UE::Tasks::FPipe> Pipe = nullptr;
+};
+
+class UClothTrainingTool::FLaunchSimsOp : public UE::Geometry::TGenericDataOperator<FSkinnedMeshVertices>
 {
 public: 
-	FGenerateClothOp(TObjectPtr<UChaosClothComponent> InClothComponent, FClothSimulationDataGenerationProxy* InDataGenerationProxy, TObjectPtr<UClothTrainingToolProperties> InToolProperties)
-		: ClothComponent(InClothComponent)
-		, DataGenerationProxy(InDataGenerationProxy)
+	FLaunchSimsOp(TArray<FSimResource>& InSimResources, FCriticalSection& InSimMutex, TObjectPtr<UClothTrainingToolProperties> InToolProperties)
+		: SimResources(InSimResources)
+		, SimMutex(InSimMutex)
 		, ToolProperties(InToolProperties)
 	{
 	}
@@ -133,49 +161,58 @@ public:
 	virtual void CalculateResult(FProgressCancel* Progress) override;
 
 private:
-	void Simulate(float DeltaTime, int32 NumSteps, int32 PoseIndex, FProgressCancel* Progress, bool bDebug = false);
+	using FPipe = UE::Tasks::FPipe;
+	using FProxy = UClothTrainingTool::FClothSimulationDataGenerationProxy;
+
+	enum class ESaveType
+	{
+		LastStep,
+		EveryStep,
+	};
+
+	void Simulate(FSimResource &SimResource, int32 AnimFrame, int32 CacheFrame, UChaosCache& Cache, FProgressCancel* Progress, float ProgressStep);
 	void PrepareAnimationSequence();
-	TArray<FTransform> GetBoneTransforms(int32 Frame) const;
-	bool IsClothComponentValid() const;
-	bool IsTransformsValid(const TArray<FTransform> &Transforms) const;
-	bool GetSimPositions(TArray<FVector3f>& OutPositions) const;
-	void AddToCache(UChaosCache* InCache, int32 Frame);
-	void AddToCache(UChaosCache* InCache, int32 Frame, const TArray<FVector3f>& Positions) const;
-
-	void BackupClothComponentState();
-	void RestoreClothComponentState();
-
-	TObjectPtr<UChaosClothComponent> ClothComponent = nullptr;
-	FClothSimulationDataGenerationProxy *DataGenerationProxy = nullptr;
-	TObjectPtr<UClothTrainingToolProperties> ToolProperties = nullptr;
+	void RestoreAnimationSequence();
+	TArray<FTransform> GetBoneTransforms(UChaosClothComponent& ClothComponent, int32 Frame) const;
+	bool GetSimPositions(FProxy& DataGenerationProxy, TArray<FVector3f>& OutPositions) const;
+	void AddToCache(FProxy& DataGenerationProxy, UChaosCache& Cache, int32 Frame) const;
 
-	bool bIsSimulationSuspendedBackup = false;
-	bool bTeleportBackup = false;
-	bool bResetBackup = false;
-	TArray<FTransform> ComponentSpaceTransformsBackup;
+	TArray<FSimResource>& SimResources;
+	FCriticalSection& SimMutex;
+	TObjectPtr<UClothTrainingToolProperties> ToolProperties = nullptr;
+	EAnimInterpolationType InterpolationTypeBackup = EAnimInterpolationType::Linear;
 
 	inline static const FName PositionXName = TEXT("PositionX");
 	inline static const FName PositionYName = TEXT("PositionY");
 	inline static const FName PositionZName = TEXT("PositionZ");
 };
 
-void UClothTrainingTool::FGenerateClothOp::PrepareAnimationSequence()
+void UClothTrainingTool::FLaunchSimsOp::PrepareAnimationSequence()
 {
 	TObjectPtr<UAnimSequence> AnimationSequence = ToolProperties->AnimationSequence;
 	if (AnimationSequence)
 	{
+		InterpolationTypeBackup = AnimationSequence->Interpolation;
 		AnimationSequence->Interpolation = EAnimInterpolationType::Step;
 	}
 }
 
-TArray<FTransform> UClothTrainingTool::FGenerateClothOp::GetBoneTransforms(int32 Frame) const
+void UClothTrainingTool::FLaunchSimsOp::RestoreAnimationSequence()
 {
 	TObjectPtr<UAnimSequence> AnimationSequence = ToolProperties->AnimationSequence;
+	if (AnimationSequence)
+	{
+		AnimationSequence->Interpolation = InterpolationTypeBackup;
+	}
+}
 
+TArray<FTransform> UClothTrainingTool::FLaunchSimsOp::GetBoneTransforms(UChaosClothComponent& ClothComponent, int32 Frame) const
+{
+	const UAnimSequence* AnimationSequence = ToolProperties->AnimationSequence;
 	const double Time = AnimationSequence->GetTimeAtFrame(Frame);
 	FAnimExtractContext ExtractionContext(Time);
 
-	UChaosClothAsset* const ClothAsset = ClothComponent->GetClothAsset();
+	UChaosClothAsset* const ClothAsset = ClothComponent.GetClothAsset();
 	const FReferenceSkeleton* const ReferenceSkeleton = ClothAsset ? &ClothAsset->GetRefSkeleton() : nullptr;
 	USkeleton* const Skeleton = ClothAsset ? ClothAsset->GetSkeleton() : nullptr;
 	const int32 NumBones = ReferenceSkeleton ? ReferenceSkeleton->GetNum() : 0;
@@ -215,25 +252,20 @@ TArray<FTransform> UClothTrainingTool::FGenerateClothOp::GetBoneTransforms(int32
 	return ComponentSpaceTransforms;
 }
 
-void UClothTrainingTool::FGenerateClothOp::AddToCache(UChaosCache* InCache, int32 Frame)
+void UClothTrainingTool::FLaunchSimsOp::AddToCache(FProxy& DataGenerationProxy, UChaosCache& OutCache, int32 Frame) const
 {
 	TArray<FVector3f> SimPositions;
-	if (!GetSimPositions(SimPositions))
+	if (!GetSimPositions(DataGenerationProxy, SimPositions))
 	{
 		return;
 	}
 
-	AddToCache(InCache, Frame, SimPositions);
-}
-
-void UClothTrainingTool::FGenerateClothOp::AddToCache(UChaosCache* InCache, int32 Frame, const TArray<FVector3f>& Positions) const
-{
 	constexpr float CacheFPS = 30;
 	const float Time = Frame / CacheFPS;
 	FPendingFrameWrite NewFrame;
 	NewFrame.Time = Time;
 
-	const int32 NumParticles = Positions.Num();
+	const int32 NumParticles = SimPositions.Num();
 	TArray<int32>& PendingID = NewFrame.PendingChannelsIndices;
 	TArray<float> PendingPX, PendingPY, PendingPZ;
 	TArray<float> PendingVX, PendingVY, PendingVZ;
@@ -244,7 +276,7 @@ void UClothTrainingTool::FGenerateClothOp::AddToCache(UChaosCache* InCache, int3
 
 	for (int32 ParticleIndex = 0; ParticleIndex < NumParticles; ++ParticleIndex)
 	{
-		const FVector3f& Position = Positions[ParticleIndex];
+		const FVector3f& Position = SimPositions[ParticleIndex];
 		PendingID[ParticleIndex] = ParticleIndex;
 		PendingPX[ParticleIndex] = Position.X;
 		PendingPY[ParticleIndex] = Position.Y;
@@ -255,21 +287,15 @@ void UClothTrainingTool::FGenerateClothOp::AddToCache(UChaosCache* InCache, int3
 	NewFrame.PendingChannelsData.Add(PositionYName, PendingPY);
 	NewFrame.PendingChannelsData.Add(PositionZName, PendingPZ);
 
-	InCache->AddFrame_Concurrent(MoveTemp(NewFrame));
+	OutCache.AddFrame_Concurrent(MoveTemp(NewFrame));
 }
 
-void UClothTrainingTool::FGenerateClothOp::CalculateResult(FProgressCancel* Progress)
+void UClothTrainingTool::FLaunchSimsOp::CalculateResult(FProgressCancel* Progress)
 {
-	if (ClothComponent == nullptr || DataGenerationProxy == nullptr || !IsClothComponentValid())
-	{
-		return;
-	}
-	
 	using UE::ClothTrainingTool::Private::ParseFrames;
 	using UE::ClothTrainingTool::Private::Range;
+	using UE::ClothTrainingTool::Private::GetCache;
 
-	bool bCancelled = false;
-	constexpr float DeltaTime = 1e-3;
 	const TArray<int32> FramesToSimulate = ToolProperties->FramesToSimulate.Len() > 0
 		? ParseFrames(ToolProperties->FramesToSimulate) 
 		: Range(ToolProperties->AnimationSequence->GetNumberOfSampledKeys());
@@ -279,10 +305,15 @@ void UClothTrainingTool::FGenerateClothOp::CalculateResult(FProgressCancel* Prog
 	{
 		return;
 	}
+	const float ProgressStep = 1.f / NumFrames;
 
+	UChaosCache& Cache = ToolProperties->bDebug ? *GetCache(ToolProperties->DebugCacheCollection) : *GetCache(ToolProperties->CacheCollection);
 	PrepareAnimationSequence();
-	BackupClothComponentState();
-	ClothComponent->ResumeSimulation();
+
+	const int32 NumThreads = ToolProperties->bDebug ? 1 : ToolProperties->NumThreads;
+	FScopeLock Lock(&SimMutex);
+
+	bool bCancelled = false;
 	for (int32 Frame = 0; Frame < NumFrames; Frame++)
 	{
 		if (Progress)
@@ -293,40 +324,39 @@ void UClothTrainingTool::FGenerateClothOp::CalculateResult(FProgressCancel* Prog
 				break;
 			}
 
-			const int32 FrameToSimulate = ToolProperties->bDebug ? ToolProperties->DebugFrame : FramesToSimulate[Frame];
-			const TArray<FTransform> BoneTransforms = GetBoneTransforms(FrameToSimulate);
-			if (!IsTransformsValid(BoneTransforms))
-			{
-				break;
-			}
-
-			ClothComponent->Pose(BoneTransforms);
-			Simulate(ToolProperties->TimeStep, ToolProperties->NumSteps, FrameToSimulate, Progress, ToolProperties->bDebug);
+			const int32 ThreadIdx = Frame % NumThreads;
+			const int32 AnimFrame = FramesToSimulate[Frame];
 
-			Progress->AdvanceCurrentScopeProgressBy(1.f / NumFrames);
+			FSimResource& SimResource = SimResources[ThreadIdx];
+			SimResource.Pipe->Launch(*FString::Printf(TEXT("SimFrame:%d"), AnimFrame), [this, &SimResource, AnimFrame, Frame, &Cache, Progress, ProgressStep]()
+			{ 
+				FMemMark Mark(FMemStack::Get());
+				Simulate(SimResource, AnimFrame, Frame, Cache, Progress, ProgressStep);
+			});
 		}
 	}
 
-	RestoreClothComponentState();
-
-	if (!bCancelled)
+	for (FSimResource& SimResource : SimResources)
 	{
-		UE_LOG(LogClothTrainingTool, Display, TEXT("Data generation complete."));
+		SimResource.Pipe->WaitUntilEmpty();
 	}
+
+	RestoreAnimationSequence();
+	Cache.FlushPendingFrames();
 }
 
-void UClothTrainingTool::FGenerateClothOp::Simulate(float DeltaTime, int32 NumSteps, int32 Frame, FProgressCancel* Progress, bool bDebug)
+void UClothTrainingTool::FLaunchSimsOp::Simulate(FSimResource &SimResource, int32 AnimFrame, int32 CacheFrame, UChaosCache& Cache, FProgressCancel* Progress, float ProgressStep)
 {
-	using UE::ClothTrainingTool::Private::GetCache;
-	UChaosCache* Cache = ToolProperties->bDebug ? GetCache(ToolProperties->DebugCacheCollection) : GetCache(ToolProperties->CacheCollection);
-	if (Cache == nullptr)
-	{
-		return;
-	}
+	UChaosClothComponent& TaskComponent = *SimResource.ClothComponent;
+	FProxy& DataGenerationProxy = *SimResource.Proxy.Get();
+
+	const float TimeStep = ToolProperties->TimeStep;
+	const int32 NumSteps = ToolProperties->NumSteps;
+	const ESaveType SaveType = ToolProperties->bDebug ? ESaveType::EveryStep : ESaveType::LastStep;
 
-	ClothComponent->ForceNextUpdateTeleportAndReset();
-	DataGenerationProxy->FillSimulationContext(DeltaTime);	
-	DataGenerationProxy->InitializeConfigs();
+	TaskComponent.ForceNextUpdateTeleportAndReset();
+	DataGenerationProxy.FillSimulationContext(TimeStep);
+	DataGenerationProxy.InitializeConfigs();
 	bool bCancelled = false;
 	for (int32 Step = 0; Step < NumSteps; ++Step)
 	{
@@ -338,33 +368,35 @@ void UClothTrainingTool::FGenerateClothOp::Simulate(float DeltaTime, int32 NumSt
 				break;
 			}
 
-			DataGenerationProxy->Tick();
+			DataGenerationProxy.Tick();
 
 			// Clear any reset flags at the end of the first step
 			if (Step == 0 && NumSteps > 1)
 			{
-				ClothComponent->ResetTeleportMode();
-				DataGenerationProxy->FillSimulationContext(DeltaTime);	
+				TaskComponent.ResetTeleportMode();
+				DataGenerationProxy.FillSimulationContext(TimeStep);	
 			}
 
-			if (bDebug)
+			if (SaveType == ESaveType::EveryStep)
 			{
-				DataGenerationProxy->WriteSimulationData();
-				AddToCache(Cache, Step);
+				DataGenerationProxy.WriteSimulationData();
+				AddToCache(DataGenerationProxy, Cache, Step);
 			}
 		}
 	}
 
-	if (!bDebug && !bCancelled)
+	if (SaveType == ESaveType::LastStep && !bCancelled)
 	{
-		DataGenerationProxy->WriteSimulationData();
-		AddToCache(Cache, Frame);
+		DataGenerationProxy.WriteSimulationData();
+		AddToCache(DataGenerationProxy, Cache, CacheFrame);
 	}
+
+	Progress->AdvanceCurrentScopeProgressBy(ProgressStep);
 }
 
-bool UClothTrainingTool::FGenerateClothOp::GetSimPositions(TArray<FVector3f> &OutPositions) const
+bool UClothTrainingTool::FLaunchSimsOp::GetSimPositions(FProxy& DataGenerationProxy, TArray<FVector3f> &OutPositions) const
 {
-	const TMap<int32, FClothSimulData> &SimulDataMap = DataGenerationProxy->GetCurrentSimulationData_AnyThread();
+	const TMap<int32, FClothSimulData> &SimulDataMap = DataGenerationProxy.GetCurrentSimulationData_AnyThread();
 	const FClothSimulData* const SimulData = SimulDataMap.Find(0);
 	if (SimulDataMap.Num() > 1)
 	{
@@ -386,34 +418,12 @@ bool UClothTrainingTool::FGenerateClothOp::GetSimPositions(TArray<FVector3f> &Ou
 	return true;
 }
 
-void UClothTrainingTool::FGenerateClothOp::BackupClothComponentState()
-{
-	bIsSimulationSuspendedBackup = ClothComponent->IsSimulationSuspended();
-	bTeleportBackup = ClothComponent->NeedsTeleport();
-	bResetBackup = ClothComponent->NeedsReset();
-	ComponentSpaceTransformsBackup = ClothComponent->GetComponentSpaceTransforms();
-}
-
-void UClothTrainingTool::FGenerateClothOp::RestoreClothComponentState()
+bool UClothTrainingTool::IsClothComponentValid() const
 {
-	bIsSimulationSuspendedBackup ? ClothComponent->SuspendSimulation() : ClothComponent->ResumeSimulation();
-	if (bResetBackup)
-	{
-		ClothComponent->ForceNextUpdateTeleportAndReset();
-	}
-	else if (bTeleportBackup)
-	{
-		ClothComponent->ForceNextUpdateTeleport();
-	}
-	else
+	if (ClothComponent == nullptr)
 	{
-		ClothComponent->ResetTeleportMode();
+		return false;
 	}
-	ClothComponent->Pose(ComponentSpaceTransformsBackup);
-}
-
-bool UClothTrainingTool::FGenerateClothOp::IsClothComponentValid() const
-{
 	if (USkinnedMeshComponent* const LeaderComponent = ClothComponent->LeaderPoseComponent.Get())
 	{
 		UE_LOG(LogClothTrainingTool, Error, TEXT("Leader pose component is not supported yet."));
@@ -425,19 +435,6 @@ bool UClothTrainingTool::FGenerateClothOp::IsClothComponentValid() const
 	}
 }
 
-bool UClothTrainingTool::FGenerateClothOp::IsTransformsValid(const TArray<FTransform>& Transforms) const
-{
-	if (Transforms.Num() != ComponentSpaceTransformsBackup.Num())
-	{
-		UE_LOG(LogClothTrainingTool, Error, TEXT("The number of bones in the AnimSequence does not match the number of bones in the ClothAsset. Try an AnimSequence with the same skeleton as the ClothAsset."));
-		return false;
-	}
-	else
-	{
-		return true;
-	}
-}
-
 // ------------------- Properties -------------------
 
 void UClothTrainingToolActionProperties::PostAction(EClothTrainingToolActions Action)
@@ -489,9 +486,9 @@ void UClothTrainingTool::Setup()
 {
 	UInteractiveTool::Setup();
 
-	if (UClothComponentToolTarget* ClothComponentTarget = Cast<UClothComponentToolTarget>(Target))
+	if (const UClothComponentToolTarget* const ClothComponentTarget = Cast<const UClothComponentToolTarget>(Target))
 	{
-		UChaosClothComponent* TargetClothComponent = ClothComponentTarget->GetClothComponent();
+		const UChaosClothComponent* const TargetClothComponent = ClothComponentTarget->GetClothComponent();
 		ClothComponent = TargetClothComponent;
 	}
 
@@ -506,101 +503,78 @@ void UClothTrainingTool::Setup()
 
 void UClothTrainingTool::RunTraining()
 {
-	if (ClothComponent == nullptr)
+	if (!IsClothComponentValid() || ToolProperties->AnimationSequence == nullptr)
 	{
 		return;
 	}
-
-
-	const FText DefaultMessage(LOCTEXT("ClothTrainingMessage", "Generate training data..."));
-	
-	using FTaskType = UE::Geometry::TModelingOpTask<FGenerateClothOp>;
-	using FExecuterType = UE::Geometry::FAsyncTaskExecuterWithProgressCancel<FTaskType>;
-	using UE::ClothTrainingTool::Private::GetCache;
-
-	if (!DataGenerationProxy.IsValid())
-	{
-		DataGenerationProxy = MakeUnique<FClothSimulationDataGenerationProxy>(*ClothComponent);
-	}
-	UChaosCacheCollection* const CacheCollection = ToolProperties->bDebug ? ToolProperties->DebugCacheCollection : ToolProperties->CacheCollection;
+	UChaosCacheCollection* const CacheCollection = GetCacheCollection();
 	if (CacheCollection == nullptr)
 	{
-		if (ToolProperties->bDebug)
-		{
-			UE_LOG(LogClothTrainingTool, Error, TEXT("Debug cache is None. Please select a valid cache for output."));
-		}
-		else
-		{
-			UE_LOG(LogClothTrainingTool, Error, TEXT("Generated Cache is None. Please select a valid cache for output."));
-		}
 		return;
 	}
+	using UE::ClothTrainingTool::Private::GetCache;
 	UChaosCache* const Cache = GetCache(CacheCollection);
+	if (Cache == nullptr)
+	{
+		return;
+	}
+	
+	using FTaskType = UE::Geometry::TModelingOpTask<FLaunchSimsOp>;
+	using FExecuterType = UE::Geometry::FAsyncTaskExecuterWithProgressCancel<FTaskType>;
 
-	FCacheUserToken CacheUserToken = Cache->BeginRecord(ClothComponent, FGuid(), FTransform::Identity);
+	const FText DefaultMessage(LOCTEXT("ClothTrainingMessage", "Generate training data..."));
 
-	TUniquePtr<FGenerateClothOp> NewOp = MakeUnique<FGenerateClothOp>(ClothComponent, DataGenerationProxy.Get(), ToolProperties);
+	if (!AllocateSimResources_GameThread(ToolProperties->NumThreads))
+	{
+		return;
+	}
+	FCacheUserToken CacheUserToken = Cache->BeginRecord(ClothComponent, FGuid(), FTransform::Identity);
 
+	TUniquePtr<FLaunchSimsOp> NewOp = MakeUnique<FLaunchSimsOp>(SimResources, *SimMutex, ToolProperties);
 	TUniquePtr<FExecuterType> BackgroundTaskExecuter = MakeUnique<FExecuterType>(MoveTemp(NewOp));
 	BackgroundTaskExecuter->StartBackgroundTask();
 
 	FScopedSlowTask SlowTask(1, DefaultMessage);
 	SlowTask.MakeDialog(true);
 
-	bool bSuccess = false;
-	while (true)
 	{
-		if (SlowTask.ShouldCancel())
-		{
-			// Release ownership to the TDeleterTask that is spawned by CancelAndDelete()
-			BackgroundTaskExecuter.Release()->CancelAndDelete();
-			break;
-		}
-		if (BackgroundTaskExecuter->IsDone())
-		{
-			bSuccess = !BackgroundTaskExecuter->GetTask().IsAborted();
-			break;
-		}
-		FPlatformProcess::Sleep(.2); // SlowTask::ShouldCancel will throttle any updates faster than .2 seconds
-		float ProgressFrac;
-		FText ProgressMessage;
-		bool bMadeProgress = BackgroundTaskExecuter->PollProgress(ProgressFrac, ProgressMessage);
-		if (bMadeProgress)
+		UE::ClothTrainingTool::Private::FTimeScope TimeScope(TEXT("Cloth generation"));
+		while (true)
 		{
-			// SlowTask expects progress to be reported before it happens; we work around this by directly updating the progress amount
-			SlowTask.CompletedWork = ProgressFrac;
-			SlowTask.EnterProgressFrame(0, ProgressMessage);
-		}
-		else
-		{
-			SlowTask.TickProgress(); // Still tick the UI when we don't get new progress frames
+			if (SlowTask.ShouldCancel())
+			{
+				// Release ownership to the TDeleterTask that is spawned by CancelAndDelete()
+				BackgroundTaskExecuter.Release()->CancelAndDelete();
+				break;
+			}
+			if (BackgroundTaskExecuter->IsDone())
+			{
+				break;
+			}
+			FPlatformProcess::Sleep(.2); // SlowTask::ShouldCancel will throttle any updates faster than .2 seconds
+			float ProgressFrac;
+			FText ProgressMessage;
+			const bool bMadeProgress = BackgroundTaskExecuter->PollProgress(ProgressFrac, ProgressMessage);
+			if (bMadeProgress)
+			{
+				// SlowTask expects progress to be reported before it happens; we work around this by directly updating the progress amount
+				SlowTask.CompletedWork = ProgressFrac;
+				SlowTask.EnterProgressFrame(0, ProgressMessage);
+			}
+			else
+			{
+				SlowTask.TickProgress(); // Still tick the UI when we don't get new progress frames
+			}
 		}
+		FreeSimResources_GameThread();
 	}
-
-	ClothComponent->SuspendSimulation();
-	Cache->bCompressChannels = true;
-	Cache->EndRecord(CacheUserToken);
-
-	if (bSuccess)
 	{
-		UPackage* const Package = CacheCollection->GetPackage();
+		UE::ClothTrainingTool::Private::FTimeScope TimeScope(TEXT("Saving"));
+		Cache->bCompressChannels = true;
+		Cache->EndRecord(CacheUserToken);
 
-		const FString SavePath = Package->GetFName().ToString();
-		UE_LOG(LogTemp, Display, TEXT("Save to %s"), *SavePath);
-		if (Package == nullptr)
-		{
-			UE_LOG(LogTemp, Error, TEXT("Failed to get package for %s"), *SavePath);
-			return;
-		}
-		FSavePackageArgs SaveArgs;
-		SaveArgs.SaveFlags = SAVE_NoError;
-		const bool bSaveSucced = UPackage::SavePackage(Package, ToolProperties->CacheCollection, *SavePath, SaveArgs);
-		if (!bSaveSucced)
-		{
-			UE_LOG(LogTemp, Display, TEXT("Failed to save cache collection"));
-		}
+		SaveCacheCollection(CacheCollection);
 	}
-
 }
 
 
@@ -632,4 +606,91 @@ void UClothTrainingTool::Shutdown(EToolShutdownType ShutdownType)
 	ToolProperties->SaveProperties(this);
 }
 
+bool UClothTrainingTool::AllocateSimResources_GameThread(int32 Num)
+{
+	SimResources.Reserve(Num);
+	for(int32 Index = 0; Index < Num; ++Index)
+	{
+		UChaosClothComponent* const CopyComponent = DuplicateObject<UChaosClothComponent>(ClothComponent, ClothComponent->GetOuter());
+		CopyComponent->RegisterComponent();
+		CopyComponent->SetWorldTransform(ClothComponent->GetComponentTransform());
+		FSimResource SimResource;
+		SimResource.ClothComponent = CopyComponent;
+		SimResource.Proxy = MakeUnique<FClothSimulationDataGenerationProxy>(*CopyComponent);
+		SimResource.Pipe = MakeUnique<UE::Tasks::FPipe>(*FString::Printf(TEXT("SimPipe:%d"), Index));
+
+		if (CopyComponent == nullptr || SimResource.Proxy == nullptr || SimResource.Pipe == nullptr)
+		{
+			UE_LOG(LogClothTrainingTool, Error, TEXT("Failed to allocate simulation resources"));
+			return false;
+		}
+		SimResources.Add(MoveTemp(SimResource));
+	}
+	SimMutex = MakeUnique<FCriticalSection>();
+	return true;
+}
+
+void UClothTrainingTool::FreeSimResources_GameThread()
+{
+	{
+		FScopeLock Lock(SimMutex.Get());
+		for (FSimResource& SimResource : SimResources)
+		{
+			SimResource.Pipe.Reset();
+			SimResource.Proxy.Reset();
+			SimResource.ClothComponent->UnregisterComponent();
+			SimResource.ClothComponent->DestroyComponent();
+		}
+		SimResources.Empty();
+	}
+	SimMutex.Reset();
+}
+
+UChaosCacheCollection* UClothTrainingTool::GetCacheCollection() const
+{
+	UChaosCacheCollection* CacheCollection = nullptr;
+	if (ToolProperties->bDebug)
+	{
+		CacheCollection = ToolProperties->DebugCacheCollection;
+		if (CacheCollection == nullptr)
+		{
+			UE_LOG(LogClothTrainingTool, Error, TEXT("Debug cache is None. Please select a valid cache for output."));
+		}
+	}
+	else
+	{
+		CacheCollection = ToolProperties->CacheCollection;
+		if (CacheCollection == nullptr)
+		{
+			UE_LOG(LogClothTrainingTool, Error, TEXT("Generated Cache is None. Please select a valid cache for output."));
+		}
+	}
+	return CacheCollection;
+}
+
+bool UClothTrainingTool::SaveCacheCollection(UChaosCacheCollection* CacheCollection) const
+{
+	if (CacheCollection == nullptr)
+	{
+		return false;
+	}
+	UPackage* const Package = CacheCollection->GetPackage();
+	if (Package == nullptr)
+	{
+		UE_LOG(LogClothTrainingTool, Error, TEXT("Failed to get package for %s"), *(CacheCollection->GetFName().ToString()));
+		return false;
+	}
+	const FString SavePath = Package->GetFName().ToString();
+	UE_LOG(LogClothTrainingTool, Display, TEXT("Save to %s"), *SavePath);
+
+	FSavePackageArgs SaveArgs;
+	SaveArgs.SaveFlags = SAVE_NoError;
+	const bool bSaveSucced = UPackage::SavePackage(Package, CacheCollection, *SavePath, SaveArgs);
+	if (!bSaveSucced)
+	{
+		UE_LOG(LogClothTrainingTool, Error, TEXT("Failed to save cache collection"));
+	}
+	return bSaveSucced;
+}
+
 #undef LOCTEXT_NAMESPACE
diff --git a/Engine/Plugins/Experimental/ChaosClothAssetEditor/Source/ChaosClothAssetEditorTools/Public/ChaosClothAsset/ClothTrainingTool.h b/Engine/Plugins/Experimental/ChaosClothAssetEditor/Source/ChaosClothAssetEditorTools/Public/ChaosClothAsset/ClothTrainingTool.h
--- a/Engine/Plugins/Experimental/ChaosClothAssetEditor/Source/ChaosClothAssetEditorTools/Public/ChaosClothAsset/ClothTrainingTool.h
+++ b/Engine/Plugins/Experimental/ChaosClothAssetEditor/Source/ChaosClothAssetEditorTools/Public/ChaosClothAsset/ClothTrainingTool.h
@@ -54,6 +54,9 @@ public:
 
 	UPROPERTY(EditAnywhere, Category = "Simulation Settings", meta = (Min = 0))
 	int32 NumSteps = 200;
+
+	UPROPERTY(EditAnywhere, Category = "Simulation Settings", meta = (Min = 1, EditCondition = "!bDebug"))
+	int32 NumThreads = 1;
 };
 
 UENUM()
@@ -116,7 +119,8 @@ public:
 
 private:
 	class FClothSimulationDataGenerationProxy;
-	class FGenerateClothOp;
+	struct FSimResource;
+	class FLaunchSimsOp;
 
 	friend class UClothTrainingToolActionProperties;
 
@@ -129,13 +133,21 @@ private:
 	EClothTrainingToolActions PendingAction = EClothTrainingToolActions::NoAction;
 
 	UPROPERTY()
-	TObjectPtr<UChaosClothComponent> ClothComponent;
+	TObjectPtr<const UChaosClothComponent> ClothComponent;
 
-	TUniquePtr<FClothSimulationDataGenerationProxy> DataGenerationProxy;
+	TUniquePtr<FCriticalSection> SimMutex;
+	TArray<FSimResource> SimResources;
 
 	void RequestAction(EClothTrainingToolActions ActionType);
-
 	void RunTraining();
+	bool IsClothComponentValid() const;
+	UChaosCacheCollection* GetCacheCollection() const;
+	bool SaveCacheCollection(UChaosCacheCollection* CacheCollection) const;
+	void PrepareAnimationSequence();
+	void RestoreAnimationSequence();
+	bool AllocateSimResources_GameThread(int32 Num);
+	void FreeSimResources_GameThread();
+
 };
 
 
