ClothTrainingTool: fix compiler error

diff --git a/Engine/Plugins/Experimental/ChaosClothAssetEditor/Source/ChaosClothAssetEditorTools/Private/ChaosClothAsset/ClothTrainingTool.cpp b/Engine/Plugins/Experimental/ChaosClothAssetEditor/Source/ChaosClothAssetEditorTools/Private/ChaosClothAsset/ClothTrainingTool.cpp
--- a/Engine/Plugins/Experimental/ChaosClothAssetEditor/Source/ChaosClothAssetEditorTools/Private/ChaosClothAsset/ClothTrainingTool.cpp
+++ b/Engine/Plugins/Experimental/ChaosClothAssetEditor/Source/ChaosClothAssetEditorTools/Private/ChaosClothAsset/ClothTrainingTool.cpp
@@ -173,7 +173,7 @@ private:
 	void Simulate(FSimResource &SimResource, int32 AnimFrame, int32 CacheFrame, UChaosCache& Cache, FProgressCancel* Progress, float ProgressStep);
 	void PrepareAnimationSequence();
 	void RestoreAnimationSequence();
-	TArray<FTransform> GetBoneTransforms(UChaosClothComponent& ClothComponent, int32 Frame) const;
+	TArray<FTransform> GetBoneTransforms(UChaosClothComponent& InClothComponent, int32 Frame) const;
 	bool GetSimPositions(FProxy& DataGenerationProxy, TArray<FVector3f>& OutPositions) const;
 	void AddToCache(FProxy& DataGenerationProxy, UChaosCache& Cache, int32 Frame) const;
 
@@ -206,13 +206,13 @@ void UClothTrainingTool::FLaunchSimsOp::RestoreAnimationSequence()
 	}
 }
 
-TArray<FTransform> UClothTrainingTool::FLaunchSimsOp::GetBoneTransforms(UChaosClothComponent& ClothComponent, int32 Frame) const
+TArray<FTransform> UClothTrainingTool::FLaunchSimsOp::GetBoneTransforms(UChaosClothComponent& InClothComponent, int32 Frame) const
 {
 	const UAnimSequence* AnimationSequence = ToolProperties->AnimationSequence;
 	const double Time = AnimationSequence->GetTimeAtFrame(Frame);
 	FAnimExtractContext ExtractionContext(Time);
 
-	UChaosClothAsset* const ClothAsset = ClothComponent.GetClothAsset();
+	UChaosClothAsset* const ClothAsset = InClothComponent.GetClothAsset();
 	const FReferenceSkeleton* const ReferenceSkeleton = ClothAsset ? &ClothAsset->GetRefSkeleton() : nullptr;
 	USkeleton* const Skeleton = ClothAsset ? ClothAsset->GetSkeleton() : nullptr;
 	const int32 NumBones = ReferenceSkeleton ? ReferenceSkeleton->GetNum() : 0;
