ClothTrainingModel: update cache reference

diff --git a/Engine/Plugins/Experimental/ChaosClothAssetEditor/Source/ChaosClothAssetEditorTools/Private/ChaosClothAsset/ClothTrainingTool.cpp b/Engine/Plugins/Experimental/ChaosClothAssetEditor/Source/ChaosClothAssetEditorTools/Private/ChaosClothAsset/ClothTrainingTool.cpp
--- a/Engine/Plugins/Experimental/ChaosClothAssetEditor/Source/ChaosClothAssetEditorTools/Private/ChaosClothAsset/ClothTrainingTool.cpp
+++ b/Engine/Plugins/Experimental/ChaosClothAssetEditor/Source/ChaosClothAssetEditorTools/Private/ChaosClothAsset/ClothTrainingTool.cpp
@@ -53,10 +53,10 @@ UClothTrainingTool::FClothSimulationDataGenerationProxy::~FClothSimulationDataGe
 
 namespace UE::ClothTrainingTool::Private
 {
-	UChaosCache* GetCache(UChaosCacheCollection* CacheCollection)
+	UChaosCache& GetCache(UChaosCacheCollection& CacheCollection)
 	{
 		static const FName CacheName = FName("SimulatedCache");
-		return CacheCollection ? CacheCollection->FindOrAddCache(CacheName) : nullptr;
+		return *CacheCollection.FindOrAddCache(CacheName);
 	}
 
 	TArray<int32> ParseFrames(const FString& FramesString)
@@ -307,7 +307,7 @@ void UClothTrainingTool::FLaunchSimsOp::CalculateResult(FProgressCancel* Progres
 	}
 	const float ProgressStep = 1.f / NumFrames;
 
-	UChaosCache& Cache = ToolProperties->bDebug ? *GetCache(ToolProperties->DebugCacheCollection) : *GetCache(ToolProperties->CacheCollection);
+	UChaosCache& Cache = ToolProperties->bDebug ? GetCache(*ToolProperties->DebugCacheCollection) : GetCache(*ToolProperties->CacheCollection);
 	PrepareAnimationSequence();
 
 	const int32 NumThreads = ToolProperties->bDebug ? 1 : ToolProperties->NumThreads;
@@ -513,11 +513,7 @@ void UClothTrainingTool::RunTraining()
 		return;
 	}
 	using UE::ClothTrainingTool::Private::GetCache;
-	UChaosCache* const Cache = GetCache(CacheCollection);
-	if (Cache == nullptr)
-	{
-		return;
-	}
+	UChaosCache& Cache = GetCache(*CacheCollection);
 	
 	using FTaskType = UE::Geometry::TModelingOpTask<FLaunchSimsOp>;
 	using FExecuterType = UE::Geometry::FAsyncTaskExecuterWithProgressCancel<FTaskType>;
@@ -528,7 +524,7 @@ void UClothTrainingTool::RunTraining()
 	{
 		return;
 	}
-	FCacheUserToken CacheUserToken = Cache->BeginRecord(ClothComponent, FGuid(), FTransform::Identity);
+	FCacheUserToken CacheUserToken = Cache.BeginRecord(ClothComponent, FGuid(), FTransform::Identity);
 
 	TUniquePtr<FLaunchSimsOp> NewOp = MakeUnique<FLaunchSimsOp>(SimResources, *SimMutex, ToolProperties);
 	TUniquePtr<FExecuterType> BackgroundTaskExecuter = MakeUnique<FExecuterType>(MoveTemp(NewOp));
@@ -570,8 +566,8 @@ void UClothTrainingTool::RunTraining()
 	}
 	{
 		UE::ClothTrainingTool::Private::FTimeScope TimeScope(TEXT("Saving"));
-		Cache->bCompressChannels = true;
-		Cache->EndRecord(CacheUserToken);
+		Cache.bCompressChannels = true;
+		Cache.EndRecord(CacheUserToken);
 
 		SaveCacheCollection(CacheCollection);
 	}
