Write skinning weight and bone data when converting from cloth asset pattern to a dynamic mesh.

#rnx

diff --git a/Engine/Plugins/Experimental/ChaosClothAssetEditor/Source/ChaosClothAssetTools/ChaosClothAssetTools.Build.cs b/Engine/Plugins/Experimental/ChaosClothAssetEditor/Source/ChaosClothAssetTools/ChaosClothAssetTools.Build.cs
--- a/Engine/Plugins/Experimental/ChaosClothAssetEditor/Source/ChaosClothAssetTools/ChaosClothAssetTools.Build.cs
+++ b/Engine/Plugins/Experimental/ChaosClothAssetEditor/Source/ChaosClothAssetTools/ChaosClothAssetTools.Build.cs
@@ -30,6 +30,7 @@ public class ChaosClothAssetTools : ModuleRules
 				"ClothPainter",                 // For clothing asset exporter
 				"ClothingSystemRuntimeCommon",
 				"GeometryCore",					// For DynamicMesh conversion
+				"SkeletalMeshDescription",		// For FSkeletalMeshAttributes::DefaultSkinWeightProfileName
 			}
 		);
 	}
diff --git a/Engine/Plugins/Experimental/ChaosClothAssetEditor/Source/ChaosClothAssetTools/Private/ChaosClothAsset/ClothPatternToDynamicMesh.cpp b/Engine/Plugins/Experimental/ChaosClothAssetEditor/Source/ChaosClothAssetTools/Private/ChaosClothAsset/ClothPatternToDynamicMesh.cpp
--- a/Engine/Plugins/Experimental/ChaosClothAssetEditor/Source/ChaosClothAssetTools/Private/ChaosClothAsset/ClothPatternToDynamicMesh.cpp
+++ b/Engine/Plugins/Experimental/ChaosClothAssetEditor/Source/ChaosClothAssetTools/Private/ChaosClothAsset/ClothPatternToDynamicMesh.cpp
@@ -8,6 +8,8 @@
 #include "ChaosClothAsset/CollectionClothFacade.h"
 #include "ToDynamicMesh.h"
 #include "DynamicMesh/MeshNormals.h"
+#include "SkeletalMeshAttributes.h"
+#include "Animation/Skeleton.h"
 
 //
 // Wrapper for accessing a Cloth Pattern. Implements the interface expected by TToDynamicMesh<>.
@@ -46,6 +48,14 @@ public:
 		// Weight map layers precomputation
 		//
 		WeightMapNames = ClothFacade.GetWeightMapNames();
+
+		// Set the reference skeleton if available
+		const FString& SkeletonPathName = ClothFacade.GetLod(LodIndex).GetSkeletonAssetPathName();
+		const USkeleton* const Skeleton = !SkeletonPathName.IsEmpty() ?
+			LoadObject<USkeleton>(nullptr, *SkeletonPathName, nullptr, LOAD_None, nullptr) :
+			nullptr;
+		RefSkeleton = Skeleton ? &Skeleton->GetReferenceSkeleton() : nullptr;
+		ensureMsgf(RefSkeleton, TEXT("Reference skeleton is not set"));
 	}
 
 	int32 NumTris() const
@@ -182,6 +192,90 @@ public:
 		return 0;
 	}
 
+	int32 NumSkinWeightAttributes() const 
+	{ 
+		return 1;
+	}
+
+	UE::AnimationCore::FBoneWeights GetVertexSkinWeight(int32 SkinWeightAttributeIndex, VertIDType VertexID) const 
+	{
+		using namespace UE::AnimationCore;
+		
+		checkfSlow(SkinWeightAttributeIndex == 0, TEXT("Cloth assets should only have one skin weight profile")); 
+
+		const TConstArrayView<int32> SimNumBoneInfluences = Pattern.GetSimNumBoneInfluences();
+		const TConstArrayView<TArray<int32>> SimBoneIndices = Pattern.GetSimBoneIndices();
+		const TConstArrayView<TArray<float>> BoneWeights = Pattern.GetSimBoneWeights();
+
+		if (ensure(VertexID >= 0 && VertexID < SimNumBoneInfluences.Num()))
+		{
+			const int32 NumInfluences = SimNumBoneInfluences[VertexID];
+			const TArray<int32> Indices = SimBoneIndices[VertexID];
+			const TArray<float> Weights = BoneWeights[VertexID];
+			
+			TArray<FBoneWeight> BoneWeightArray;
+			BoneWeightArray.SetNumUninitialized(NumInfluences);
+
+			for (int32 Idx = 0; Idx < NumInfluences; ++Idx)
+			{
+				BoneWeightArray[Idx] = FBoneWeight(static_cast<FBoneIndexType>(Indices[Idx]), Weights[Idx]);
+			}
+
+			return FBoneWeights::Create(BoneWeightArray, FBoneWeightsSettings());
+		}
+		else 
+		{
+			return FBoneWeights();
+		}
+	}
+
+	FName GetSkinWeightAttributeName(int32 SkinWeightAttributeIndex) const 
+	{ 
+		checkfSlow(SkinWeightAttributeIndex == 0, TEXT("Cloth assets should only have one skin weight profile")); 
+		
+		return FSkeletalMeshAttributes::DefaultSkinWeightProfileName;
+	}
+
+	int32 GetNumBones() const 
+	{ 
+		return RefSkeleton ? RefSkeleton->GetRawBoneNum() : 0;
+	}
+
+    FName GetBoneName(int32 BoneIdx) const
+	{
+		if (ensure(BoneIdx >= 0 && BoneIdx < GetNumBones()) && RefSkeleton)
+		{
+			return RefSkeleton->GetRawRefBoneInfo()[BoneIdx].Name;
+		}
+		
+		return NAME_None;
+	}
+
+	int32 GetBoneParentIndex(int32 BoneIdx) const
+	{
+		if (ensure(BoneIdx >= 0 && BoneIdx < GetNumBones()) && RefSkeleton)
+		{
+			return RefSkeleton->GetRawRefBoneInfo()[BoneIdx].ParentIndex;
+		}
+
+		return INDEX_NONE;
+	}
+
+	FTransform GetBonePose(int32 BoneIdx) const
+	{
+		if (ensure(BoneIdx >= 0 && BoneIdx < GetNumBones()) && RefSkeleton)
+		{
+			return RefSkeleton->GetRawRefBonePose()[BoneIdx];
+		}
+
+		return FTransform::Identity;
+	}
+
+	FVector4f GetBoneColor(int32 BoneIdx) const
+	{
+		return FVector4f::One();
+	}
+
 	const TArray<int32>& GetNormalIDs() const { return VertIDs; }
 	FVector3f GetNormal(NormalIDType ID) const { return Pattern.GetSimRestNormal()[ID]; }
 	bool GetNormalTri(const TriIDType& TriID, NormalIDType& ID0, NormalIDType& ID1, NormalIDType& ID2) const { return GetTri(TriID, ID0, ID1, ID2); }
@@ -214,6 +308,8 @@ private:
 	const UE::Chaos::ClothAsset::FCollectionClothPatternConstFacade Pattern;
 
 	TArray<int32> EmptyArray;
+
+	const FReferenceSkeleton* RefSkeleton = nullptr;
 };
 
 
@@ -255,7 +351,7 @@ void FClothPatternToDynamicMesh::Convert(const UChaosClothAsset* ClothAssetMeshI
 
 #else
 
-void  FClothAssetToDynamicMesh::Convert(const UChaosClothAsset* ClothAssetMeshIn, FDynamicMesh3& MeshOut, int32 LODIndex, int32 PatternIndex)
+void FClothAssetToDynamicMesh::Convert(const UChaosClothAsset* ClothAssetMeshIn, FDynamicMesh3& MeshOut, int32 LODIndex, int32 PatternIndex)
 {
 	// Conversion only supported with editor.
 	check(0);
diff --git a/Engine/Plugins/Experimental/ChaosClothAssetEditor/Source/ChaosClothAssetTools/Public/ChaosClothAsset/ClothPatternToDynamicMesh.h b/Engine/Plugins/Experimental/ChaosClothAssetEditor/Source/ChaosClothAssetTools/Public/ChaosClothAsset/ClothPatternToDynamicMesh.h
--- a/Engine/Plugins/Experimental/ChaosClothAssetEditor/Source/ChaosClothAssetTools/Public/ChaosClothAsset/ClothPatternToDynamicMesh.h
+++ b/Engine/Plugins/Experimental/ChaosClothAssetEditor/Source/ChaosClothAssetTools/Public/ChaosClothAsset/ClothPatternToDynamicMesh.h
@@ -2,7 +2,9 @@
 
 #pragma once
 
-#include "CoreMinimal.h"
+
+#include "CoreTypes.h"
+#include "Templates/SharedPointer.h"
 
 namespace UE::Geometry { class FDynamicMesh3; }
 class UChaosClothAsset;
