Cloth Editor: create a TransferSkinWeight node. Replicates the current functionality of the weight transfer tool, in a Dataflow node

diff --git a/Engine/Plugins/Experimental/ChaosClothAssetEditor/Source/ChaosClothAssetDataflowNodes/ChaosClothAssetDataflowNodes.Build.cs b/Engine/Plugins/Experimental/ChaosClothAssetEditor/Source/ChaosClothAssetDataflowNodes/ChaosClothAssetDataflowNodes.Build.cs
--- a/Engine/Plugins/Experimental/ChaosClothAssetEditor/Source/ChaosClothAssetDataflowNodes/ChaosClothAssetDataflowNodes.Build.cs
+++ b/Engine/Plugins/Experimental/ChaosClothAssetEditor/Source/ChaosClothAssetDataflowNodes/ChaosClothAssetDataflowNodes.Build.cs
@@ -26,11 +26,14 @@ public class ChaosClothAssetDataflowNodes : ModuleRules
 				"DataflowEditor",
 				"DataflowEngine",
 				"DatasmithImporter",
+				"DynamicMesh",
 				"Engine",
 				"ExternalSource",
-				"Slate",
 				"GeometryCore",
-				"MeshConversion"
+				"MeshConversion",
+				"MeshDescription",
+				"SkeletalMeshDescription",
+				"Slate",
 				// ... add private dependencies that you statically link with here ...	
 			}
 			);
diff --git a/Engine/Plugins/Experimental/ChaosClothAssetEditor/Source/ChaosClothAssetDataflowNodes/Private/ChaosClothAsset/DataflowNodes.cpp b/Engine/Plugins/Experimental/ChaosClothAssetEditor/Source/ChaosClothAssetDataflowNodes/Private/ChaosClothAsset/DataflowNodes.cpp
--- a/Engine/Plugins/Experimental/ChaosClothAssetEditor/Source/ChaosClothAssetDataflowNodes/Private/ChaosClothAsset/DataflowNodes.cpp
+++ b/Engine/Plugins/Experimental/ChaosClothAssetEditor/Source/ChaosClothAssetDataflowNodes/Private/ChaosClothAsset/DataflowNodes.cpp
@@ -13,6 +13,7 @@
 #include "ChaosClothAsset/SimulationXPBDConfigNode.h"
 #include "ChaosClothAsset/BindToRootBoneNode.h"
 #include "ChaosClothAsset/AddWeightMapNode.h"
+#include "ChaosClothAsset/TransferSkinWeightsNode.h"
 #include "Dataflow/DataflowNodeFactory.h"
 #include "Framework/Notifications/NotificationManager.h"
 #include "Widgets/Notifications/SNotificationList.h"
@@ -35,6 +36,7 @@ namespace UE::Chaos::ClothAsset::DataflowNodes
 		DATAFLOW_NODE_REGISTER_CREATION_FACTORY(FChaosClothAssetSimulationXPBDConfigNode);
 		DATAFLOW_NODE_REGISTER_CREATION_FACTORY(FChaosClothAssetBindToRootBoneNode);
 		DATAFLOW_NODE_REGISTER_CREATION_FACTORY(FChaosClothAssetAddWeightMapNode);
+		DATAFLOW_NODE_REGISTER_CREATION_FACTORY(FChaosClothAssetTransferSkinWeightsNode);
 	}
 
 	void LogAndToastWarning(const FText& Error)
diff --git a/Engine/Plugins/Experimental/ChaosClothAssetEditor/Source/ChaosClothAssetDataflowNodes/Private/ChaosClothAsset/TerminalNode.cpp b/Engine/Plugins/Experimental/ChaosClothAssetEditor/Source/ChaosClothAssetDataflowNodes/Private/ChaosClothAsset/TerminalNode.cpp
--- a/Engine/Plugins/Experimental/ChaosClothAssetEditor/Source/ChaosClothAssetDataflowNodes/Private/ChaosClothAsset/TerminalNode.cpp
+++ b/Engine/Plugins/Experimental/ChaosClothAssetEditor/Source/ChaosClothAssetDataflowNodes/Private/ChaosClothAsset/TerminalNode.cpp
@@ -34,7 +34,6 @@ void FChaosClothAssetTerminalNode::SetAssetValue(TObjectPtr<UObject> Asset, Data
 		ClothFacade.DefineSchema();
 
 		// Iterate through the LODs
-		FString SkeletonPathName;
 		FString PhysicsAssetPathName;
 
 		const TArray<const FManagedArrayCollection*> CollectionLods = GetCollectionLods();
@@ -95,12 +94,11 @@ void FChaosClothAssetTerminalNode::SetAssetValue(TObjectPtr<UObject> Asset, Data
 					}
 				}
 
-				// Set properties, skeleton, and physics asset only with LOD 0 at the moment
+				// Set properties and physics asset only with LOD 0 at the moment
 				if (LodIndex == 0)
 				{
 					using namespace ::Chaos::Softs;
 
-					SkeletonPathName = InClothLodFacade.GetSkeletonAssetPathName();
 					PhysicsAssetPathName = InClothLodFacade.GetPhysicsAssetPathName();
 				
 					FCollectionPropertyMutableFacade(ClothCollection).Append(*InClothCollection);
@@ -116,10 +114,7 @@ void FChaosClothAssetTerminalNode::SetAssetValue(TObjectPtr<UObject> Asset, Data
 
 		// Set reference skeleton
 		constexpr bool bRebuildModels = false;  // Avoid rebuilding the asset twice
-		USkeleton* const Skeleton = !SkeletonPathName.IsEmpty() ?
-			LoadObject<USkeleton>(nullptr, *SkeletonPathName, nullptr, LOAD_None, nullptr) :
-			nullptr;
-		ClothAsset->SetSkeleton(Skeleton, bRebuildModels);
+		ClothAsset->UpdateSkeletonFromCollection(bRebuildModels);
 
 		// Set physics asset (note: the cloth asset's physics asset is only replaced if a collection path name is found valid)
 		UPhysicsAsset* const PhysicsAsset = !PhysicsAssetPathName.IsEmpty() ?
diff --git a/Engine/Plugins/Experimental/ChaosClothAssetEditor/Source/ChaosClothAssetDataflowNodes/Private/ChaosClothAsset/TransferSkinWeightsNode.cpp b/Engine/Plugins/Experimental/ChaosClothAssetEditor/Source/ChaosClothAssetDataflowNodes/Private/ChaosClothAsset/TransferSkinWeightsNode.cpp
new file mode 100644
--- /dev/null
+++ b/Engine/Plugins/Experimental/ChaosClothAssetEditor/Source/ChaosClothAssetDataflowNodes/Private/ChaosClothAsset/TransferSkinWeightsNode.cpp
@@ -0,0 +1,168 @@
+// Copyright Epic Games, Inc. All Rights Reserved.
+
+#include "ChaosClothAsset/TransferSkinWeightsNode.h"
+#include "ChaosClothAsset/DataflowNodes.h"
+#include "ChaosClothAsset/CollectionClothFacade.h"
+#include "MeshDescriptionToDynamicMesh.h"
+#include "Rendering/SkeletalMeshLODImporterData.h"
+#include "Rendering/SkeletalMeshModel.h"
+#include "DynamicMesh/MeshTransforms.h"
+#include "Operations/TransferBoneWeights.h"
+#include "SkeletalMeshAttributes.h"
+#include "Engine/SkeletalMesh.h"
+#include "Animation/Skeleton.h"
+
+#include UE_INLINE_GENERATED_CPP_BY_NAME(TransferSkinWeightsNode)
+
+#define LOCTEXT_NAMESPACE "ChaosClothAssetTransferSkinWeightsNode"
+
+namespace UE::ChaosClothAsset::Private
+{
+	void SkeletalMeshToDynamicMesh(USkeletalMesh* FromSkeletalMeshAsset, int32 SourceLODIdx, FDynamicMesh3& ToDynamicMesh)
+	{
+		FMeshDescription SourceMesh;
+
+		// Check first if we have bulk data available and non-empty.
+		if (FromSkeletalMeshAsset->IsLODImportedDataBuildAvailable(SourceLODIdx) && !FromSkeletalMeshAsset->IsLODImportedDataEmpty(SourceLODIdx))
+		{
+			FSkeletalMeshImportData SkeletalMeshImportData;
+			FromSkeletalMeshAsset->LoadLODImportedData(SourceLODIdx, SkeletalMeshImportData);
+			SkeletalMeshImportData.GetMeshDescription(SourceMesh);
+		}
+		else
+		{
+			// Fall back on the LOD model directly if no bulk data exists. When we commit
+			// the mesh description, we override using the bulk data. This can happen for older
+			// skeletal meshes, from UE 4.24 and earlier.
+			const FSkeletalMeshModel* SkeletalMeshModel = FromSkeletalMeshAsset->GetImportedModel();
+			if (SkeletalMeshModel && SkeletalMeshModel->LODModels.IsValidIndex(SourceLODIdx))
+			{
+				SkeletalMeshModel->LODModels[SourceLODIdx].GetMeshDescription(SourceMesh, FromSkeletalMeshAsset);
+			}
+		}
+
+		FMeshDescriptionToDynamicMesh Converter;
+		Converter.Convert(&SourceMesh, ToDynamicMesh);
+	};
+
+}
+
+
+FChaosClothAssetTransferSkinWeightsNode::FChaosClothAssetTransferSkinWeightsNode(const Dataflow::FNodeParameters& InParam, FGuid InGuid)
+	: FDataflowNode(InParam, InGuid)
+{
+	RegisterInputConnection(&Collection);
+	RegisterOutputConnection(&Collection, &Collection);
+	RegisterInputConnection(&SkeletalMesh);
+	RegisterInputConnection(&Transform);
+	RegisterInputConnection(&SkeletalMeshLOD);
+}
+
+void FChaosClothAssetTransferSkinWeightsNode::Evaluate(Dataflow::FContext& Context, const FDataflowOutput* Out) const
+{
+	using namespace UE::Chaos::ClothAsset;
+	using namespace UE::AnimationCore;
+	using namespace UE::Geometry;
+
+	if (Out->IsA<FManagedArrayCollection>(&Collection))
+	{
+		// Evaluate inputs
+		FManagedArrayCollection InputCollection = GetValue<FManagedArrayCollection>(Context, &Collection);
+		const TSharedRef<FManagedArrayCollection> ClothCollection = MakeShared<FManagedArrayCollection>(MoveTemp(InputCollection));
+		const TObjectPtr<USkeletalMesh> InputSkeletalMesh = GetValue<TObjectPtr<USkeletalMesh>>(Context, &SkeletalMesh);
+		const FTransform InputRelativeTransform = GetValue<FTransform>(Context, &Transform);
+		const int32 SourceLODIdx = GetValue<int32>(Context, &SkeletalMeshLOD);
+
+		if (InputSkeletalMesh)
+		{
+			if (!InputSkeletalMesh->IsValidLODIndex(SourceLODIdx))
+			{
+				UE::Chaos::ClothAsset::DataflowNodes::LogAndToastWarning(LOCTEXT("Warning_InvalidSkeletalMeshLOD", "TransferSkinWeightsNode: The specified input SkeletalMesh LOD is not valid"));
+				return;
+			}
+
+			// Convert source Skeletal Mesh to Dynamic Mesh
+			FDynamicMesh3 SourceDynamicMesh;
+			UE::ChaosClothAsset::Private::SkeletalMeshToDynamicMesh(InputSkeletalMesh, SourceLODIdx, SourceDynamicMesh);
+			MeshTransforms::ApplyTransform(SourceDynamicMesh, InputRelativeTransform, true);
+
+			FCollectionClothFacade ClothFacade(ClothCollection);
+
+			const USkeleton* const InputSkeleleton = InputSkeletalMesh->GetSkeleton();
+
+			// Compute bone index mappings
+			TMap<FName, FBoneIndexType> TargetBoneToIndex;
+			const FReferenceSkeleton& TargetRefSkeleton = InputSkeleleton->GetReferenceSkeleton();
+			for (int32 Index = 0; Index < TargetRefSkeleton.GetRawBoneNum(); ++Index)
+			{
+				TargetBoneToIndex.Add(TargetRefSkeleton.GetRawRefBoneInfo()[Index].Name, Index);
+			}
+
+			// Setup bone weight transfer operator
+			FTransferBoneWeights TransferBoneWeights(&SourceDynamicMesh, FSkeletalMeshAttributes::DefaultSkinWeightProfileName);
+			if (TransferBoneWeights.Validate() != EOperationValidationResult::Ok)
+			{
+				UE::Chaos::ClothAsset::DataflowNodes::LogAndToastWarning(LOCTEXT("Warning_TransferWeightsFailed", "TransferSkinWeightsNode: Transferring skin weights failed"));
+				return;
+			}
+
+			// Iterate over the LODs and transfer the bone weights from the source Skeletal mesh to the Cloth asset
+			for (int TargetLODIdx = 0; TargetLODIdx < ClothFacade.GetNumLods(); ++TargetLODIdx)
+			{
+				FCollectionClothLodFacade ClothLodFacade = ClothFacade.GetLod(TargetLODIdx);
+
+				// Cloth collection data arrays we are writing to
+				TArrayView<int32> SimNumBoneInfluences = ClothLodFacade.GetSimNumBoneInfluences();
+				TArrayView<TArray<int32>> SimBoneIndices = ClothLodFacade.GetSimBoneIndices();
+				TArrayView<TArray<float>> SimBoneWeights = ClothLodFacade.GetSimBoneWeights();
+
+				TArrayView<int32> RenderNumBoneInfluences = ClothLodFacade.GetRenderNumBoneInfluences();
+				TArrayView<TArray<int32>> RenderBoneIndices = ClothLodFacade.GetRenderBoneIndices();
+				TArrayView<TArray<float>> RenderBoneWeights = ClothLodFacade.GetRenderBoneWeights();
+
+				const TArrayView<FVector3f> SimPositions = ClothLodFacade.GetSimRestPosition();
+
+				checkSlow(SimPositions.Num() == SimBoneIndices.Num());
+
+				const int32 NumVert = ClothLodFacade.GetNumSimVertices();
+				constexpr bool bUseParallel = true;
+
+				// Iterate over each vertex and write the data from FBoneWeights into cloth collection managed arrays
+				ParallelFor(NumVert, [&](int32 VertexID)
+				{
+					const FVector3d PosD = FVector3d(SimPositions[VertexID]);
+
+					UE::AnimationCore::FBoneWeights BoneWeights;
+					TransferBoneWeights.Compute(PosD, FTransformSRT3d::Identity(), BoneWeights, &TargetBoneToIndex);
+
+					const int32 NumBones = BoneWeights.Num();
+
+					SimNumBoneInfluences[VertexID] = NumBones;
+					SimBoneIndices[VertexID].SetNum(NumBones);
+					SimBoneWeights[VertexID].SetNum(NumBones);
+
+					RenderNumBoneInfluences[VertexID] = NumBones;
+					RenderBoneIndices[VertexID].SetNum(NumBones);
+					RenderBoneWeights[VertexID].SetNum(NumBones);
+
+					for (int BoneIdx = 0; BoneIdx < NumBones; ++BoneIdx)
+					{
+						SimBoneIndices[VertexID][BoneIdx] = BoneWeights[BoneIdx].GetBoneIndex();
+						SimBoneWeights[VertexID][BoneIdx] = BoneWeights[BoneIdx].GetWeight();
+
+						RenderBoneIndices[VertexID][BoneIdx] = BoneWeights[BoneIdx].GetBoneIndex();
+						RenderBoneWeights[VertexID][BoneIdx] = BoneWeights[BoneIdx].GetWeight();
+					}
+
+				}, bUseParallel ? EParallelForFlags::None : EParallelForFlags::ForceSingleThread);
+
+
+				ClothLodFacade.SetSkeletonAssetPathName(InputSkeleleton->GetPathName());
+			}
+		}
+
+		SetValue<FManagedArrayCollection>(Context, *ClothCollection, &Collection);
+	}
+}
+
+#undef LOCTEXT_NAMESPACE
diff --git a/Engine/Plugins/Experimental/ChaosClothAssetEditor/Source/ChaosClothAssetDataflowNodes/Public/ChaosClothAsset/TransferSkinWeightsNode.h b/Engine/Plugins/Experimental/ChaosClothAssetEditor/Source/ChaosClothAssetDataflowNodes/Public/ChaosClothAsset/TransferSkinWeightsNode.h
new file mode 100644
--- /dev/null
+++ b/Engine/Plugins/Experimental/ChaosClothAssetEditor/Source/ChaosClothAssetDataflowNodes/Public/ChaosClothAsset/TransferSkinWeightsNode.h
@@ -0,0 +1,37 @@
+// Copyright Epic Games, Inc. All Rights Reserved.
+
+#pragma once 
+
+#include "Dataflow/DataflowNode.h"
+#include "GeometryCollection/ManagedArrayCollection.h"
+#include "BoneWeights.h"
+#include "TransferSkinWeightsNode.generated.h"
+
+class USkeletalMesh;
+
+USTRUCT(Meta = (DataflowCloth))
+struct FChaosClothAssetTransferSkinWeightsNode : public FDataflowNode
+{
+	GENERATED_USTRUCT_BODY()
+	DATAFLOW_NODE_DEFINE_INTERNAL(FChaosClothAssetTransferSkinWeightsNode, "TransferSkinWeights", "Cloth", "Cloth Transfer Skin Weights")
+	//DATAFLOW_NODE_RENDER_TYPE(FGeometryCollection::StaticType(), "Collection")  // TODO: Leave out the render type until there is something to render
+
+public:
+
+	UPROPERTY(Meta = (Dataflowinput, DataflowOutput, DisplayName = "Collection", DataflowPassthrough = "Collection"))
+	FManagedArrayCollection Collection;
+
+	UPROPERTY(EditAnywhere, Category = "Source Mesh", Meta = (Dataflowinput, DisplayName = "SkeletalMesh"))
+	TObjectPtr<USkeletalMesh> SkeletalMesh;
+
+	UPROPERTY(EditAnywhere, Category = "Source Mesh", Meta = (Dataflowinput, DisplayName = "SkeletalMeshLOD"))
+	int32 SkeletalMeshLOD;
+
+	// The relative transform between the SkeletalMesh and ClothCollection
+	UPROPERTY(EditAnywhere, Category = "Source Mesh", Meta = (Dataflowinput, DisplayName = "Transform"))
+	FTransform Transform;
+
+	FChaosClothAssetTransferSkinWeightsNode(const Dataflow::FNodeParameters& InParam, FGuid InGuid = FGuid::NewGuid());
+
+	virtual void Evaluate(Dataflow::FContext& Context, const FDataflowOutput* Out) const override;
+};
