ClothTrainingTool: start simulation without blocking game thread

diff --git a/Engine/Plugins/Experimental/ChaosClothAssetEditor/Source/ChaosClothAssetEditorTools/Private/ChaosClothAsset/ClothTrainingTool.cpp b/Engine/Plugins/Experimental/ChaosClothAssetEditor/Source/ChaosClothAssetEditorTools/Private/ChaosClothAsset/ClothTrainingTool.cpp
--- a/Engine/Plugins/Experimental/ChaosClothAssetEditor/Source/ChaosClothAssetEditorTools/Private/ChaosClothAsset/ClothTrainingTool.cpp
+++ b/Engine/Plugins/Experimental/ChaosClothAssetEditor/Source/ChaosClothAssetEditorTools/Private/ChaosClothAsset/ClothTrainingTool.cpp
@@ -2,7 +2,7 @@
 
 #include "ChaosClothAsset/ClothTrainingTool.h"
 #include "CoreMinimal.h"
-
+#include "Animation/AnimSequence.h"
 #include "Animation/AttributesRuntime.h"
 #include "BonePose.h"
 #include "Chaos/ChaosCache.h"
@@ -20,8 +20,7 @@
 #include "InteractiveToolManager.h"
 #include "Internationalization/Regex.h"
 #include "ModelingOperators.h"
-#include "Misc/DateTime.h"
-#include "Misc/ScopedSlowTask.h"
+#include "Misc/AsyncTaskNotification.h"
 #include "PhysicsEngine/PhysicsAsset.h"
 #include "Tasks/Pipe.h"
 #include "ToolTargetManager.h"
@@ -501,88 +500,105 @@ void UClothTrainingTool::Setup()
 	AddToolPropertySource(ActionProperties);
 }
 
-void UClothTrainingTool::RunTraining()
+void UClothTrainingTool::StartTraining()
 {
+	check(PendingAction == EClothTrainingToolActions::StartTrain);
 	if (!IsClothComponentValid() || ToolProperties->AnimationSequence == nullptr)
 	{
+		PendingAction = EClothTrainingToolActions::NoAction;
 		return;
 	}
 	UChaosCacheCollection* const CacheCollection = GetCacheCollection();
 	if (CacheCollection == nullptr)
 	{
+		PendingAction = EClothTrainingToolActions::NoAction;
 		return;
 	}
-	using UE::ClothTrainingTool::Private::GetCache;
-	UChaosCache& Cache = GetCache(*CacheCollection);
-	
-	using FTaskType = UE::Geometry::TModelingOpTask<FLaunchSimsOp>;
-	using FExecuterType = UE::Geometry::FAsyncTaskExecuterWithProgressCancel<FTaskType>;
-
-	const FText DefaultMessage(LOCTEXT("ClothTrainingMessage", "Generate training data..."));
-
-	if (!AllocateSimResources_GameThread(ToolProperties->NumThreads))
+	if (TaskResource != nullptr)
+	{
+		PendingAction = EClothTrainingToolActions::NoAction;
+		return;
+	}
+	TaskResource = MakeUnique<FTaskResource>();
+	if (!TaskResource->AllocateSimResources_GameThread(*ClothComponent, ToolProperties->NumThreads))
 	{
+		PendingAction = EClothTrainingToolActions::NoAction;
 		return;
 	}
-	FCacheUserToken CacheUserToken = Cache.BeginRecord(ClothComponent, FGuid(), FTransform::Identity);
 
-	TUniquePtr<FLaunchSimsOp> NewOp = MakeUnique<FLaunchSimsOp>(SimResources, *SimMutex, ToolProperties);
-	TUniquePtr<FExecuterType> BackgroundTaskExecuter = MakeUnique<FExecuterType>(MoveTemp(NewOp));
-	BackgroundTaskExecuter->StartBackgroundTask();
+	using UE::ClothTrainingTool::Private::GetCache;
+	UChaosCache* const Cache = &GetCache(*CacheCollection);
+	TaskResource->Cache = Cache;
+	TaskResource->CacheUserToken = MakeUnique<FCacheUserToken>(Cache->BeginRecord(ClothComponent, FGuid(), FTransform::Identity));
+
+	TUniquePtr<FLaunchSimsOp> NewOp = MakeUnique<FLaunchSimsOp>(TaskResource->SimResources, *TaskResource->SimMutex, ToolProperties);
+	TaskResource->Executer = MakeUnique<FExecuterType>(MoveTemp(NewOp));
+	TaskResource->Executer->StartBackgroundTask();
+
+	FAsyncTaskNotificationConfig NotificationConfig;
+	NotificationConfig.TitleText = LOCTEXT("SimulateCloth", "Simulating Cloth");
+	NotificationConfig.ProgressText = FText::FromString(TEXT("0%"));
+	NotificationConfig.bCanCancel = true;
+	NotificationConfig.bKeepOpenOnSuccess = true;
+	NotificationConfig.bKeepOpenOnFailure = true;
+	TaskResource->Notification = MakeUnique<FAsyncTaskNotification>(NotificationConfig);
+	TaskResource->StartTime = FDateTime::UtcNow();
+	TaskResource->LastUpdateTime = TaskResource->StartTime;
+
+	PendingAction = EClothTrainingToolActions::TickTrain;
+}
+
+void UClothTrainingTool::TickTraining()
+{
+	check(PendingAction == EClothTrainingToolActions::TickTrain && TaskResource != nullptr);
 
-	FScopedSlowTask SlowTask(1, DefaultMessage);
-	SlowTask.MakeDialog(true);
+	bool bFinished = false;
+	const bool bCancelled = TaskResource->Notification->GetPromptAction() == EAsyncTaskNotificationPromptAction::Cancel;
+	if (bCancelled)
+	{
+		TaskResource->Executer.Release()->CancelAndDelete();
+		bFinished = true;
+	}
+	else if (TaskResource->Executer->IsDone())
+	{
+		bFinished = true;
+	}
 
+	if (!bFinished)
 	{
-		UE::ClothTrainingTool::Private::FTimeScope TimeScope(TEXT("Cloth generation"));
-		while (true)
+		const FDateTime CurrentTime = FDateTime::UtcNow();
+		const double SinceLastUpdate = (CurrentTime - TaskResource->LastUpdateTime).GetTotalSeconds();
+		if (SinceLastUpdate < 0.2)
 		{
-			if (SlowTask.ShouldCancel())
-			{
-				// Release ownership to the TDeleterTask that is spawned by CancelAndDelete()
-				BackgroundTaskExecuter.Release()->CancelAndDelete();
-				break;
-			}
-			if (BackgroundTaskExecuter->IsDone())
-			{
-				break;
-			}
-			FPlatformProcess::Sleep(.2); // SlowTask::ShouldCancel will throttle any updates faster than .2 seconds
-			float ProgressFrac;
-			FText ProgressMessage;
-			const bool bMadeProgress = BackgroundTaskExecuter->PollProgress(ProgressFrac, ProgressMessage);
-			if (bMadeProgress)
-			{
-				// SlowTask expects progress to be reported before it happens; we work around this by directly updating the progress amount
-				SlowTask.CompletedWork = ProgressFrac;
-				SlowTask.EnterProgressFrame(0, ProgressMessage);
-			}
-			else
-			{
-				SlowTask.TickProgress(); // Still tick the UI when we don't get new progress frames
-			}
+			return;
 		}
-		FreeSimResources_GameThread();
+
+		float ProgressFrac;
+		FText ProgressMessage;
+		const bool bMadeProgress = TaskResource->Executer->PollProgress(ProgressFrac, ProgressMessage);
+		if (bMadeProgress)
+		{
+			ProgressMessage = FText::FromString(FString::Printf(TEXT("%d%%"), int32(ProgressFrac * 100)));
+			TaskResource->Notification->SetProgressText(ProgressMessage);
+		}
+		TaskResource->LastUpdateTime = CurrentTime;
 	}
+	else
 	{
-		UE::ClothTrainingTool::Private::FTimeScope TimeScope(TEXT("Saving"));
-		Cache.bCompressChannels = true;
-		Cache.EndRecord(CacheUserToken);
-
-		SaveCacheCollection(CacheCollection);
+		FreeTaskResource(bCancelled);
+		PendingAction = EClothTrainingToolActions::NoAction;
 	}
 }
 
-
 void UClothTrainingTool::OnTick(float DeltaTime)
 {
-	if (PendingAction != EClothTrainingToolActions::NoAction)
+	if (PendingAction == EClothTrainingToolActions::StartTrain)
 	{
-		if (PendingAction == EClothTrainingToolActions::Train)
-		{
-			RunTraining();
-		}
-		PendingAction = EClothTrainingToolActions::NoAction;
+		StartTraining();
+	}
+	else if (PendingAction == EClothTrainingToolActions::TickTrain)
+	{
+		TickTraining();
 	}
 }
 
@@ -598,18 +614,27 @@ void UClothTrainingTool::RequestAction(EClothTrainingToolActions ActionType)
 
 void UClothTrainingTool::Shutdown(EToolShutdownType ShutdownType)
 {
+	if (TaskResource != nullptr)
+	{
+		if (TaskResource->Executer != nullptr)
+		{
+			TaskResource->Executer.Release()->CancelAndDelete();
+		}
+		constexpr bool bCancelled = true;
+		FreeTaskResource(bCancelled);
+	}
 	Super::Shutdown(ShutdownType);
 	ToolProperties->SaveProperties(this);
 }
 
-bool UClothTrainingTool::AllocateSimResources_GameThread(int32 Num)
+bool UClothTrainingTool::FTaskResource::AllocateSimResources_GameThread(const UChaosClothComponent& ClothComponent, int32 Num)
 {
 	SimResources.Reserve(Num);
 	for(int32 Index = 0; Index < Num; ++Index)
 	{
-		UChaosClothComponent* const CopyComponent = DuplicateObject<UChaosClothComponent>(ClothComponent, ClothComponent->GetOuter());
+		UChaosClothComponent* const CopyComponent = DuplicateObject<UChaosClothComponent>(&ClothComponent, ClothComponent.GetOuter());
 		CopyComponent->RegisterComponent();
-		CopyComponent->SetWorldTransform(ClothComponent->GetComponentTransform());
+		CopyComponent->SetWorldTransform(ClothComponent.GetComponentTransform());
 		FSimResource SimResource;
 		SimResource.ClothComponent = CopyComponent;
 		SimResource.Proxy = MakeUnique<FClothSimulationDataGenerationProxy>(*CopyComponent);
@@ -626,7 +651,7 @@ bool UClothTrainingTool::AllocateSimResources_GameThread(int32 Num)
 	return true;
 }
 
-void UClothTrainingTool::FreeSimResources_GameThread()
+void UClothTrainingTool::FTaskResource::FreeSimResources_GameThread()
 {
 	{
 		FScopeLock Lock(SimMutex.Get());
@@ -689,4 +714,33 @@ bool UClothTrainingTool::SaveCacheCollection(UChaosCacheCollection* CacheCollect
 	return bSaveSucced;
 }
 
+void UClothTrainingTool::FreeTaskResource(bool bCancelled)
+{
+	TaskResource->Notification->SetProgressText(LOCTEXT("Finishing", "Finishing, please wait"));
+	TaskResource->FreeSimResources_GameThread();
+	const FDateTime CurrentTime = FDateTime::UtcNow();
+	UE_LOG(LogClothTrainingTool, Log, TEXT("Training finished in %f seconds"), (CurrentTime - TaskResource->StartTime).GetTotalSeconds());
+
+	{
+		UE::ClothTrainingTool::Private::FTimeScope TimeScope(TEXT("Saving"));
+		TaskResource->Cache->bCompressChannels = true;
+		TaskResource->Cache->EndRecord(*TaskResource->CacheUserToken);
+
+		UChaosCacheCollection* const CacheCollection = GetCacheCollection();
+		ensureMsgf(CacheCollection != nullptr, TEXT("CacheCollection should not be nullptr"));
+		SaveCacheCollection(CacheCollection);
+	}
+	if (bCancelled)
+	{
+		TaskResource->Notification->SetProgressText(LOCTEXT("Cancelled", "Cancelled"));
+		TaskResource->Notification->SetComplete(false);
+	}
+	else
+	{
+		TaskResource->Notification->SetProgressText(LOCTEXT("Finished", "Finished"));
+		TaskResource->Notification->SetComplete(true);
+	}
+	TaskResource.Reset();
+}
+
 #undef LOCTEXT_NAMESPACE
diff --git a/Engine/Plugins/Experimental/ChaosClothAssetEditor/Source/ChaosClothAssetEditorTools/Public/ChaosClothAsset/ClothTrainingTool.h b/Engine/Plugins/Experimental/ChaosClothAssetEditor/Source/ChaosClothAssetEditorTools/Public/ChaosClothAsset/ClothTrainingTool.h
--- a/Engine/Plugins/Experimental/ChaosClothAssetEditor/Source/ChaosClothAssetEditorTools/Public/ChaosClothAsset/ClothTrainingTool.h
+++ b/Engine/Plugins/Experimental/ChaosClothAssetEditor/Source/ChaosClothAssetEditorTools/Public/ChaosClothAsset/ClothTrainingTool.h
@@ -5,13 +5,15 @@
 #include "UObject/Object.h"
 #include "SingleSelectionTool.h"
 #include "InteractiveToolBuilder.h"
-#include "Animation/AnimSequence.h"
 #include "ModelingOperators.h"
 #include "MeshOpPreviewHelpers.h"
-#include "Misc/SlowTask.h"
+#include "Misc/DateTime.h"
 
 #include "ClothTrainingTool.generated.h"
 
+struct FCacheUserToken;
+class FAsyncTaskNotification;
+class UAnimSequence;
 class UChaosCache;
 class UChaosCacheCollection;
 class UChaosClothComponent;
@@ -63,7 +65,8 @@ UENUM()
 enum class EClothTrainingToolActions
 {
 	NoAction,
-	Train
+	StartTrain,
+	TickTrain
 };
 
 UCLASS()
@@ -81,7 +84,7 @@ public:
 	UFUNCTION(CallInEditor, Category = Actions, meta = (DisplayName = "Begin Generating", DisplayPriority = 1))
 		void StartGenerating()
 	{
-		PostAction(EClothTrainingToolActions::Train);
+		PostAction(EClothTrainingToolActions::StartTrain);
 	}
 
 };
@@ -122,6 +125,9 @@ private:
 	struct FSimResource;
 	class FLaunchSimsOp;
 
+	using FTaskType = UE::Geometry::TModelingOpTask<FLaunchSimsOp>;
+	using FExecuterType = UE::Geometry::FAsyncTaskExecuterWithProgressCancel<FTaskType>;
+
 	friend class UClothTrainingToolActionProperties;
 
 	UPROPERTY()
@@ -135,19 +141,32 @@ private:
 	UPROPERTY()
 	TObjectPtr<const UChaosClothComponent> ClothComponent;
 
-	TUniquePtr<FCriticalSection> SimMutex;
-	TArray<FSimResource> SimResources;
+	struct FTaskResource
+	{
+		TUniquePtr<FCriticalSection> SimMutex;
+		TArray<FSimResource> SimResources;
+
+		TUniquePtr<FExecuterType> Executer;
+		TUniquePtr<FAsyncTaskNotification> Notification;
+		UChaosCache *Cache = nullptr;
+		TUniquePtr<FCacheUserToken> CacheUserToken;
+		FDateTime StartTime;
+		FDateTime LastUpdateTime;
+
+		bool AllocateSimResources_GameThread(const UChaosClothComponent& ClothComponent, int32 Num);
+		void FreeSimResources_GameThread();
+	};
+
+	TUniquePtr<FTaskResource> TaskResource = nullptr;
 
 	void RequestAction(EClothTrainingToolActions ActionType);
+	void StartTraining();
+	void TickTraining();
 	void RunTraining();
+	void FreeTaskResource(bool bCancelled);
 	bool IsClothComponentValid() const;
 	UChaosCacheCollection* GetCacheCollection() const;
 	bool SaveCacheCollection(UChaosCacheCollection* CacheCollection) const;
-	void PrepareAnimationSequence();
-	void RestoreAnimationSequence();
-	bool AllocateSimResources_GameThread(int32 Num);
-	void FreeSimResources_GameThread();
-
 };
 
 
