ClothTrainingTool: fix static analysis error

diff --git a/Engine/Plugins/Experimental/ChaosClothAssetEditor/Source/ChaosClothAssetEditorTools/Private/ChaosClothAsset/ClothTrainingTool.cpp b/Engine/Plugins/Experimental/ChaosClothAssetEditor/Source/ChaosClothAssetEditorTools/Private/ChaosClothAsset/ClothTrainingTool.cpp
--- a/Engine/Plugins/Experimental/ChaosClothAssetEditor/Source/ChaosClothAssetEditorTools/Private/ChaosClothAsset/ClothTrainingTool.cpp
+++ b/Engine/Plugins/Experimental/ChaosClothAssetEditor/Source/ChaosClothAssetEditorTools/Private/ChaosClothAsset/ClothTrainingTool.cpp
@@ -186,6 +186,22 @@ private:
 	inline static const FName PositionZName = TEXT("PositionZ");
 };
 
+struct UClothTrainingTool::FTaskResource
+{
+	TUniquePtr<FCriticalSection> SimMutex;
+	TArray<FSimResource> SimResources;
+
+	TUniquePtr<FExecuterType> Executer;
+	TUniquePtr<FAsyncTaskNotification> Notification;
+	UChaosCache *Cache = nullptr;
+	TUniquePtr<FCacheUserToken> CacheUserToken;
+	FDateTime StartTime;
+	FDateTime LastUpdateTime;
+
+	bool AllocateSimResources_GameThread(const UChaosClothComponent& ClothComponent, int32 Num);
+	void FreeSimResources_GameThread();
+};
+
 void UClothTrainingTool::FLaunchSimsOp::PrepareAnimationSequence()
 {
 	TObjectPtr<UAnimSequence> AnimationSequence = ToolProperties->AnimationSequence;
diff --git a/Engine/Plugins/Experimental/ChaosClothAssetEditor/Source/ChaosClothAssetEditorTools/Public/ChaosClothAsset/ClothTrainingTool.h b/Engine/Plugins/Experimental/ChaosClothAssetEditor/Source/ChaosClothAssetEditorTools/Public/ChaosClothAsset/ClothTrainingTool.h
--- a/Engine/Plugins/Experimental/ChaosClothAssetEditor/Source/ChaosClothAssetEditorTools/Public/ChaosClothAsset/ClothTrainingTool.h
+++ b/Engine/Plugins/Experimental/ChaosClothAssetEditor/Source/ChaosClothAssetEditorTools/Public/ChaosClothAsset/ClothTrainingTool.h
@@ -122,8 +122,9 @@ public:
 
 private:
 	class FClothSimulationDataGenerationProxy;
-	struct FSimResource;
 	class FLaunchSimsOp;
+	struct FSimResource;
+	struct FTaskResource;
 
 	using FTaskType = UE::Geometry::TModelingOpTask<FLaunchSimsOp>;
 	using FExecuterType = UE::Geometry::FAsyncTaskExecuterWithProgressCancel<FTaskType>;
@@ -141,22 +142,6 @@ private:
 	UPROPERTY()
 	TObjectPtr<const UChaosClothComponent> ClothComponent;
 
-	struct FTaskResource
-	{
-		TUniquePtr<FCriticalSection> SimMutex;
-		TArray<FSimResource> SimResources;
-
-		TUniquePtr<FExecuterType> Executer;
-		TUniquePtr<FAsyncTaskNotification> Notification;
-		UChaosCache *Cache = nullptr;
-		TUniquePtr<FCacheUserToken> CacheUserToken;
-		FDateTime StartTime;
-		FDateTime LastUpdateTime;
-
-		bool AllocateSimResources_GameThread(const UChaosClothComponent& ClothComponent, int32 Num);
-		void FreeSimResources_GameThread();
-	};
-
 	TUniquePtr<FTaskResource> TaskResource = nullptr;
 
 	void RequestAction(EClothTrainingToolActions ActionType);
