ClothTrainingTool: fix clang compiler error

diff --git a/Engine/Plugins/Experimental/ChaosClothAssetEditor/Source/ChaosClothAssetEditorTools/Private/ChaosClothAsset/ClothTrainingTool.cpp b/Engine/Plugins/Experimental/ChaosClothAssetEditor/Source/ChaosClothAssetEditorTools/Private/ChaosClothAsset/ClothTrainingTool.cpp
--- a/Engine/Plugins/Experimental/ChaosClothAssetEditor/Source/ChaosClothAssetEditorTools/Private/ChaosClothAsset/ClothTrainingTool.cpp
+++ b/Engine/Plugins/Experimental/ChaosClothAssetEditor/Source/ChaosClothAssetEditorTools/Private/ChaosClothAsset/ClothTrainingTool.cpp
@@ -198,7 +198,7 @@ struct UClothTrainingTool::FTaskResource
 	FDateTime StartTime;
 	FDateTime LastUpdateTime;
 
-	bool AllocateSimResources_GameThread(const UChaosClothComponent& ClothComponent, int32 Num);
+	bool AllocateSimResources_GameThread(const UChaosClothComponent& InClothComponent, int32 Num);
 	void FreeSimResources_GameThread();
 };
 
@@ -643,14 +643,14 @@ void UClothTrainingTool::Shutdown(EToolShutdownType ShutdownType)
 	ToolProperties->SaveProperties(this);
 }
 
-bool UClothTrainingTool::FTaskResource::AllocateSimResources_GameThread(const UChaosClothComponent& ClothComponent, int32 Num)
+bool UClothTrainingTool::FTaskResource::AllocateSimResources_GameThread(const UChaosClothComponent& InClothComponent, int32 Num)
 {
 	SimResources.Reserve(Num);
 	for(int32 Index = 0; Index < Num; ++Index)
 	{
-		UChaosClothComponent* const CopyComponent = DuplicateObject<UChaosClothComponent>(&ClothComponent, ClothComponent.GetOuter());
+		UChaosClothComponent* const CopyComponent = DuplicateObject<UChaosClothComponent>(&InClothComponent, InClothComponent.GetOuter());
 		CopyComponent->RegisterComponent();
-		CopyComponent->SetWorldTransform(ClothComponent.GetComponentTransform());
+		CopyComponent->SetWorldTransform(InClothComponent.GetComponentTransform());
 		FSimResource SimResource;
 		SimResource.ClothComponent = CopyComponent;
 		SimResource.Proxy = MakeUnique<FClothSimulationDataGenerationProxy>(*CopyComponent);
