ClothTrainingTool: fix clang compiler error

diff --git a/Engine/Plugins/Experimental/ChaosClothAssetEditor/Source/ChaosClothAssetEditorTools/Public/ChaosClothAsset/ClothTrainingTool.h b/Engine/Plugins/Experimental/ChaosClothAssetEditor/Source/ChaosClothAssetEditorTools/Public/ChaosClothAsset/ClothTrainingTool.h
--- a/Engine/Plugins/Experimental/ChaosClothAssetEditor/Source/ChaosClothAssetEditorTools/Public/ChaosClothAsset/ClothTrainingTool.h
+++ b/Engine/Plugins/Experimental/ChaosClothAssetEditor/Source/ChaosClothAssetEditorTools/Public/ChaosClothAsset/ClothTrainingTool.h
@@ -142,7 +142,7 @@ private:
 	UPROPERTY()
 	TObjectPtr<const UChaosClothComponent> ClothComponent;
 
-	TUniquePtr<FTaskResource> TaskResource = nullptr;
+	TUniquePtr<FTaskResource> TaskResource;
 
 	void RequestAction(EClothTrainingToolActions ActionType);
 	void StartTraining();
