Cloth Editor: have TransferSkinWeightsTool produce a TransferSkinWeightsNode
- when the user clicks Accept on the transfer tool, it will now add a Transfer Node to the graph rather than modifying the cloth asset directly
- a transform gizmo has been added to the tool to move the source SkeletalMesh around. The transform is copied to the node on accept
- the transfer operation for preview is now done in a backgroung thread using UMeshOpPreviewWithBackgroundCompute
- ClothEditorContextObject now handles creation of a node and wiring it into an existing Dataflow graph. WeightMapPaint Tool has been modified to call the same functions.

diff --git a/Engine/Plugins/Experimental/ChaosClothAssetEditor/Source/ChaosClothAssetEditor/Private/ChaosClothAsset/ClothEditorMode.cpp b/Engine/Plugins/Experimental/ChaosClothAssetEditor/Source/ChaosClothAssetEditor/Private/ChaosClothAsset/ClothEditorMode.cpp
--- a/Engine/Plugins/Experimental/ChaosClothAssetEditor/Source/ChaosClothAssetEditor/Private/ChaosClothAsset/ClothEditorMode.cpp
+++ b/Engine/Plugins/Experimental/ChaosClothAssetEditor/Source/ChaosClothAssetEditor/Private/ChaosClothAsset/ClothEditorMode.cpp
@@ -55,7 +55,7 @@
 #include "ToolTargets/SkeletalMeshComponentToolTarget.h"
 #include "Dataflow/DataflowGraphEditor.h"
 #include "ContextObjectStore.h"
-
+#include "BaseGizmos/TransformGizmoUtil.h"
 
 #include UE_INLINE_GENERATED_CPP_BY_NAME(ClothEditorMode)
 
@@ -145,6 +145,8 @@ void UChaosClothAssetEditorMode::Enter()
 		}
 	});
 
+	// Register gizmo ContextObject for use inside interactive tools
+	UE::TransformGizmoUtil::RegisterTransformGizmoContextObject(GetInteractiveToolsContext());
 }
 
 void UChaosClothAssetEditorMode::AddToolTargetFactories()
@@ -198,7 +200,6 @@ void UChaosClothAssetEditorMode::RegisterPreviewTools()
 
 	const FChaosClothAssetEditorCommands& CommandInfos = FChaosClothAssetEditorCommands::Get();
 	RegisterClothTool(CommandInfos.BeginClothTrainingTool, FChaosClothAssetEditorCommands::BeginClothTrainingToolIdentifier, NewObject<UClothTrainingToolBuilder>(), PreviewSceneToolsContext);
-	RegisterClothTool(CommandInfos.BeginTransferSkinWeightsTool, FChaosClothAssetEditorCommands::BeginTransferSkinWeightsToolIdentifier, NewObject<UClothTransferSkinWeightsToolBuilder>(), PreviewSceneToolsContext);
 }
 
 void UChaosClothAssetEditorMode::RegisterTools()
@@ -216,6 +217,7 @@ void UChaosClothAssetEditorMode::RegisterTools()
 
 	RegisterClothTool(CommandInfos.BeginAttributeEditorTool, FChaosClothAssetEditorCommands::BeginAttributeEditorToolIdentifier, NewObject<UAttributeEditorToolBuilder>(), RestSpaceToolsContext);
 	RegisterClothTool(CommandInfos.BeginWeightMapPaintTool, FChaosClothAssetEditorCommands::BeginWeightMapPaintToolIdentifier, NewObject<UClothEditorWeightMapPaintToolBuilder>(), RestSpaceToolsContext);
+	RegisterClothTool(CommandInfos.BeginTransferSkinWeightsTool, FChaosClothAssetEditorCommands::BeginTransferSkinWeightsToolIdentifier, NewObject<UClothTransferSkinWeightsToolBuilder>(), RestSpaceToolsContext);
 }
 
 bool UChaosClothAssetEditorMode::ShouldToolStartBeAllowed(const FString& ToolIdentifier) const
diff --git a/Engine/Plugins/Experimental/ChaosClothAssetEditor/Source/ChaosClothAssetEditorTools/ChaosClothAssetEditorTools.Build.cs b/Engine/Plugins/Experimental/ChaosClothAssetEditor/Source/ChaosClothAssetEditorTools/ChaosClothAssetEditorTools.Build.cs
--- a/Engine/Plugins/Experimental/ChaosClothAssetEditor/Source/ChaosClothAssetEditorTools/ChaosClothAssetEditorTools.Build.cs
+++ b/Engine/Plugins/Experimental/ChaosClothAssetEditor/Source/ChaosClothAssetEditorTools/ChaosClothAssetEditorTools.Build.cs
@@ -58,7 +58,8 @@ public class ChaosClothAssetEditorTools : ModuleRules
 				"ChaosClothAssetDataflowNodes",
 				"DataflowCore",
 				"DataflowEngine",
-				"DataflowEditor"
+				"DataflowEditor",
+				"Chaos"				// FManagedArrayCollection::StaticType
 				// ... add private dependencies that you statically link with here ...	
 			}
 			);
diff --git a/Engine/Plugins/Experimental/ChaosClothAssetEditor/Source/ChaosClothAssetEditorTools/Private/ChaosClothAsset/ClothEditorContextObject.cpp b/Engine/Plugins/Experimental/ChaosClothAssetEditor/Source/ChaosClothAssetEditorTools/Private/ChaosClothAsset/ClothEditorContextObject.cpp
--- a/Engine/Plugins/Experimental/ChaosClothAssetEditor/Source/ChaosClothAssetEditorTools/Private/ChaosClothAsset/ClothEditorContextObject.cpp
+++ b/Engine/Plugins/Experimental/ChaosClothAssetEditor/Source/ChaosClothAssetEditorTools/Private/ChaosClothAsset/ClothEditorContextObject.cpp
@@ -1,6 +1,10 @@
 // Copyright Epic Games, Inc. All Rights Reserved.
 
 #include "ChaosClothAsset/ClothEditorContextObject.h"
+#include "Dataflow/DataflowEdNode.h"
+#include "Dataflow/DataflowGraphEditor.h"
+#include "Dataflow/DataflowSNode.h"
+#include "Dataflow/DataflowObject.h"
 
 #include UE_INLINE_GENERATED_CPP_BY_NAME(ClothEditorContextObject)
 
@@ -30,4 +34,127 @@ const TObjectPtr<const UDataflow> UClothEditorContextObject::GetDataflowGraph()
 	return DataflowGraph;
 }
 
+UEdGraphNode* UClothEditorContextObject::GetSingleSelectedNode() const
+{
+	const TSharedPtr<const SDataflowGraphEditor> PinnedDataflowGraphEditor = DataflowGraphEditor.Pin();
+	if (!PinnedDataflowGraphEditor)
+	{
+		return nullptr;
+	}
+
+	return PinnedDataflowGraphEditor->GetSingleSelectedNode();
+}
+
+
+UEdGraphNode* UClothEditorContextObject::GetSingleSelectedNodeWithOutputType(const FName& SelectedNodeOutputTypeName) const
+{
+	const TSharedPtr<const SDataflowGraphEditor> PinnedDataflowGraphEditor = DataflowGraphEditor.Pin();
+	if (!PinnedDataflowGraphEditor)
+	{
+		return nullptr;
+	}
+
+	UEdGraphNode* const SelectedNode = PinnedDataflowGraphEditor->GetSingleSelectedNode();
+	if (!SelectedNode)
+	{
+		return nullptr;
+	}
+
+	const UDataflowEdNode* const SelectedDataflowEdNode = CastChecked<UDataflowEdNode>(SelectedNode);
+	const TSharedPtr<const FDataflowNode> SelectedDataflowNode = SelectedDataflowEdNode->GetDataflowNode();
+	for (const FDataflowOutput* const Output : SelectedDataflowNode->GetOutputs())
+	{
+		if (Output->GetType() == SelectedNodeOutputTypeName)
+		{
+			return SelectedNode;
+		}
+	}
+
+	return nullptr;
+}
+
+UEdGraphNode* UClothEditorContextObject::CreateNewNode(const FName& NewNodeTypeName)
+{
+	const TSharedPtr<const SDataflowGraphEditor> PinnedDataflowGraphEditor = DataflowGraphEditor.Pin();
+	if (!PinnedDataflowGraphEditor)
+	{
+		return nullptr;
+	}
+
+	const TSharedPtr<FAssetSchemaAction_Dataflow_CreateNode_DataflowEdNode> NodeAction =
+		FAssetSchemaAction_Dataflow_CreateNode_DataflowEdNode::CreateAction(DataflowGraph, NewNodeTypeName);
+	UEdGraphNode* const NewEdNode = NodeAction->PerformAction(DataflowGraph, nullptr, PinnedDataflowGraphEditor->GetPasteLocation(), true);
+
+	return NewEdNode;
+}
+
+UEdGraphNode* UClothEditorContextObject::CreateAndConnectNewNode(
+	const FName& NewNodeTypeName,
+	UEdGraphNode& UpstreamNode,
+	const FName& ConnectionTypeName)
+{
+	// First find the specified output of the upstream node, plus any pins it's connected to
+
+	UEdGraphPin* UpstreamNodeOutputPin = nullptr;
+	TArray<UEdGraphPin*> ExistingNodeInputPins;
+
+	const UDataflowEdNode* const UpstreamDataflowEdNode = CastChecked<UDataflowEdNode>(&UpstreamNode);
+	const TSharedPtr<const FDataflowNode> UpstreamDataflowNode = UpstreamDataflowEdNode->GetDataflowNode();
+
+	for (const FDataflowOutput* const Output : UpstreamDataflowNode->GetOutputs())
+	{
+		if (Output->GetType() == ConnectionTypeName)
+		{
+			UpstreamNodeOutputPin = UpstreamDataflowEdNode->FindPin(*Output->GetName().ToString(), EGPD_Output);
+			ExistingNodeInputPins = UpstreamNodeOutputPin->LinkedTo;
+			break;
+		}
+	}
+
+	// Add the new node 
+
+	UEdGraphNode* const NewEdNode = CreateNewNode(NewNodeTypeName);
+	checkf(NewEdNode, TEXT("Failed to create a new node in the DataflowGraph"));
+
+	UDataflowEdNode* const NewDataflowEdNode = CastChecked<UDataflowEdNode>(NewEdNode);
+	const TSharedPtr<FDataflowNode> NewDataflowNode = NewDataflowEdNode->GetDataflowNode();
+
+	// Re-wire the graph
+
+	if (UpstreamNodeOutputPin)
+	{
+		UEdGraphPin* NewNodeInputPin = nullptr;
+		for (const FDataflowInput* const NewNodeInput : NewDataflowNode->GetInputs())
+		{
+			if (NewNodeInput->GetType() == ConnectionTypeName)
+			{
+				NewNodeInputPin = NewDataflowEdNode->FindPin(*NewNodeInput->GetName().ToString(), EGPD_Input);
+			}
+		}
+
+		UEdGraphPin* NewNodeOutputPin = nullptr;
+		for (const FDataflowOutput* const NewNodeOutput : NewDataflowNode->GetOutputs())
+		{
+			if (NewNodeOutput->GetType() == ConnectionTypeName)
+			{
+				NewNodeOutputPin = NewDataflowEdNode->FindPin(*NewNodeOutput->GetName().ToString(), EGPD_Output);
+				break;
+			}
+		}
+
+		check(NewNodeInputPin);
+		check(NewNodeOutputPin);
+
+		DataflowGraph->GetSchema()->TryCreateConnection(UpstreamNodeOutputPin, NewNodeInputPin);
+
+		for (UEdGraphPin* DownstreamInputPin : ExistingNodeInputPins)
+		{
+			DataflowGraph->GetSchema()->TryCreateConnection(NewNodeOutputPin, DownstreamInputPin);
+		}
+	}
+
+	DataflowGraph->NotifyGraphChanged();
+
+	return NewEdNode;
+}
 
diff --git a/Engine/Plugins/Experimental/ChaosClothAssetEditor/Source/ChaosClothAssetEditorTools/Private/ChaosClothAsset/ClothTransferSkinWeightsTool.cpp b/Engine/Plugins/Experimental/ChaosClothAssetEditor/Source/ChaosClothAssetEditorTools/Private/ChaosClothAsset/ClothTransferSkinWeightsTool.cpp
--- a/Engine/Plugins/Experimental/ChaosClothAssetEditor/Source/ChaosClothAssetEditorTools/Private/ChaosClothAsset/ClothTransferSkinWeightsTool.cpp
+++ b/Engine/Plugins/Experimental/ChaosClothAssetEditor/Source/ChaosClothAssetEditorTools/Private/ChaosClothAsset/ClothTransferSkinWeightsTool.cpp
@@ -2,46 +2,28 @@
 
 #include "ChaosClothAsset/ClothTransferSkinWeightsTool.h"
 
-#include "ChaosClothAsset/ClothComponent.h"
-#include "ChaosClothAsset/ClothComponentToolTarget.h"
-#include "ChaosClothAsset/ClothPatternToDynamicMesh.h"
-#include "ChaosClothAsset/CollectionClothFacade.h"
-
-#include "BoneWeights.h"
 #include "SkeletalMeshAttributes.h"
-
-#include "ToolTargetManager.h"
 #include "MeshDescriptionToDynamicMesh.h"
-
 #include "Rendering/SkeletalMeshLODImporterData.h"
 #include "Rendering/SkeletalMeshModel.h"
-
-#include "DynamicMesh/DynamicMeshAttributeSet.h"
 #include "DynamicMesh/DynamicVertexSkinWeightsAttribute.h"
 #include "DynamicMesh/DynamicBoneAttribute.h"
-#include "DynamicMesh/DynamicMeshAABBTree3.h"
-#include "DynamicMesh/MeshTransforms.h"
-
 #include "Operations/TransferBoneWeights.h"
-
-#include "TransformTypes.h"
-
-#include "InteractiveTool.h"
-#include "InteractiveToolManager.h"
-
-#include "Engine/SkeletalMesh.h"
-#include "PreviewMesh.h"
-#include "DynamicMeshEditor.h"
-#include "Util/ColorConstants.h"
 #include "ToolSetupUtil.h"
 #include "ModelingToolTargetUtil.h"
-#include "TargetInterfaces/SkeletalMeshBackedTarget.h"
-#include "Components/SkeletalMeshComponent.h"
+#include "Engine/SkeletalMesh.h"
+#include "ChaosClothAsset/ClothEditorContextObject.h"
+#include "ContextObjectStore.h"
+#include "Dataflow/DataflowEdNode.h"
+#include "ChaosClothAsset/TransferSkinWeightsNode.h"
+#include "BaseGizmos/TransformGizmoUtil.h"
+#include "InteractiveToolObjects.h"
+#include "MeshOpPreviewHelpers.h"
 
-#define LOCTEXT_NAMESPACE "ClothSkinWeightRetargetingTool"
+#define LOCTEXT_NAMESPACE "ClothTransferSkinWeightsTool"
 
 
-namespace ClothTransferSkinWeightsToolHelpers
+namespace UE::Chaos::ClothAsset::Private
 {
 	void SkeletalMeshToDynamicMesh(USkeletalMesh* FromSkeletalMeshAsset, int32 SourceLODIdx, FDynamicMesh3& ToDynamicMesh)
 	{
@@ -71,147 +53,295 @@ namespace ClothTransferSkinWeightsToolHelpers
 	};
 
 
-	void ClothComponentToDynamicMesh(const UChaosClothComponent* ClothComponent, UE::Geometry::FDynamicMesh3& MeshOut)
+	class FClothTransferSkinWeightsOp : public UE::Geometry::FDynamicMeshOperator
 	{
-		using namespace UE::Chaos::ClothAsset;
+	public:
 
-		const UChaosClothAsset* ChaosClothAsset = ClothComponent->GetClothAsset();
-		if (!ChaosClothAsset)
+		FClothTransferSkinWeightsOp(FDynamicMesh3&& OriginalTargetMesh, const TSharedPtr<const FDynamicMesh3> SourceMesh, const FTransform& SourceMeshTransform) :
+			SourceMesh(SourceMesh),
+			SourceMeshTransform(SourceMeshTransform)
 		{
-			return;
+			*ResultMesh = MoveTemp(OriginalTargetMesh);
 		}
 
-		const FCollectionClothConstFacade ClothFacade(ChaosClothAsset->GetClothCollection());
-		constexpr int32 LodIndex = 0;
-		const FCollectionClothLodConstFacade ClothLodFacade = ClothFacade.GetLod(LodIndex);
-		constexpr bool bGet2DPattern = false;
-
-		UE::Geometry::FDynamicMeshEditor MeshEditor(&MeshOut);
-		FClothPatternToDynamicMesh Converter;
+	private:
 
-		for (int32 PatternIndex = 0; PatternIndex < ClothLodFacade.GetNumPatterns(); ++PatternIndex)
+		// FDynamicMeshOperator interface
+		virtual void CalculateResult(FProgressCancel* Progress) override
 		{
-			UE::Geometry::FDynamicMesh3 PatternMesh;
-			Converter.Convert(ChaosClothAsset, LodIndex, PatternIndex, bGet2DPattern, PatternMesh);
+			using namespace UE::Geometry;
+			using namespace UE::AnimationCore;
+
+			// Copy over bone attributes from the source mesh to the target/preview
+			ResultMesh->Attributes()->CopyBoneAttributes(*SourceMesh->Attributes());
 
-			UE::Geometry::FMeshIndexMappings IndexMaps;
-			MeshEditor.AppendMesh(&PatternMesh, IndexMaps);
+			// Do the transfer
+			FTransferBoneWeights TransferBoneWeights(SourceMesh.Get(), FSkeletalMeshAttributes::DefaultSkinWeightProfileName);
+			bool bComputeOK = false;
+			if (TransferBoneWeights.Validate() == EOperationValidationResult::Ok)
+			{
+				bComputeOK = TransferBoneWeights.Compute(*ResultMesh, FTransformSRT3d(SourceMeshTransform.Inverse()), FSkeletalMeshAttributes::DefaultSkinWeightProfileName);
+			}
+
+			FGeometryResult OpResult;
+			if (!bComputeOK)
+			{
+				OpResult.Result = EGeometryResultType::Failure;
+			}
+			else
+			{
+				OpResult.Result = EGeometryResultType::Success;
+			}
+			SetResultInfo(OpResult);
 		}
+
+		// Inputs
+		const TSharedPtr<const FDynamicMesh3> SourceMesh;
+		const FTransform SourceMeshTransform;		
 	};
+
+
 }
 
 
-// ------------------- Properties -------------------
+// ------------------- Builder -------------------
 
-void UClothTransferSkinWeightsToolActionProperties::PostAction(EClothTransferSkinWeightsToolActions Action)
+USingleSelectionMeshEditingTool* UClothTransferSkinWeightsToolBuilder::CreateNewTool(const FToolBuilderState& SceneState) const
 {
-	if (ParentTool.IsValid())
+	UClothTransferSkinWeightsTool* NewTool = NewObject<UClothTransferSkinWeightsTool>(SceneState.ToolManager);
+
+	if (UClothEditorContextObject* ContextObject = SceneState.ToolManager->GetContextObjectStore()->FindContext<UClothEditorContextObject>())
 	{
-		ParentTool->RequestAction(Action);
+		NewTool->SetClothEditorContextObject(ContextObject);
 	}
+
+	return NewTool;
 }
 
 
-// ------------------- Builder -------------------
+// ------------------- Tool -------------------
 
-const FToolTargetTypeRequirements& UClothTransferSkinWeightsToolBuilder::GetTargetRequirements() const
-{
-	static FToolTargetTypeRequirements TypeRequirements({ 
-		UPrimitiveComponentBackedTarget::StaticClass(),
-		UClothAssetBackedTarget::StaticClass()
-		});
-	return TypeRequirements;
-}
 
-bool UClothTransferSkinWeightsToolBuilder::CanBuildTool(const FToolBuilderState& SceneState) const
+void UClothTransferSkinWeightsTool::Setup()
 {
-	const bool ClothComponentSelected = (SceneState.TargetManager->CountSelectedAndTargetable(SceneState, GetTargetRequirements()) == 1);
-	
-	const static FToolTargetTypeRequirements SourceMeshRequirements(USkeletalMeshBackedTarget::StaticClass());
-	const bool SkeletalMeshComponentSelected = (SceneState.TargetManager->CountSelectedAndTargetable(SceneState, SourceMeshRequirements) == 1);
+	USingleSelectionMeshEditingTool::Setup();
 
-	return ClothComponentSelected && SkeletalMeshComponentSelected;
-}
+	ToolProperties = NewObject<UClothTransferSkinWeightsToolProperties>(this);
+	AddToolPropertySource(ToolProperties);
 
-USingleSelectionMeshEditingTool* UClothTransferSkinWeightsToolBuilder::CreateNewTool(const FToolBuilderState& SceneState) const
-{
-	UClothTransferSkinWeightsTool* NewTool = NewObject<UClothTransferSkinWeightsTool>(SceneState.ToolManager);
 
-	// Setting Target and World on the new tool is handled in USingleSelectionMeshEditingToolBuilder::InitializeNewTool
+	//
+	// Set up Preview mesh that will show the results of the computation
+	//
 
-	return NewTool;
-}
+	TargetClothPreview = NewObject<UMeshOpPreviewWithBackgroundCompute>(this);
+	TargetClothPreview->Setup(GetTargetWorld(), this);
+	ToolSetupUtil::ApplyRenderingConfigurationToPreview(TargetClothPreview->PreviewMesh, Target);
+	TargetClothPreview->ConfigureMaterials(ToolSetupUtil::GetVertexColorMaterial(GetToolManager()), ToolSetupUtil::GetDefaultWorkingMaterial(GetToolManager()));
 
-void UClothTransferSkinWeightsToolBuilder::PostSetupTool(UInteractiveTool* Tool, const FToolBuilderState& SceneState) const
-{
-	if (UClothTransferSkinWeightsTool* NewTool = Cast<UClothTransferSkinWeightsTool>(Tool))
+	// Mesh topology is not being changed 
+	TargetClothPreview->SetIsMeshTopologyConstant(true, EMeshRenderAttributeFlags::VertexColors);
+
+	TargetClothPreview->OnMeshUpdated.AddUObject(this, &UClothTransferSkinWeightsTool::PreviewMeshUpdatedCallback);
+	TargetClothPreview->SetVisibility(true);
+
+	//
+	// Source mesh (populated from the SkeletalMesh tool property)
+	//
+
+	SourceMeshParentActor = GetTargetWorld()->SpawnActor<AInternalToolFrameworkActor>();
+	SourceMeshComponent = NewObject<UDynamicMeshComponent>(SourceMeshParentActor);
+	ToolSetupUtil::ApplyRenderingConfigurationToPreview(SourceMeshComponent, nullptr);
+	SourceMeshParentActor->SetRootComponent(SourceMeshComponent);
+	SourceMeshComponent->RegisterComponent();
+
+	// Watch for property changes
+
+	ToolProperties->WatchProperty(ToolProperties->SourceMesh, [this](TObjectPtr<USkeletalMesh>) { UpdateSourceMesh(); });
+
+	ToolProperties->WatchProperty(ToolProperties->BoneName, [this](const FName&) 
+	{
+		if (TargetClothPreview)
+		{
+			TargetClothPreview->PreviewMesh->NotifyDeferredEditCompleted(UPreviewMesh::ERenderUpdateMode::FastUpdate, EMeshRenderAttributeFlags::VertexColors, false);
+		}
+	});
+	
+	ToolProperties->WatchProperty(ToolProperties->bHideSourceMesh, [this](bool bNewHideSourceMesh) 
+	{ 
+		if (SourceMeshComponent)
+		{
+			SourceMeshComponent->SetVisibility(!bNewHideSourceMesh);
+		}
+	});
+
+	ToolProperties->WatchProperty(ToolProperties->SourceMeshTransform, [this](const FTransform& NewTransform) 
 	{
-		for (UActorComponent* SelectedComponent : SceneState.SelectedComponents)
+		if (SourceMeshTransformProxy)
 		{
-			if (USkeletalMeshComponent* SkeletalMeshComponent = Cast<USkeletalMeshComponent>(SelectedComponent))
+			SourceMeshTransformProxy->SetTransform(NewTransform);
+		}
+	},
+		[](const FTransform& A, const FTransform& B) -> bool
+	{
+		return !A.Equals(B);
+	});
+
+	ToolProperties->WatchProperty(ToolProperties->SourceMeshLOD, [this](int32 NewLOD)
+	{
+		const bool bLODIsValid = !ToolProperties->SourceMesh || ToolProperties->SourceMesh->IsValidLODIndex(NewLOD);
+
+		if (bLODIsValid)
+		{
+			if (bHasInvalidLODWarning)
 			{
-				NewTool->ToolProperties->SourceMesh = SkeletalMeshComponent->GetSkeletalMeshAsset();
-				NewTool->ToolProperties->SourceMeshTransform = SkeletalMeshComponent->GetComponentTransform();
-				NewTool->SourceComponent = SkeletalMeshComponent;
-				break;
+				GetToolManager()->DisplayMessage(FText(), EToolMessageLevel::UserWarning);    // clear old warning
+				bHasInvalidLODWarning = false;
 			}
 		}
-	}
-}
-
-// ------------------- Tool -------------------
+		else
+		{
+			GetToolManager()->DisplayMessage(LOCTEXT("InvalidLODIndex", "Specified source mesh LOD is not valid"), EToolMessageLevel::UserWarning);
+			bHasInvalidLODWarning = true;
+		}
 
-void UClothTransferSkinWeightsTool::Setup()
-{
-	USingleSelectionMeshEditingTool::Setup();
+		UpdateSourceMesh();
+	});
 
-	UClothComponentToolTarget* ClothComponentToolTarget = Cast<UClothComponentToolTarget>(Target);
-	ClothComponent = ClothComponentToolTarget->GetClothComponent();
-	
-	ToolProperties = NewObject<UClothTransferSkinWeightsToolProperties>(this);
-	AddToolPropertySource(ToolProperties);
 
-	ActionProperties = NewObject<UClothTransferSkinWeightsToolActionProperties>(this);
-	ActionProperties->ParentTool = this;
-	AddToolPropertySource(ActionProperties);
+	UInteractiveGizmoManager* const GizmoManager = GetToolManager()->GetPairedGizmoManager();
+	ensure(GizmoManager);
+	SourceMeshTransformProxy = NewObject<UTransformProxy>(this);
+	ensure(SourceMeshTransformProxy);
+	SourceMeshTransformProxy->SetTransform(FTransform::Identity);
 
+	SourceMeshTransformProxy->OnTransformChanged.AddWeakLambda(this, [this](UTransformProxy*, FTransform NewTransform)
+	{
+		if (SourceMeshParentActor)
+		{
+			SourceMeshParentActor->SetActorTransform(NewTransform);
+		}
+		ToolProperties->SourceMeshTransform = NewTransform;
+	});
 
-	PreviewMesh = NewObject<UPreviewMesh>(this);
-	if (PreviewMesh == nullptr)
+	SourceMeshTransformProxy->OnEndTransformEdit.AddWeakLambda(this, [this](UTransformProxy* Proxy)
 	{
-		return;
-	}
-	PreviewMesh->CreateInWorld(GetTargetWorld(), FTransform::Identity);
-	ToolSetupUtil::ApplyRenderingConfigurationToPreview(PreviewMesh, nullptr);
+		// Recompute result after moving the transform gizmo
+		TargetClothPreview->InvalidateResult();
+	});
+
+	SourceMeshTransformGizmo = UE::TransformGizmoUtil::CreateCustomTransformGizmo(GizmoManager, ETransformGizmoSubElements::StandardTranslateRotate, this);
+	ensure(SourceMeshTransformGizmo);
 
+	SourceMeshTransformGizmo->SetActiveTarget(SourceMeshTransformProxy, GetToolManager());
+	SourceMeshTransformGizmo->SetVisibility(ToolProperties->SourceMesh != nullptr);
+	SourceMeshTransformGizmo->bUseContextCoordinateSystem = false;
+	SourceMeshTransformGizmo->bUseContextGizmoMode = false;
+	SourceMeshTransformGizmo->CurrentCoordinateSystem = EToolContextCoordinateSystem::Local;
 
-	PreviewMesh->SetTransform(ClothComponentToolTarget->GetWorldTransform());
+	UpdateSourceMesh();
+	SetPreviewMeshColorFunction();
 
-	ToolProperties->WatchProperty(ToolProperties->SourceMesh, [this](TObjectPtr<USkeletalMesh>) { UpdatePreviewMesh(); });
-	//ToolProperties->WatchProperty(ToolProperties->SourceMeshTransform, [this](const FTransform&) { UpdatePreviewMesh(); });
-	ToolProperties->WatchProperty(ToolProperties->BoneName, [this](const FName&) {  UpdatePreviewMeshColor(); });
-	ToolProperties->WatchProperty(ToolProperties->bHideSourceMesh, [this](bool) { UpdateSourceMeshRender(); });
+	UE::ToolTarget::HideSourceObject(Target);
 }
 
 void UClothTransferSkinWeightsTool::Shutdown(EToolShutdownType ShutdownType)
 {
 	USingleSelectionMeshEditingTool::Shutdown(ShutdownType);
 
-	if (PreviewMesh)
+	if (ShutdownType == EToolShutdownType::Accept)
+	{
+		AddNewNode();
+	}
+
+	if (SourceMeshTransformProxy)
+	{
+		SourceMeshTransformProxy->OnTransformChanged.RemoveAll(this);
+		SourceMeshTransformProxy->OnEndTransformEdit.RemoveAll(this);
+	}
+
+	if (TargetClothPreview)
+	{
+		TargetClothPreview->OnMeshUpdated.RemoveAll(this);
+		TargetClothPreview->Shutdown();
+		TargetClothPreview = nullptr;
+	}
+
+	if (SourceMeshComponent)
+	{
+		SourceMeshComponent->DestroyComponent();
+		SourceMeshComponent = nullptr;
+	}
+
+	if (SourceMeshParentActor)
 	{
-		PreviewMesh->Disconnect();
+		SourceMeshParentActor->Destroy();
+		SourceMeshParentActor = nullptr;
 	}
 
+	GetToolManager()->GetPairedGizmoManager()->DestroyAllGizmosByOwner(this);
+	SourceMeshTransformGizmo = nullptr;
+
 	UE::ToolTarget::ShowSourceObject(Target);
-	SourceComponent->SetVisibility(true);
 }
 
-void UClothTransferSkinWeightsTool::UpdatePreviewMeshColor()
+
+bool UClothTransferSkinWeightsTool::CanAccept() const
+{
+	return ToolProperties && ToolProperties->SourceMesh && ToolProperties->SourceMesh->IsValidLODIndex(ToolProperties->SourceMeshLOD);
+}
+
+void UClothTransferSkinWeightsTool::OnTick(float DeltaTime)
+{
+	if (TargetClothPreview)
+	{
+		TargetClothPreview->Tick(DeltaTime);
+	}
+}
+
+TUniquePtr<UE::Geometry::FDynamicMeshOperator> UClothTransferSkinWeightsTool::MakeNewOperator()
+{
+	UE::Geometry::FDynamicMesh3 TargetDynamicMesh = UE::ToolTarget::GetDynamicMeshCopy(Target, true);
+	const TSharedPtr<const UE::Geometry::FDynamicMesh3> SourceDynamicMesh = MakeShared<UE::Geometry::FDynamicMesh3>(*SourceMeshComponent->GetMesh());
+
+	TUniquePtr<UE::Chaos::ClothAsset::Private::FClothTransferSkinWeightsOp> TransferOp =
+		MakeUnique<UE::Chaos::ClothAsset::Private::FClothTransferSkinWeightsOp>(MoveTemp(TargetDynamicMesh), SourceDynamicMesh, ToolProperties->SourceMeshTransform);
+
+	return TransferOp;
+}
+
+void UClothTransferSkinWeightsTool::SetClothEditorContextObject(TObjectPtr<UClothEditorContextObject> InClothEditorContextObject)
+{
+	ClothEditorContextObject = InClothEditorContextObject;
+}
+
+void UClothTransferSkinWeightsTool::AddNewNode()
+{
+	checkf(ClothEditorContextObject, TEXT("Expected a non-null ClothEditorContextObject. Check UClothTransferSkinWeightsToolBuilder::CreateNewTool"));
+
+	const FName ConnectionType(FManagedArrayCollection::StaticType());
+	UEdGraphNode* const CurrentlySelectedNode = ClothEditorContextObject->GetSingleSelectedNodeWithOutputType(ConnectionType);
+	checkf(CurrentlySelectedNode, TEXT("No node with FManagedArrayCollection output is currently selected in the Dataflow graph"));
+
+	const FName NewNodeType(FChaosClothAssetTransferSkinWeightsNode::StaticType());
+	UEdGraphNode* const NewNode = ClothEditorContextObject->CreateAndConnectNewNode(NewNodeType, *CurrentlySelectedNode, ConnectionType);
+	checkf(NewNode, TEXT("Unexpectedly failed to create a new FChaosClothAssetTransferSkinWeightsNode"));
+
+	UDataflowEdNode* const NewDataflowEdNode = CastChecked<UDataflowEdNode>(NewNode);
+	const TSharedPtr<FDataflowNode> NewDataflowNode = NewDataflowEdNode->GetDataflowNode();
+	FChaosClothAssetTransferSkinWeightsNode* const NewTransferNode = NewDataflowNode->AsType<FChaosClothAssetTransferSkinWeightsNode>();
+
+	NewTransferNode->SkeletalMesh = ToolProperties->SourceMesh;
+	NewTransferNode->SkeletalMeshLOD = ToolProperties->SourceMeshLOD;
+	NewTransferNode->Transform = ToolProperties->SourceMeshTransform;
+}
+
+void UClothTransferSkinWeightsTool::SetPreviewMeshColorFunction()
 {
 	using namespace UE::AnimationCore;
 	using namespace UE::Geometry;
 
-	PreviewMesh->SetTriangleColorFunction([this](const FDynamicMesh3* Mesh, int TriangleID) -> FColor
+	TargetClothPreview->PreviewMesh->SetTriangleColorFunction([this](const FDynamicMesh3* Mesh, int TriangleID) -> FColor
 	{
 		const FName& CurrentBoneName = ToolProperties->BoneName;
 		if (!TargetMeshBoneNameToIndex.Contains(CurrentBoneName))
@@ -223,7 +353,7 @@ void UClothTransferSkinWeightsTool::UpdatePreviewMeshColor()
 		const FIndex3i Tri = Mesh->GetTriangle(TriangleID);
 
 		const FName ProfileName = FSkeletalMeshAttributes::DefaultSkinWeightProfileName; // always use default profile for now, later this will be set by the user
-		FDynamicMeshVertexSkinWeightsAttribute* Attribute = PreviewMesh->GetPreviewDynamicMesh()->Attributes()->GetSkinWeightsAttribute(ProfileName);
+		FDynamicMeshVertexSkinWeightsAttribute* Attribute = Mesh->Attributes()->GetSkinWeightsAttribute(ProfileName);
 		if (Attribute == nullptr)
 		{
 			FLinearColor Lin(1.0f, 0.3f, 0.3f, 1.0f);
@@ -233,7 +363,7 @@ void UClothTransferSkinWeightsTool::UpdatePreviewMeshColor()
 		float AvgWeight = 0.0f;
 		for (int32 VID = 0; VID < 3; ++VID)
 		{
-			int32 VertexID = Tri[VID];
+			const int32 VertexID = Tri[VID];
 			FBoneWeights Data;
 			Attribute->GetValue(VertexID, Data);
 			for (FBoneWeight Wt : Data)
@@ -246,226 +376,95 @@ void UClothTransferSkinWeightsTool::UpdatePreviewMeshColor()
 		}
 
 		AvgWeight /= 3.0f;
-		FLinearColor Lin(AvgWeight, AvgWeight, AvgWeight, 1.0f);
+		const FLinearColor Lin(AvgWeight, AvgWeight, AvgWeight, 1.0f);
 		return Lin.ToFColor(/*bSRGB = */ true);
 	},
 	UPreviewMesh::ERenderUpdateMode::FullUpdate);
 }
 
-void UClothTransferSkinWeightsTool::UpdatePreviewMesh()
+void UClothTransferSkinWeightsTool::UpdateSourceMesh()
 {
-	using namespace UE::AnimationCore;
-	using namespace UE::Geometry;
-	
-	//TODO: for now, assume we are always transfering from LOD 0, but make this a parameter in the future...
-	constexpr int32 SourceLODIdx = 0; 
+	checkf(ToolProperties, TEXT("ToolProperties is expected to be non-null. Be sure to run Setup() on this tool when it is created."));
 
-	// User hasn't specified the source mesh in the UI
-	if (ToolProperties->SourceMesh == nullptr)
+	if (ToolProperties->SourceMesh && ToolProperties->SourceMesh->IsValidLODIndex(ToolProperties->SourceMeshLOD))
 	{
-		//TODO: Display error message
-		return;
+		// Set up source mesh (from the SkeletalMesh)
+		FDynamicMesh3 SourceDynamicMesh;
+		UE::Chaos::ClothAsset::Private::SkeletalMeshToDynamicMesh(ToolProperties->SourceMesh, ToolProperties->SourceMeshLOD, SourceDynamicMesh);
+		SourceMeshComponent->SetMesh(MoveTemp(SourceDynamicMesh));
+
+		checkf(SourceMeshComponent, TEXT("Source mesh specified in the Tool Properties, but no SourceMesh exists"));
+		SourceMeshParentActor->SetActorTransform(ToolProperties->SourceMeshTransform);
+		SourceMeshComponent->SetMaterial(0, ToolSetupUtil::GetTransparentSculptMaterial(GetToolManager(), FLinearColor::Red, 0.4, true));
+		SourceMeshComponent->SetVisibility(!ToolProperties->bHideSourceMesh);
+
+		SourceMeshTransformGizmo->SetNewGizmoTransform(SourceMeshParentActor->GetActorTransform());
+		SourceMeshTransformGizmo->SetVisibility(!ToolProperties->bHideSourceMesh);
+		SourceMeshTransformGizmo->ActiveGizmoMode = EToolContextTransformGizmoMode::Combined;
 	}
-
-	// Convert source Skeletal Mesh to Dynamic Mesh
-	FDynamicMesh3 SourceDynamicMesh;
-	ClothTransferSkinWeightsToolHelpers::SkeletalMeshToDynamicMesh(ToolProperties->SourceMesh, SourceLODIdx, SourceDynamicMesh);
-	MeshTransforms::ApplyTransform(SourceDynamicMesh, ToolProperties->SourceMeshTransform, true);
-
-	// Convert target ClothComponent to Dynamic Mesh
-	UE::Geometry::FDynamicMesh3 TargetDynamicMesh;
-	TargetDynamicMesh.EnableAttributes();
-	TargetDynamicMesh.Attributes()->AttachSkinWeightsAttribute(FSkeletalMeshAttributes::DefaultSkinWeightProfileName, new FDynamicMeshVertexSkinWeightsAttribute(&TargetDynamicMesh));
-	ClothTransferSkinWeightsToolHelpers::ClothComponentToDynamicMesh(ClothComponent, TargetDynamicMesh);
-
-	// Enable bone attribute for the target mesh and copy over the bone data from the cloth asset
-	// TODO: Once we support skin weights and bones in ToDynamicMesh converter we can remove this logic
-	const UChaosClothAsset* TargetClothAsset = ClothComponent->GetClothAsset();
-	const FReferenceSkeleton& TargetRefSkeleton = TargetClothAsset->GetRefSkeleton();
-	TargetDynamicMesh.Attributes()->EnableBones(TargetRefSkeleton.GetRawBoneNum());
-	UE::Geometry::FDynamicMeshBoneNameAttribute* BoneNameAttrib = TargetDynamicMesh.Attributes()->GetBoneNames();
-	this->TargetMeshBoneNameToIndex.Reset();
-	for (int BoneID = 0; BoneID < TargetRefSkeleton.GetRawBoneNum(); ++BoneID)
-	{	
-		const FName& BoneName = TargetRefSkeleton.GetRawRefBoneInfo()[BoneID].Name;
-		BoneNameAttrib->SetValue(BoneID, BoneName);
-		this->TargetMeshBoneNameToIndex.Add(BoneName, BoneID);
-	}
-
-	// Do the transfer
-	FTransferBoneWeights TransferBoneWeights(&SourceDynamicMesh, FSkeletalMeshAttributes::DefaultSkinWeightProfileName);
-	const FTransformSRT3d TargetToWorld = ClothComponent->GetComponentTransform();
-	if (TransferBoneWeights.Validate() == EOperationValidationResult::Ok)
+	else
 	{
-		TransferBoneWeights.Compute(TargetDynamicMesh, TargetToWorld, FSkeletalMeshAttributes::DefaultSkinWeightProfileName);
-	}
-
-	// Get set of bone indices used in the target mesh
-	TMap<FName, FBoneIndexType> UsedBoneNames;
-
-	const TMap<FName, TUniquePtr<FDynamicMeshVertexSkinWeightsAttribute>>& WeightLayers = TargetDynamicMesh.Attributes()->GetSkinWeightsAttributes();
-	for (const TPair<FName, TUniquePtr<FDynamicMeshVertexSkinWeightsAttribute>>& LayerPair : WeightLayers)
-	{
-		const FDynamicMeshVertexSkinWeightsAttribute* Layer = LayerPair.Value.Get();
-
-		for (int VertexID = 0; VertexID < TargetDynamicMesh.MaxVertexID(); ++VertexID)
-		{
-			if (TargetDynamicMesh.IsVertex(VertexID))
-			{
-				FBoneWeights Data;
-				Layer->GetValue(VertexID, Data);
-				for (FBoneWeight Wt : Data)
-				{
-					const FBoneIndexType BoneIndex = Wt.GetBoneIndex();
-					const FName& BoneName = TargetRefSkeleton.GetRawRefBoneInfo()[BoneIndex].Name;
-					UsedBoneNames.Add(BoneName, BoneIndex);
-				}
-			}
-		}
+		SourceMeshComponent->SetVisibility(false);
+		SourceMeshTransformGizmo->SetVisibility(false);
 	}
 
-	// Update list of bone names in the Properties panel
-	UsedBoneNames.ValueSort([](int16 A, int16 B) { return A < B; });
-	UsedBoneNames.GetKeys(ToolProperties->BoneNameList);
-
-	// Update the preview mesh
-	PreviewMesh->UpdatePreview(&TargetDynamicMesh);
-	PreviewMesh->SetMaterial(ToolSetupUtil::GetDefaultSculptMaterial(GetToolManager()));
-	PreviewMesh->SetOverrideRenderMaterial(ToolSetupUtil::GetSelectionMaterial(GetToolManager()));
-
-	UpdatePreviewMeshColor();
-
-	PreviewMesh->SetTransform(TargetToWorld);
-	PreviewMesh->SetVisible(true);
-
-	UE::ToolTarget::HideSourceObject(Target);
+	TargetClothPreview->InvalidateResult();
 }
 
-
-void UClothTransferSkinWeightsTool::UpdateSourceMeshRender()
+void UClothTransferSkinWeightsTool::PreviewMeshUpdatedCallback(UMeshOpPreviewWithBackgroundCompute* Preview)
 {
-	if (ToolProperties && SourceComponent)
-	{
-		SourceComponent->SetVisibility(!ToolProperties->bHideSourceMesh);
-	}
-}
+	constexpr bool bProcessOnlyIfValid = true;
 
-void UClothTransferSkinWeightsTool::TransferWeights()
-{
-	using namespace UE::AnimationCore;
-	using namespace UE::Geometry;
-	using namespace UE::Chaos::ClothAsset;
-	
-	//TODO: for now, assume we are always transfering from LOD 0, but make this a parameter in the future...
-	constexpr int32 SourceLODIdx = 0; 
-	
-	// User hasn't specified the source mesh in the UI
-	if (ToolProperties->SourceMesh == nullptr) 
+	Preview->ProcessCurrentMesh([&](const FDynamicMesh3& ResultMesh)
 	{
-		//TODO: Display error message
-		return;
-	}
+		using namespace UE::Geometry;
+		using namespace UE::AnimationCore;
 
-	// Convert source Skeletal Mesh to Dynamic Mesh
-	FDynamicMesh3 SourceDynamicMesh;
-	ClothTransferSkinWeightsToolHelpers::SkeletalMeshToDynamicMesh(ToolProperties->SourceMesh, SourceLODIdx, SourceDynamicMesh);
-	MeshTransforms::ApplyTransform(SourceDynamicMesh, ToolProperties->SourceMeshTransform, true);
+		ToolProperties->BoneNameList.Reset();
+		TargetMeshBoneNameToIndex.Reset();
 
-	UChaosClothAsset* TargetClothAsset = ClothComponent->GetClothAsset(); 
+		// Rebuild the set of selectable bone names, and the name -> index map
+		// TODO: Do we maybe want to do this in the background op and then copy the final results back to the member variables?
 
-	// Compute bone index mappings
-	TMap<FName, FBoneIndexType> TargetBoneToIndex;
-	const FReferenceSkeleton& TargetRefSkeleton = TargetClothAsset->GetRefSkeleton();
-	for (int32 Index = 0; Index < TargetRefSkeleton.GetRawBoneNum(); ++Index)
-	{
-		TargetBoneToIndex.Add(TargetRefSkeleton.GetRawRefBoneInfo()[Index].Name, Index);
-	}
-
-	// Setup bone weight transfer operator
-	FTransferBoneWeights TransferBoneWeights(&SourceDynamicMesh, FSkeletalMeshAttributes::DefaultSkinWeightProfileName); 
-	if (TransferBoneWeights.Validate() != EOperationValidationResult::Ok)
-	{
-		//TODO: Display error message
-		return;
-	}
-
-	FCollectionClothFacade ClothFacade(TargetClothAsset->GetClothCollection());
-    FTransformSRT3d TargetToWorld = ClothComponent->GetComponentTransform();
-
-	// Iterate over the LODs and transfer the bone weights from the source Skeletal mesh to the Cloth asset
-	for (int TargetLODIdx = 0; TargetLODIdx < ClothFacade.GetNumLods(); ++TargetLODIdx) 
-	{
-		FCollectionClothLodFacade ClothLodFacade = ClothFacade.GetLod(TargetLODIdx);
-
-		// Cloth collection data arrays we are writing to
-		TArrayView<int32> SimNumBoneInfluences = ClothLodFacade.GetSimNumBoneInfluences();
-		TArrayView<TArray<int32>> SimBoneIndices = ClothLodFacade.GetSimBoneIndices();
-		TArrayView<TArray<float>> SimBoneWeights = ClothLodFacade.GetSimBoneWeights();
-
-		TArrayView<int32> RenderNumBoneInfluences = ClothLodFacade.GetRenderNumBoneInfluences();
-		TArrayView<TArray<int32>> RenderBoneIndices = ClothLodFacade.GetRenderBoneIndices();
-		TArrayView<TArray<float>> RenderBoneWeights = ClothLodFacade.GetRenderBoneWeights();
-
-		const TArrayView<FVector3f> SimPositions =  ClothLodFacade.GetSimRestPosition();
-		
-		checkSlow(SimPositions.Num() == SimBoneIndices.Num());
-		
-		const int32 NumVert = ClothLodFacade.GetNumSimVertices();
-		constexpr bool bUseParallel = true; 
-
-		// Iterate over each vertex and write the data from FBoneWeights into cloth collection managed arrays
-		ParallelFor(NumVert, [&](int32 VertexID)
+		if (ResultMesh.Attributes() && ResultMesh.Attributes()->GetBoneNames() && ResultMesh.Attributes()->GetBoneParentIndices())
 		{
-			const FVector3f Pos = SimPositions[VertexID];
-			const FVector3d PosD = FVector3d((double)Pos[0], (double)Pos[1], (double)Pos[2]);
-			
-			UE::AnimationCore::FBoneWeights BoneWeights;
-			TransferBoneWeights.Compute(PosD, TargetToWorld, BoneWeights, &TargetBoneToIndex);
-			
-			const int32 NumBones = BoneWeights.Num();
-			
-			SimNumBoneInfluences[VertexID] = NumBones;
-			SimBoneIndices[VertexID].SetNum(NumBones);
-			SimBoneWeights[VertexID].SetNum(NumBones);
-
-			RenderNumBoneInfluences[VertexID] = NumBones;
-			RenderBoneIndices[VertexID].SetNum(NumBones);
-			RenderBoneWeights[VertexID].SetNum(NumBones);
-
-			for (int BoneIdx = 0; BoneIdx < NumBones; ++BoneIdx) 
+			// Get set of bone indices used in the target mesh
+			const TArray<FName>& TargetBoneNames = ResultMesh.Attributes()->GetBoneNames()->GetAttribValues();
+			const TArray<int32>& TargetBoneIndices = ResultMesh.Attributes()->GetBoneParentIndices()->GetAttribValues();
+
+			const TMap<FName, TUniquePtr<FDynamicMeshVertexSkinWeightsAttribute>>& WeightLayers = ResultMesh.Attributes()->GetSkinWeightsAttributes();
+			TMap<FName, FBoneIndexType> UsedBoneNames;
+			for (const TPair<FName, TUniquePtr<FDynamicMeshVertexSkinWeightsAttribute>>& LayerPair : WeightLayers)
 			{
-				SimBoneIndices[VertexID][BoneIdx] = BoneWeights[BoneIdx].GetBoneIndex();
-				SimBoneWeights[VertexID][BoneIdx] = BoneWeights[BoneIdx].GetWeight();
+				const FDynamicMeshVertexSkinWeightsAttribute* Layer = LayerPair.Value.Get();
 
-				RenderBoneIndices[VertexID][BoneIdx] = BoneWeights[BoneIdx].GetBoneIndex();
-				RenderBoneWeights[VertexID][BoneIdx] = BoneWeights[BoneIdx].GetWeight();
+				for (const int VertexID : ResultMesh.VertexIndicesItr())
+				{
+					FBoneWeights Data;
+					Layer->GetValue(VertexID, Data);
+					for (FBoneWeight Wt : Data)
+					{
+						const FBoneIndexType BoneIndex = Wt.GetBoneIndex();
+						const FName& BoneName = TargetBoneNames[BoneIndex];
+						UsedBoneNames.Add(BoneName, BoneIndex);
+					}
+				}
 			}
 
-		}, bUseParallel ? EParallelForFlags::None : EParallelForFlags::ForceSingleThread);
-	}
-}
-
+			// Build the bone name -> index map
+			for (int32 BoneID = 0; BoneID < TargetBoneNames.Num(); ++BoneID)
+			{
+				const FName& BoneName = TargetBoneNames[BoneID];
+				const int32 BoneIndex = TargetBoneIndices[BoneID];
+				TargetMeshBoneNameToIndex.Add(BoneName, BoneID);
+			}
 
-void UClothTransferSkinWeightsTool::OnTick(float DeltaTime)
-{
-	if (PendingAction != EClothTransferSkinWeightsToolActions::NoAction)
-	{
-		if (PendingAction == EClothTransferSkinWeightsToolActions::Transfer)
-		{
-			TransferWeights();
+			// Update list of bone names in the Properties panel
+			UsedBoneNames.ValueSort([](int16 A, int16 B) { return A < B; });
+			UsedBoneNames.GetKeys(ToolProperties->BoneNameList);
 		}
-		PendingAction = EClothTransferSkinWeightsToolActions::NoAction;
-	}
-}
-
-
-void UClothTransferSkinWeightsTool::RequestAction(EClothTransferSkinWeightsToolActions ActionType)
-{
-	if (PendingAction != EClothTransferSkinWeightsToolActions::NoAction)
-	{
-		return;
-	}
-	PendingAction = ActionType;
+	}, 
+	bProcessOnlyIfValid);
 }
 
 #undef LOCTEXT_NAMESPACE
diff --git a/Engine/Plugins/Experimental/ChaosClothAssetEditor/Source/ChaosClothAssetEditorTools/Private/ChaosClothAsset/ClothWeightMapPaintTool.cpp b/Engine/Plugins/Experimental/ChaosClothAssetEditor/Source/ChaosClothAssetEditorTools/Private/ChaosClothAsset/ClothWeightMapPaintTool.cpp
--- a/Engine/Plugins/Experimental/ChaosClothAssetEditor/Source/ChaosClothAssetEditorTools/Private/ChaosClothAsset/ClothWeightMapPaintTool.cpp
+++ b/Engine/Plugins/Experimental/ChaosClothAssetEditor/Source/ChaosClothAssetEditorTools/Private/ChaosClothAsset/ClothWeightMapPaintTool.cpp
@@ -1320,44 +1320,15 @@ void UClothEditorWeightMapPaintTool::AddNewNode(const FString& NewMapName)
 
 	// TODO: If we are operating on a single pattern, expand the weight map to cover the whole LOD
 
-	TSharedPtr<SDataflowGraphEditor> GraphEditor = ClothEditorContextObject->GetDataflowGraphEditor().Pin();
+	const FName ConnectionType(FManagedArrayCollection::StaticType());
+	UEdGraphNode* const CurrentlySelectedNode = ClothEditorContextObject->GetSingleSelectedNodeWithOutputType(ConnectionType);
+	checkf(CurrentlySelectedNode, TEXT("No node with FManagedArrayCollection output is currently selected in the Dataflow graph"));
 
-	//
-	// We want to turn this:
-	// 
-	// [SelectedNode] ----> [ExistingDownstreamNode(s)]
-	// 
-	// to this:
-	//
-	// [SelectedNode] ----> [NewNode] ----> [ExistingDownstreamNode(s)]
-	//
+	const FName NewNodeType(FChaosClothAssetAddWeightMapNode::StaticType());
+	UEdGraphNode* const NewNode = ClothEditorContextObject->CreateAndConnectNewNode(NewNodeType, *CurrentlySelectedNode, ConnectionType);
+	checkf(NewNode, TEXT("Unexpectedly failed to create a new FChaosClothAssetAddWeightMapNode"));
 
-	// First find the Collection output of the currently selected node, and any nodes it's connected to
-
-	const UEdGraphNode* const SingleSelectedNode = GraphEditor->GetSingleSelectedNode();
-	checkf(SingleSelectedNode, TEXT("There should be a single selected node in the Dataflow graph"));
-
-	const UDataflowEdNode* const SelectedDataflowEdNode = CastChecked<UDataflowEdNode>(SingleSelectedNode);
-	const TSharedPtr<const FDataflowNode> DataflowNode = SelectedDataflowEdNode->GetDataflowNode();
-	UEdGraphPin* SelectedNodeOutputPin = nullptr;
-	TArray<UEdGraphPin*> ExistingNodeInputPins;
-	for (const FDataflowOutput* Output : DataflowNode->GetOutputs())
-	{
-		if (Output->GetType() == FName("FManagedArrayCollection"))
-		{
-			SelectedNodeOutputPin = SelectedDataflowEdNode->FindPin(*Output->GetName().ToString(), EGPD_Output);
-			ExistingNodeInputPins = SelectedNodeOutputPin->LinkedTo;
-			break;
-		}
-	}
-
-	// Now add the new node and re-wire the graph
-
-	const TObjectPtr<UDataflow> Dataflow = ClothEditorContextObject->GetDataflowGraph();
-	const TSharedPtr<FAssetSchemaAction_Dataflow_CreateNode_DataflowEdNode> NodeAction = FAssetSchemaAction_Dataflow_CreateNode_DataflowEdNode::CreateAction(Dataflow, "FChaosClothAssetAddWeightMapNode");
-
-	UEdGraphNode* const NewEdNode = NodeAction->PerformAction(Dataflow, nullptr, GraphEditor->GetPasteLocation(), true);
-	UDataflowEdNode* const NewDataflowEdNode = CastChecked<UDataflowEdNode>(NewEdNode);
+	UDataflowEdNode* const NewDataflowEdNode = CastChecked<UDataflowEdNode>(NewNode);
 	const TSharedPtr<FDataflowNode> NewDataflowNode = NewDataflowEdNode->GetDataflowNode();
 	FChaosClothAssetAddWeightMapNode* const NewWeightMapNode = NewDataflowNode->AsType<FChaosClothAssetAddWeightMapNode>();
 
@@ -1365,24 +1336,6 @@ void UClothEditorWeightMapPaintTool::AddNewNode(const FString& NewMapName)
 	NewWeightMapNode->Name = AddWeightMapProperties->Name;
 	// NewWeightMapNode->SetName(FName(AddWeightMapProperties->NewMapName));   TODO: This doesn't seem to work to set the Node name
 
-	// Re-wire the graph
-	FDataflowInput* NewNodeCollectionInput = NewWeightMapNode->FindInput(&(NewWeightMapNode->Collection));
-	FDataflowOutput* NewNodeCollectionOutput = NewWeightMapNode->FindOutput(&(NewWeightMapNode->Collection));
-	UEdGraphPin* NewNodeInputPin = NewDataflowEdNode->FindPin(*NewNodeCollectionInput->GetName().ToString(), EGPD_Input);
-	UEdGraphPin* NewNodeOutputPin = NewDataflowEdNode->FindPin(*NewNodeCollectionInput->GetName().ToString(), EGPD_Output);
-
-	bool bWasAnythingChanged = Dataflow->GetSchema()->TryCreateConnection(SelectedNodeOutputPin, NewNodeInputPin);
-
-	for (UEdGraphPin* DownstreamInputPin : ExistingNodeInputPins)
-	{
-		bWasAnythingChanged |= Dataflow->GetSchema()->TryCreateConnection(NewNodeOutputPin, DownstreamInputPin);
-	}
-
-	if (bWasAnythingChanged)
-	{
-		Dataflow->NotifyGraphChanged();
-	}
-
 }
 
 
diff --git a/Engine/Plugins/Experimental/ChaosClothAssetEditor/Source/ChaosClothAssetEditorTools/Public/ChaosClothAsset/ClothEditorContextObject.h b/Engine/Plugins/Experimental/ChaosClothAssetEditor/Source/ChaosClothAssetEditorTools/Public/ChaosClothAsset/ClothEditorContextObject.h
--- a/Engine/Plugins/Experimental/ChaosClothAssetEditor/Source/ChaosClothAssetEditorTools/Public/ChaosClothAsset/ClothEditorContextObject.h
+++ b/Engine/Plugins/Experimental/ChaosClothAssetEditor/Source/ChaosClothAssetEditorTools/Public/ChaosClothAsset/ClothEditorContextObject.h
@@ -7,6 +7,7 @@
 
 class SDataflowGraphEditor;
 class UDataflow;
+class UEdGraphNode;
 
 UCLASS()
 class CHAOSCLOTHASSETEDITORTOOLS_API UClothEditorContextObject : public UObject
@@ -23,6 +24,45 @@ public:
 	TObjectPtr<UDataflow> GetDataflowGraph();
 	const TObjectPtr<const UDataflow> GetDataflowGraph() const;
 
+	/** 
+	 * Return the single selected node in the Dataflow Graph Editor, or nullptr if multiple or no nodes are selected 
+	 */
+	UEdGraphNode* GetSingleSelectedNode() const;
+
+	/** 
+	* Return the single selected node in the Dataflow Graph Editor only if it has an output of the specified type
+	* If there is not a single node selected, or if it does not have the specified output, return null 
+	*/
+	UEdGraphNode* GetSingleSelectedNodeWithOutputType(const FName& SelectedNodeOutputTypeName) const;
+
+	/** 
+	 * Create a node with the specified type in the graph 
+	 */
+	UEdGraphNode* CreateNewNode(const FName& NewNodeTypeName);
+
+
+	/** Create a node with the specified type, then connect it to the output of the specified UpstreamNode.
+	* If the specified output of the upstream node is already connected to another node downstream, we first break
+	* that connecttion, then insert the new node along the previous connection.
+	* We want to turn this:
+	* 
+	* [UpstreamNode] ----> [DownstreamNode(s)]
+	* 
+	* to this:
+	*
+	* [UpstreamNode] ----> [NewNode] ----> [DownstreamNode(s)]
+	*
+	* 
+	* @param NewNodeTypeName The type of node to create, by name
+	* @param UpstreamNode Node to connect the new node to
+	* @param ConnectionTypeName The type of output of the upstream node to connect our new node to
+	* @return The newly-created node
+	*/
+	UEdGraphNode* CreateAndConnectNewNode(
+		const FName& NewNodeTypeName,
+		UEdGraphNode& UpstreamNode,
+		const FName& ConnectionTypeName);
+
 private:
 
 	TWeakPtr<SDataflowGraphEditor> DataflowGraphEditor;
@@ -31,3 +71,5 @@ private:
 	TObjectPtr<UDataflow> DataflowGraph = nullptr;
 
 };
+
+
diff --git a/Engine/Plugins/Experimental/ChaosClothAssetEditor/Source/ChaosClothAssetEditorTools/Public/ChaosClothAsset/ClothTransferSkinWeightsTool.h b/Engine/Plugins/Experimental/ChaosClothAssetEditor/Source/ChaosClothAssetEditorTools/Public/ChaosClothAsset/ClothTransferSkinWeightsTool.h
--- a/Engine/Plugins/Experimental/ChaosClothAssetEditor/Source/ChaosClothAssetEditorTools/Public/ChaosClothAsset/ClothTransferSkinWeightsTool.h
+++ b/Engine/Plugins/Experimental/ChaosClothAssetEditor/Source/ChaosClothAssetEditorTools/Public/ChaosClothAsset/ClothTransferSkinWeightsTool.h
@@ -2,20 +2,19 @@
 
 #pragma once
 
-#include "UObject/Object.h"
-#include "SingleSelectionTool.h"
-#include "InteractiveToolBuilder.h"
 #include "BaseTools/SingleSelectionMeshEditingTool.h"
 #include "BoneIndices.h"
-#include "TransformTypes.h"
-
+#include "ModelingOperators.h"
 #include "ClothTransferSkinWeightsTool.generated.h"
 
-class UChaosClothComponent;
 class UClothTransferSkinWeightsTool;
 class USkeletalMesh;
-class UPreviewMesh;
-class USkeletalMeshComponent;
+class UClothEditorContextObject;
+class UTransformProxy;
+class UCombinedTransformGizmo;
+class UMeshOpPreviewWithBackgroundCompute;
+class AInternalToolFrameworkActor;
+class UDynamicMeshComponent;
 
 UCLASS()
 class CHAOSCLOTHASSETEDITORTOOLS_API UClothTransferSkinWeightsToolProperties : public UInteractiveToolPropertySet
@@ -29,6 +28,9 @@ public:
 	UPROPERTY(EditAnywhere, Category = Source)
 	FTransform SourceMeshTransform;
 
+	UPROPERTY(EditAnywhere, Category = Source)
+	int32 SourceMeshLOD = 0;
+
 	UPROPERTY(EditAnywhere, Category = Source)
 	bool bHideSourceMesh = false;
 
@@ -47,97 +49,75 @@ public:
 
 };
 
-UENUM()
-enum class EClothTransferSkinWeightsToolActions
-{
-	NoAction,
-	Transfer
-};
-
-UCLASS()
-class CHAOSCLOTHASSETEDITORTOOLS_API UClothTransferSkinWeightsToolActionProperties : public UObject
-{
-	GENERATED_BODY()
-public:
-
-	TWeakObjectPtr<UClothTransferSkinWeightsTool> ParentTool;
-
-	void Initialize(UClothTransferSkinWeightsTool* ParentToolIn) { ParentTool = ParentToolIn; }
-
-	void PostAction(EClothTransferSkinWeightsToolActions Action);
-
-	UFUNCTION(CallInEditor, Category = Actions, meta = (DisplayName = "Transfer weights", DisplayPriority = 1))
-	void TransferWeights()
-	{
-		PostAction(EClothTransferSkinWeightsToolActions::Transfer);
-	}
-
-};
-
 
 UCLASS()
 class CHAOSCLOTHASSETEDITORTOOLS_API UClothTransferSkinWeightsToolBuilder : public USingleSelectionMeshEditingToolBuilder
 {
 	GENERATED_BODY()
 
-protected:
-
-	virtual const FToolTargetTypeRequirements& GetTargetRequirements() const override;
-
-public:
+private:
 
-	virtual bool CanBuildTool(const FToolBuilderState& SceneState) const override;
-	
 	virtual USingleSelectionMeshEditingTool* CreateNewTool(const FToolBuilderState& SceneState) const override;	
 
-	virtual void PostSetupTool(UInteractiveTool* Tool, const FToolBuilderState& SceneState) const override;
 };
 
-
 UCLASS()
-class CHAOSCLOTHASSETEDITORTOOLS_API UClothTransferSkinWeightsTool : public USingleSelectionMeshEditingTool
+class CHAOSCLOTHASSETEDITORTOOLS_API UClothTransferSkinWeightsTool : public USingleSelectionMeshEditingTool, public UE::Geometry::IDynamicMeshOperatorFactory
 {
 	GENERATED_BODY()
-public:
+
+private:
+
+	friend class UClothTransferSkinWeightsToolBuilder;
 
 	// UInteractiveTool
 	virtual void Setup() override;
 	virtual void Shutdown(EToolShutdownType ShutdownType) override;
+	virtual bool HasAccept() const override { return true; }
+	virtual bool HasCancel() const override { return true; }
+	virtual bool CanAccept() const override;
 	virtual void OnTick(float DeltaTime) override;
 
-private:
+	// IDynamicMeshOperatorFactory
+	virtual TUniquePtr<UE::Geometry::FDynamicMeshOperator> MakeNewOperator() override;
+	
+	void SetClothEditorContextObject(TObjectPtr<UClothEditorContextObject> InClothEditorContextObject);
 
-	friend class UClothTransferSkinWeightsToolActionProperties;
-	friend class UClothTransferSkinWeightsToolBuilder;
+	void AddNewNode();
 
-	UPROPERTY()
-	TObjectPtr<UClothTransferSkinWeightsToolProperties> ToolProperties;
+	void SetPreviewMeshColorFunction();
 
-	UPROPERTY()
-	TObjectPtr<UClothTransferSkinWeightsToolActionProperties> ActionProperties;
+	void UpdateSourceMesh();
 
-	EClothTransferSkinWeightsToolActions PendingAction = EClothTransferSkinWeightsToolActions::NoAction;
+	void PreviewMeshUpdatedCallback(UMeshOpPreviewWithBackgroundCompute* Preview);
 
-	UPROPERTY()
-	TObjectPtr<const UChaosClothComponent> ClothComponent;
 
-	UPROPERTY()
-	TObjectPtr<USkeletalMeshComponent> SourceComponent;
+	UPROPERTY(Transient)
+	TObjectPtr<UClothTransferSkinWeightsToolProperties> ToolProperties;
 
+	UPROPERTY(Transient)
+	TObjectPtr<UClothEditorContextObject> ClothEditorContextObject;
 
-	void RequestAction(EClothTransferSkinWeightsToolActions ActionType);
+	UPROPERTY(Transient)
+	TObjectPtr<UMeshOpPreviewWithBackgroundCompute> TargetClothPreview;
 
-	void TransferWeights();
+	UPROPERTY(Transient)
+	TObjectPtr<AInternalToolFrameworkActor> SourceMeshParentActor;
 
-	UPROPERTY()
-	TObjectPtr<UPreviewMesh> PreviewMesh;
+	UPROPERTY(Transient)
+	TObjectPtr<UDynamicMeshComponent> SourceMeshComponent;
 
-	TMap<FName, FBoneIndexType> TargetMeshBoneNameToIndex;
+	// Source mesh transform gizmo support
+	UPROPERTY(Transient)
+	TObjectPtr<UTransformProxy> SourceMeshTransformProxy;
 
-	void UpdatePreviewMeshColor();
-	void UpdatePreviewMesh();
+	UPROPERTY(Transient)
+	TObjectPtr<UCombinedTransformGizmo> SourceMeshTransformGizmo;
+
+	// Used to lookup the index of the currently selected-by-name bone
+	TMap<FName, FBoneIndexType> TargetMeshBoneNameToIndex;
 
-	void UpdateSourceMeshRender();
+	bool bHasInvalidLODWarning = false;
 };
 
 
