Cloth Editor, TransferSkinWeightsTool:
- hook up gizmo and transform properties using FTransformGizmoDataBinder
- dsiable numeric input when no source mesh is selected
- display warning in the tool if transfer fails for some reason

diff --git a/Engine/Plugins/Experimental/ChaosClothAssetEditor/Source/ChaosClothAssetEditorTools/Private/ChaosClothAsset/ClothTransferSkinWeightsTool.cpp b/Engine/Plugins/Experimental/ChaosClothAssetEditor/Source/ChaosClothAssetEditorTools/Private/ChaosClothAsset/ClothTransferSkinWeightsTool.cpp
--- a/Engine/Plugins/Experimental/ChaosClothAssetEditor/Source/ChaosClothAssetEditorTools/Private/ChaosClothAsset/ClothTransferSkinWeightsTool.cpp
+++ b/Engine/Plugins/Experimental/ChaosClothAssetEditor/Source/ChaosClothAssetEditorTools/Private/ChaosClothAsset/ClothTransferSkinWeightsTool.cpp
@@ -142,7 +142,13 @@ void UClothTransferSkinWeightsTool::Setup()
 	// Mesh topology is not being changed 
 	TargetClothPreview->SetIsMeshTopologyConstant(true, EMeshRenderAttributeFlags::VertexColors);
 
+	TargetClothPreview->OnOpCompleted.AddUObject(this, &UClothTransferSkinWeightsTool::OpFinishedCallback);
 	TargetClothPreview->OnMeshUpdated.AddUObject(this, &UClothTransferSkinWeightsTool::PreviewMeshUpdatedCallback);
+
+	// Set the initial preview mesh before any computation runs
+	UE::Geometry::FDynamicMesh3 InitialPreviewMesh = UE::ToolTarget::GetDynamicMeshCopy(Target, true);
+	TargetClothPreview->PreviewMesh->UpdatePreview(MoveTemp(InitialPreviewMesh));
+
 	TargetClothPreview->SetVisibility(true);
 
 	//
@@ -161,30 +167,14 @@ void UClothTransferSkinWeightsTool::Setup()
 
 	ToolProperties->WatchProperty(ToolProperties->BoneName, [this](const FName&) 
 	{
-		if (TargetClothPreview)
-		{
-			TargetClothPreview->PreviewMesh->NotifyDeferredEditCompleted(UPreviewMesh::ERenderUpdateMode::FastUpdate, EMeshRenderAttributeFlags::VertexColors, false);
-		}
+		check(TargetClothPreview);
+		TargetClothPreview->PreviewMesh->NotifyDeferredEditCompleted(UPreviewMesh::ERenderUpdateMode::FastUpdate, EMeshRenderAttributeFlags::VertexColors, false);
 	});
 	
 	ToolProperties->WatchProperty(ToolProperties->bHideSourceMesh, [this](bool bNewHideSourceMesh) 
 	{ 
-		if (SourceMeshComponent)
-		{
-			SourceMeshComponent->SetVisibility(!bNewHideSourceMesh);
-		}
-	});
-
-	ToolProperties->WatchProperty(ToolProperties->SourceMeshTransform, [this](const FTransform& NewTransform) 
-	{
-		if (SourceMeshTransformProxy)
-		{
-			SourceMeshTransformProxy->SetTransform(NewTransform);
-		}
-	},
-		[](const FTransform& A, const FTransform& B) -> bool
-	{
-		return !A.Equals(B);
+		check(SourceMeshComponent);
+		SourceMeshComponent->SetVisibility(!bNewHideSourceMesh);
 	});
 
 	ToolProperties->WatchProperty(ToolProperties->SourceMeshLOD, [this](int32 NewLOD)
@@ -209,6 +199,35 @@ void UClothTransferSkinWeightsTool::Setup()
 	});
 
 
+	//
+	// Transform/Gizmo/Proxy stuff
+	//
+
+	ToolProperties->WatchProperty(ToolProperties->SourceMeshTranslation, [this](const FVector3d& NewTranslation)
+	{
+		if (DataBinder)
+		{
+			DataBinder->UpdateAfterDataEdit();
+		}
+	});
+
+	ToolProperties->WatchProperty(ToolProperties->SourceMeshRotation, [this](const FVector3d& NewTranslation)
+	{
+		if (DataBinder)
+		{
+			DataBinder->UpdateAfterDataEdit();
+		}
+	});
+
+	ToolProperties->WatchProperty(ToolProperties->SourceMeshScale, [this](const FVector3d& NewTranslation)
+	{
+		if (DataBinder)
+		{
+			DataBinder->UpdateAfterDataEdit();
+		}
+	});
+
+
 	UInteractiveGizmoManager* const GizmoManager = GetToolManager()->GetPairedGizmoManager();
 	ensure(GizmoManager);
 	SourceMeshTransformProxy = NewObject<UTransformProxy>(this);
@@ -221,13 +240,12 @@ void UClothTransferSkinWeightsTool::Setup()
 		{
 			SourceMeshParentActor->SetActorTransform(NewTransform);
 		}
-		ToolProperties->SourceMeshTransform = NewTransform;
-	});
 
-	SourceMeshTransformProxy->OnEndTransformEdit.AddWeakLambda(this, [this](UTransformProxy* Proxy)
-	{
 		// Recompute result after moving the transform gizmo
-		TargetClothPreview->InvalidateResult();
+		if (ToolProperties->SourceMesh && ToolProperties->SourceMesh->IsValidLODIndex(ToolProperties->SourceMeshLOD))
+		{
+			TargetClothPreview->InvalidateResult();
+		}
 	});
 
 	SourceMeshTransformGizmo = UE::TransformGizmoUtil::CreateCustomTransformGizmo(GizmoManager, ETransformGizmoSubElements::StandardTranslateRotate, this);
@@ -239,6 +257,11 @@ void UClothTransferSkinWeightsTool::Setup()
 	SourceMeshTransformGizmo->bUseContextGizmoMode = false;
 	SourceMeshTransformGizmo->CurrentCoordinateSystem = EToolContextCoordinateSystem::Local;
 
+	DataBinder = MakeShared<FTransformGizmoDataBinder>();
+	DataBinder->InitializeBoundVectors(&ToolProperties->SourceMeshTranslation, &ToolProperties->SourceMeshRotation, &ToolProperties->SourceMeshScale);
+	DataBinder->BindToInitializedGizmo(SourceMeshTransformGizmo, SourceMeshTransformProxy);
+
+
 	UpdateSourceMesh();
 	SetPreviewMeshColorFunction();
 
@@ -301,11 +324,15 @@ void UClothTransferSkinWeightsTool::OnTick(float DeltaTime)
 
 TUniquePtr<UE::Geometry::FDynamicMeshOperator> UClothTransferSkinWeightsTool::MakeNewOperator()
 {
+	checkf(ToolProperties->SourceMesh, TEXT("Expected non-null SourceMesh before starting Operator"));
+	checkf(ToolProperties->SourceMesh->IsValidLODIndex(ToolProperties->SourceMeshLOD), TEXT("Source Mesh LOD level should be validated before starting Operator"));
+
 	UE::Geometry::FDynamicMesh3 TargetDynamicMesh = UE::ToolTarget::GetDynamicMeshCopy(Target, true);
 	const TSharedPtr<const UE::Geometry::FDynamicMesh3> SourceDynamicMesh = MakeShared<UE::Geometry::FDynamicMesh3>(*SourceMeshComponent->GetMesh());
 
+	const FTransform SourceMeshTransform = SourceMeshTransformGizmo->GetGizmoTransform();
 	TUniquePtr<UE::Chaos::ClothAsset::Private::FClothTransferSkinWeightsOp> TransferOp =
-		MakeUnique<UE::Chaos::ClothAsset::Private::FClothTransferSkinWeightsOp>(MoveTemp(TargetDynamicMesh), SourceDynamicMesh, ToolProperties->SourceMeshTransform);
+		MakeUnique<UE::Chaos::ClothAsset::Private::FClothTransferSkinWeightsOp>(MoveTemp(TargetDynamicMesh), SourceDynamicMesh, SourceMeshTransform);
 
 	return TransferOp;
 }
@@ -333,7 +360,7 @@ void UClothTransferSkinWeightsTool::AddNewNode()
 
 	NewTransferNode->SkeletalMesh = ToolProperties->SourceMesh;
 	NewTransferNode->SkeletalMeshLOD = ToolProperties->SourceMeshLOD;
-	NewTransferNode->Transform = ToolProperties->SourceMeshTransform;
+	NewTransferNode->Transform = TransformFromProperties();
 }
 
 void UClothTransferSkinWeightsTool::SetPreviewMeshColorFunction()
@@ -382,6 +409,13 @@ void UClothTransferSkinWeightsTool::SetPreviewMeshColorFunction()
 	UPreviewMesh::ERenderUpdateMode::FullUpdate);
 }
 
+
+FTransform UClothTransferSkinWeightsTool::TransformFromProperties() const
+{
+	const FRotator Rotation = FRotator::MakeFromEuler(ToolProperties->SourceMeshRotation);
+	return FTransform(Rotation, ToolProperties->SourceMeshTranslation, ToolProperties->SourceMeshScale);
+}
+
 void UClothTransferSkinWeightsTool::UpdateSourceMesh()
 {
 	checkf(ToolProperties, TEXT("ToolProperties is expected to be non-null. Be sure to run Setup() on this tool when it is created."));
@@ -394,23 +428,43 @@ void UClothTransferSkinWeightsTool::UpdateSourceMesh()
 		SourceMeshComponent->SetMesh(MoveTemp(SourceDynamicMesh));
 
 		checkf(SourceMeshComponent, TEXT("Source mesh specified in the Tool Properties, but no SourceMesh exists"));
-		SourceMeshParentActor->SetActorTransform(ToolProperties->SourceMeshTransform);
+		
+		SourceMeshParentActor->SetActorTransform(TransformFromProperties());
 		SourceMeshComponent->SetMaterial(0, ToolSetupUtil::GetTransparentSculptMaterial(GetToolManager(), FLinearColor::Red, 0.4, true));
 		SourceMeshComponent->SetVisibility(!ToolProperties->bHideSourceMesh);
 
 		SourceMeshTransformGizmo->SetNewGizmoTransform(SourceMeshParentActor->GetActorTransform());
 		SourceMeshTransformGizmo->SetVisibility(!ToolProperties->bHideSourceMesh);
 		SourceMeshTransformGizmo->ActiveGizmoMode = EToolContextTransformGizmoMode::Combined;
+
+		TargetClothPreview->InvalidateResult();
 	}
 	else
 	{
 		SourceMeshComponent->SetVisibility(false);
 		SourceMeshTransformGizmo->SetVisibility(false);
-	}
+	}	
+}
 
-	TargetClothPreview->InvalidateResult();
+
+void UClothTransferSkinWeightsTool::OpFinishedCallback(const UE::Geometry::FDynamicMeshOperator* Op)
+{
+	if (Op->GetResultInfo().Result == UE::Geometry::EGeometryResultType::Failure)
+	{
+		GetToolManager()->DisplayMessage(LOCTEXT("TransferOpFailedWarning", "Weight transfer failed"), EToolMessageLevel::UserWarning);
+		bHasOpFailedWarning = true;
+	}
+	else
+	{
+		if (bHasOpFailedWarning)
+		{
+			GetToolManager()->DisplayMessage(FText(), EToolMessageLevel::UserWarning);    // clear old warning
+			bHasOpFailedWarning = false;
+		}
+	}
 }
 
+
 void UClothTransferSkinWeightsTool::PreviewMeshUpdatedCallback(UMeshOpPreviewWithBackgroundCompute* Preview)
 {
 	constexpr bool bProcessOnlyIfValid = true;
diff --git a/Engine/Plugins/Experimental/ChaosClothAssetEditor/Source/ChaosClothAssetEditorTools/Public/ChaosClothAsset/ClothTransferSkinWeightsTool.h b/Engine/Plugins/Experimental/ChaosClothAssetEditor/Source/ChaosClothAssetEditorTools/Public/ChaosClothAsset/ClothTransferSkinWeightsTool.h
--- a/Engine/Plugins/Experimental/ChaosClothAssetEditor/Source/ChaosClothAssetEditorTools/Public/ChaosClothAsset/ClothTransferSkinWeightsTool.h
+++ b/Engine/Plugins/Experimental/ChaosClothAssetEditor/Source/ChaosClothAssetEditorTools/Public/ChaosClothAsset/ClothTransferSkinWeightsTool.h
@@ -5,6 +5,9 @@
 #include "BaseTools/SingleSelectionMeshEditingTool.h"
 #include "BoneIndices.h"
 #include "ModelingOperators.h"
+#include "Transforms/TransformGizmoDataBinder.h"
+
+
 #include "ClothTransferSkinWeightsTool.generated.h"
 
 class UClothTransferSkinWeightsTool;
@@ -15,6 +18,7 @@ class UCombinedTransformGizmo;
 class UMeshOpPreviewWithBackgroundCompute;
 class AInternalToolFrameworkActor;
 class UDynamicMeshComponent;
+class FTransformGizmoDataBinder;
 
 UCLASS()
 class CHAOSCLOTHASSETEDITORTOOLS_API UClothTransferSkinWeightsToolProperties : public UInteractiveToolPropertySet
@@ -25,8 +29,15 @@ public:
 	UPROPERTY(EditAnywhere, Category = Source)
 	TObjectPtr<USkeletalMesh> SourceMesh;
 
-	UPROPERTY(EditAnywhere, Category = Source)
-	FTransform SourceMeshTransform;
+	UPROPERTY(EditAnywhere, Category = Source, meta = (EditCondition = "SourceMesh != nullptr"))
+	FVector3d SourceMeshTranslation;
+
+	UPROPERTY(EditAnywhere, Category = Source, meta = (EditCondition = "SourceMesh != nullptr"))
+	FVector3d SourceMeshRotation;
+
+	UPROPERTY(EditAnywhere, Category = Source, meta = (EditCondition = "SourceMesh != nullptr"))
+	FVector3d SourceMeshScale;
+
 
 	UPROPERTY(EditAnywhere, Category = Source)
 	int32 SourceMeshLOD = 0;
@@ -87,8 +98,12 @@ private:
 
 	void SetPreviewMeshColorFunction();
 
+	FTransform TransformFromProperties() const;
+
 	void UpdateSourceMesh();
 
+	void OpFinishedCallback(const UE::Geometry::FDynamicMeshOperator* Op);
+
 	void PreviewMeshUpdatedCallback(UMeshOpPreviewWithBackgroundCompute* Preview);
 
 
@@ -114,10 +129,13 @@ private:
 	UPROPERTY(Transient)
 	TObjectPtr<UCombinedTransformGizmo> SourceMeshTransformGizmo;
 
+	TSharedPtr<FTransformGizmoDataBinder> DataBinder;
+
 	// Used to lookup the index of the currently selected-by-name bone
 	TMap<FName, FBoneIndexType> TargetMeshBoneNameToIndex;
 
 	bool bHasInvalidLODWarning = false;
+	bool bHasOpFailedWarning = false;
 };
 
 
