Cloth Editor: enabling animations in the Preview Scene
- adding an AnimationAsset property in the cloth preview scene. This will be automatically associated with the scene's SkeletalMesh, if it exists
- adding SClothAnimationScrubPanel, a heavily simplified version of SAnimationScrubPanel, to control playback of the preview animation. The panel will be hidden unless there is a valid SkeletalMesh and Animation in the preview scene.
- put the preview's ClothComponent and SkeletalMeshComponent in the same Actor so they can be attached together

diff --git a/Engine/Plugins/Experimental/ChaosClothAssetEditor/Source/ChaosClothAssetEditor/ChaosClothAssetEditor.Build.cs b/Engine/Plugins/Experimental/ChaosClothAssetEditor/Source/ChaosClothAssetEditor/ChaosClothAssetEditor.Build.cs
--- a/Engine/Plugins/Experimental/ChaosClothAssetEditor/Source/ChaosClothAssetEditor/ChaosClothAssetEditor.Build.cs
+++ b/Engine/Plugins/Experimental/ChaosClothAssetEditor/Source/ChaosClothAssetEditor/ChaosClothAssetEditor.Build.cs
@@ -67,6 +67,9 @@ public class ChaosClothAssetEditor : ModuleRules
 				"TypedElementFramework",
 				"TypedElementRuntime",
 				"AssetDefinition",
+				"EditorWidgets",
+				"KismetWidgets",      // SScrubControlPanel
+				"AnimGraph"           // UAnimSingleNodeInstance
 				// ... add private dependencies that you statically link with here ...	
 			}
 			);
diff --git a/Engine/Plugins/Experimental/ChaosClothAssetEditor/Source/ChaosClothAssetEditor/Private/ChaosClothAsset/ClothEditor3DViewportClient.cpp b/Engine/Plugins/Experimental/ChaosClothAssetEditor/Source/ChaosClothAssetEditor/Private/ChaosClothAsset/ClothEditor3DViewportClient.cpp
--- a/Engine/Plugins/Experimental/ChaosClothAssetEditor/Source/ChaosClothAssetEditor/Private/ChaosClothAsset/ClothEditor3DViewportClient.cpp
+++ b/Engine/Plugins/Experimental/ChaosClothAssetEditor/Source/ChaosClothAssetEditor/Private/ChaosClothAsset/ClothEditor3DViewportClient.cpp
@@ -227,6 +227,16 @@ FBox FChaosClothAssetEditor3DViewportClient::PreviewBoundingBox() const
 	return FBox(ForceInitToZero);
 }
 
+TWeakPtr<FChaosClothPreviewScene> FChaosClothAssetEditor3DViewportClient::GetClothPreviewScene()
+{
+	return ClothPreviewScene;
+}
+
+TWeakPtr<const FChaosClothPreviewScene> FChaosClothAssetEditor3DViewportClient::GetClothPreviewScene() const
+{
+	return ClothPreviewScene;
+}
+
 void FChaosClothAssetEditor3DViewportClient::ProcessClick(FSceneView& View, HHitProxy* HitProxy, FKey Key, EInputEvent Event, uint32 HitX, uint32 HitY)
 {
 	FEditorViewportClient::ProcessClick(View, HitProxy, Key, Event, HitX, HitY);
diff --git a/Engine/Plugins/Experimental/ChaosClothAssetEditor/Source/ChaosClothAssetEditor/Private/ChaosClothAsset/ClothEditor3DViewportClient.h b/Engine/Plugins/Experimental/ChaosClothAssetEditor/Source/ChaosClothAssetEditor/Private/ChaosClothAsset/ClothEditor3DViewportClient.h
--- a/Engine/Plugins/Experimental/ChaosClothAssetEditor/Source/ChaosClothAssetEditor/Private/ChaosClothAsset/ClothEditor3DViewportClient.h
+++ b/Engine/Plugins/Experimental/ChaosClothAssetEditor/Source/ChaosClothAssetEditor/Private/ChaosClothAsset/ClothEditor3DViewportClient.h
@@ -52,6 +52,9 @@ public:
 
 	FBox PreviewBoundingBox() const;
 
+	TWeakPtr<FChaosClothPreviewScene> GetClothPreviewScene();
+	TWeakPtr<const FChaosClothPreviewScene> GetClothPreviewScene() const;
+
 private:
 
 	// FGCObject override
diff --git a/Engine/Plugins/Experimental/ChaosClothAssetEditor/Source/ChaosClothAssetEditor/Private/ChaosClothAsset/ClothEditorMode.cpp b/Engine/Plugins/Experimental/ChaosClothAssetEditor/Source/ChaosClothAssetEditor/Private/ChaosClothAsset/ClothEditorMode.cpp
--- a/Engine/Plugins/Experimental/ChaosClothAssetEditor/Source/ChaosClothAssetEditor/Private/ChaosClothAsset/ClothEditorMode.cpp
+++ b/Engine/Plugins/Experimental/ChaosClothAssetEditor/Source/ChaosClothAssetEditor/Private/ChaosClothAsset/ClothEditorMode.cpp
@@ -660,12 +660,7 @@ void UChaosClothAssetEditorMode::InitializeTargets(const TArray<TObjectPtr<UObje
 	{
 		if (UChaosClothAsset* ChaosClothAsset = Cast<UChaosClothAsset>(Asset))
 		{
-			PreviewScene->CreateClothActor(ChaosClothAsset);
-
-			USelection* SelectedComponents = GetModeManager()->GetSelectedComponents();
-			SelectedComponents->Modify();
-			SelectedComponents->Select(PreviewScene->ClothComponent);
-
+			PreviewScene->SetClothAsset(ChaosClothAsset);
 			break;
 		}
 	}
@@ -694,25 +689,25 @@ void UChaosClothAssetEditorMode::HardResetSimulation()
 
 void UChaosClothAssetEditorMode::SuspendSimulation()
 {
-	if (PreviewScene && PreviewScene->ClothComponent)
+	if (PreviewScene && PreviewScene->GetClothComponent())
 	{
-		PreviewScene->ClothComponent->SuspendSimulation();
+		PreviewScene->GetClothComponent()->SuspendSimulation();
 	}
 }
 
 void UChaosClothAssetEditorMode::ResumeSimulation()
 {
-	if (PreviewScene && PreviewScene->ClothComponent)
+	if (PreviewScene && PreviewScene->GetClothComponent())
 	{
-		PreviewScene->ClothComponent->ResumeSimulation();
+		PreviewScene->GetClothComponent()->ResumeSimulation();
 	}
 }
 
 bool UChaosClothAssetEditorMode::IsSimulationSuspended() const
 {
-	if (PreviewScene && PreviewScene->ClothComponent)
+	if (PreviewScene && PreviewScene->GetClothComponent())
 	{
-		return PreviewScene->ClothComponent->IsSimulationSuspended();
+		return PreviewScene->GetClothComponent()->IsSimulationSuspended();
 	}
 
 	return false;
@@ -753,7 +748,7 @@ void UChaosClothAssetEditorMode::ModeTick(float DeltaTime)
 
 	if (bShouldClearTeleportFlag)
 	{
-		PreviewScene->ClothComponent->ResetTeleportMode();
+		PreviewScene->GetClothComponent()->ResetTeleportMode();
 		bShouldClearTeleportFlag = false;
 	}
 
@@ -761,11 +756,11 @@ void UChaosClothAssetEditorMode::ModeTick(float DeltaTime)
 	{
 		if (bHardReset)
 		{
-			const FComponentReregisterContext Context(PreviewScene->ClothComponent);
+			const FComponentReregisterContext Context(PreviewScene->GetClothComponent());
 		}
 		else
 		{
-			PreviewScene->ClothComponent->ForceNextUpdateTeleportAndReset();
+			PreviewScene->GetClothComponent()->ForceNextUpdateTeleportAndReset();
 		}
 
 		bShouldResetSimulation = false;
@@ -853,9 +848,9 @@ FBox UChaosClothAssetEditorMode::SelectionBoundingBox() const
 
 FBox UChaosClothAssetEditorMode::PreviewBoundingBox() const
 {
-	if (PreviewScene->ClothComponent)
+	if (PreviewScene->GetClothComponent())
 	{
-		FBoxSphereBounds Bounds = PreviewScene->ClothComponent->Bounds;
+		FBoxSphereBounds Bounds = PreviewScene->GetClothComponent()->Bounds;
 		return Bounds.GetBox();
 	}
 
diff --git a/Engine/Plugins/Experimental/ChaosClothAssetEditor/Source/ChaosClothAssetEditor/Private/ChaosClothAsset/ClothEditorPreviewScene.cpp b/Engine/Plugins/Experimental/ChaosClothAssetEditor/Source/ChaosClothAssetEditor/Private/ChaosClothAsset/ClothEditorPreviewScene.cpp
--- a/Engine/Plugins/Experimental/ChaosClothAssetEditor/Source/ChaosClothAssetEditor/Private/ChaosClothAsset/ClothEditorPreviewScene.cpp
+++ b/Engine/Plugins/Experimental/ChaosClothAssetEditor/Source/ChaosClothAssetEditor/Private/ChaosClothAsset/ClothEditorPreviewScene.cpp
@@ -6,6 +6,8 @@
 #include "Animation/SkeletalMeshActor.h"
 #include "AssetEditorModeManager.h"
 #include "Elements/Framework/EngineElementsLibrary.h"
+#include "Animation/AnimSingleNodeInstance.h"
+#include "ComponentReregisterContext.h"
 
 #define LOCTEXT_NAMESPACE "UChaosClothEditorPreviewScene"
 
@@ -32,21 +34,18 @@ FChaosClothPreviewScene::FChaosClothPreviewScene(FPreviewScene::ConstructionValu
 
 	if (PreviewSceneDescription->SkeletalMeshAsset)
 	{
-		CreateSkeletalMeshActor();
+		CreateSkeletalMeshComponent();
+		ReattachSkeletalMeshAndAnimation();
 	}
 }
 
 FChaosClothPreviewScene::~FChaosClothPreviewScene()
 {
-	if (SkeletalMeshActor)
+	if (SkeletalMeshComponent)
 	{
-		if (SkeletalMeshActor->GetSkeletalMeshComponent())
-		{
-			SkeletalMeshActor->GetSkeletalMeshComponent()->TransformUpdated.RemoveAll(this);
-			SkeletalMeshActor->GetSkeletalMeshComponent()->SelectionOverrideDelegate.Unbind();
-		}
-
-		SkeletalMeshActor->UnregisterAllComponents();
+		SkeletalMeshComponent->TransformUpdated.RemoveAll(this);
+		SkeletalMeshComponent->SelectionOverrideDelegate.Unbind();
+		SkeletalMeshComponent->UnregisterComponent();
 	}
 
 	if (ClothComponent)
@@ -62,8 +61,54 @@ void FChaosClothPreviewScene::AddReferencedObjects(FReferenceCollector& Collecto
 
 	Collector.AddReferencedObject(PreviewSceneDescription);
 	Collector.AddReferencedObject(ClothComponent);
-	Collector.AddReferencedObject(SkeletalMeshActor);
-	Collector.AddReferencedObject(ClothActor);
+	Collector.AddReferencedObject(SkeletalMeshComponent);
+	Collector.AddReferencedObject(SceneActor);
+	Collector.AddReferencedObject(PreviewAnimInstance);
+}
+
+
+void FChaosClothPreviewScene::ReattachSkeletalMeshAndAnimation()
+{
+	if (SkeletalMeshComponent)
+	{
+		if (PreviewSceneDescription->AnimationAsset)
+		{
+			PreviewAnimInstance = NewObject<UAnimSingleNodeInstance>(SkeletalMeshComponent);
+			PreviewAnimInstance->SetAnimationAsset(PreviewSceneDescription->AnimationAsset);
+
+			SkeletalMeshComponent->SetAnimationMode(EAnimationMode::AnimationSingleNode);
+			SkeletalMeshComponent->InitAnim(true);
+			SkeletalMeshComponent->AnimationData.PopulateFrom(PreviewAnimInstance);
+			SkeletalMeshComponent->AnimScriptInstance = PreviewAnimInstance;
+			SkeletalMeshComponent->AnimScriptInstance->InitializeAnimation();
+			SkeletalMeshComponent->ValidateAnimation();
+			SkeletalMeshComponent->Stop();
+		}
+		else
+		{
+			SkeletalMeshComponent->Stop();
+			SkeletalMeshComponent->AnimationData = FSingleAnimationPlayData();
+			SkeletalMeshComponent->AnimScriptInstance = nullptr;
+		}
+	}
+
+	if (ClothComponent && SkeletalMeshComponent)
+	{
+		const bool bWasSimming = !ClothComponent->IsSimulationSuspended();
+		ClothComponent->SuspendSimulation();
+
+		ClothComponent->AttachToComponent(SkeletalMeshComponent, FAttachmentTransformRules::SnapToTargetNotIncludingScale);
+
+		// hard reset and resume simulation if it was running prior
+		{
+			const FComponentReregisterContext Context(ClothComponent);
+		}
+		if (bWasSimming)
+		{
+			ClothComponent->ResumeSimulation();
+		}
+	}
+
 }
 
 void FChaosClothPreviewScene::SceneDescriptionPropertyChanged(struct FPropertyChangedEvent& PropertyChangedEvent)
@@ -72,17 +117,46 @@ void FChaosClothPreviewScene::SceneDescriptionPropertyChanged(struct FPropertyCh
 	{
 		if (PreviewSceneDescription->SkeletalMeshAsset)
 		{
-			CreateSkeletalMeshActor();
+			CreateSkeletalMeshComponent();
+		}
+		else
+		{
+			DeleteSkeletalMeshComponent();
 		}
 	}
 
 	if (PropertyChangedEvent.Property->GetFName() == GET_MEMBER_NAME_CHECKED(UChaosClothPreviewSceneDescription, SkeletalMeshTransform))
 	{
-		if (SkeletalMeshActor)
+		if (SkeletalMeshComponent)
+		{
+			SkeletalMeshComponent->SetComponentToWorld(PreviewSceneDescription->SkeletalMeshTransform);
+		}
+	}
+
+	if (PropertyChangedEvent.Property->GetFName() == GET_MEMBER_NAME_CHECKED(UChaosClothPreviewSceneDescription, AnimationAsset))
+	{
+		if (!SkeletalMeshComponent || !PreviewSceneDescription->AnimationAsset)
 		{
-			SkeletalMeshActor->GetSkeletalMeshComponent()->SetComponentToWorld(PreviewSceneDescription->SkeletalMeshTransform);
+			PreviewAnimInstance = nullptr;
 		}
 	}
+
+	ReattachSkeletalMeshAndAnimation();
+}
+
+UChaosClothComponent* FChaosClothPreviewScene::GetClothComponent()
+{
+	return ClothComponent;
+}
+
+const UChaosClothComponent* FChaosClothPreviewScene::GetClothComponent() const
+{
+	return ClothComponent;
+}
+
+const USkeletalMeshComponent* FChaosClothPreviewScene::GetSkeletalMeshComponent() const
+{
+	return SkeletalMeshComponent;
 }
 
 void FChaosClothPreviewScene::SetModeManager(TSharedPtr<FAssetEditorModeManager> InClothPreviewEditorModeManager)
@@ -97,36 +171,62 @@ const TSharedPtr<const FAssetEditorModeManager> FChaosClothPreviewScene::GetClot
 
 void FChaosClothPreviewScene::SkeletalMeshTransformChanged(USceneComponent* UpdatedComponent, EUpdateTransformFlags UpdateTransformFlags, ETeleportType Teleport)
 {
-	ensure(UpdatedComponent == SkeletalMeshActor->GetSkeletalMeshComponent());
+	ensure(UpdatedComponent == SkeletalMeshComponent);
 	PreviewSceneDescription->SkeletalMeshTransform = UpdatedComponent->GetComponentToWorld();
 }
 
 
-void FChaosClothPreviewScene::CreateSkeletalMeshActor()
+void FChaosClothPreviewScene::DeleteSkeletalMeshComponent()
 {
-	if (SkeletalMeshActor)
+	if (SkeletalMeshComponent)
 	{
-		if (SkeletalMeshActor->GetSkeletalMeshComponent())
+		if (SceneActor)
 		{
-			SkeletalMeshActor->GetSkeletalMeshComponent()->TransformUpdated.RemoveAll(this);
-			SkeletalMeshActor->GetSkeletalMeshComponent()->SelectionOverrideDelegate.Unbind();
+			SceneActor->RemoveOwnedComponent(SkeletalMeshComponent);
 		}
 
-		SkeletalMeshActor->UnregisterAllComponents();
+		SkeletalMeshComponent->TransformUpdated.RemoveAll(this);
+		SkeletalMeshComponent->SelectionOverrideDelegate.Unbind();
+		SkeletalMeshComponent->UnregisterComponent();
 	}
 
-	ensure(PreviewSceneDescription->SkeletalMeshAsset);
+	SkeletalMeshComponent = nullptr;
 
-	SkeletalMeshActor = GetWorld()->SpawnActor<ASkeletalMeshActor>(ASkeletalMeshActor::StaticClass(), FTransform::Identity);
-	if (SkeletalMeshActor)
+	// Detach the ClothComponent and make it the RootComponent of the SceneActor
+	if (SceneActor && ClothComponent)
 	{
-		SkeletalMeshActor->GetSkeletalMeshComponent()->SetSkeletalMeshAsset(PreviewSceneDescription->SkeletalMeshAsset);
+		const bool bWasSimming = !ClothComponent->IsSimulationSuspended();
+		ClothComponent->SuspendSimulation();
+
+		ClothComponent->DetachFromComponent(FDetachmentTransformRules::KeepRelativeTransform);
+		SceneActor->SetRootComponent(ClothComponent);
+		SceneActor->RegisterAllComponents();
+
+		// hard reset and resume simulation if it was running prior
+		{
+			const FComponentReregisterContext Context(ClothComponent);
+		}
+		if (bWasSimming)
+		{
+			ClothComponent->ResumeSimulation();
+		}
 	}
+}
+
+void FChaosClothPreviewScene::CreateSkeletalMeshComponent()
+{
+	ensure(PreviewSceneDescription->SkeletalMeshAsset);
 
-	SkeletalMeshActor->RegisterAllComponents();
+	// Remove any existing skeletal mesh component
+	DeleteSkeletalMeshComponent();
 
-	SkeletalMeshActor->GetSkeletalMeshComponent()->TransformUpdated.AddRaw(this, &FChaosClothPreviewScene::SkeletalMeshTransformChanged);
-	SkeletalMeshActor->GetSkeletalMeshComponent()->SelectionOverrideDelegate = UPrimitiveComponent::FSelectionOverride::CreateRaw(this, &FChaosClothPreviewScene::IsComponentSelected);
+	SkeletalMeshComponent = NewObject<USkeletalMeshComponent>(SceneActor);
+	SkeletalMeshComponent->SetSkeletalMeshAsset(PreviewSceneDescription->SkeletalMeshAsset);
+	SkeletalMeshComponent->TransformUpdated.AddRaw(this, &FChaosClothPreviewScene::SkeletalMeshTransformChanged);
+	SkeletalMeshComponent->SelectionOverrideDelegate = UPrimitiveComponent::FSelectionOverride::CreateRaw(this, &FChaosClothPreviewScene::IsComponentSelected);
+
+	SceneActor->SetRootComponent(SkeletalMeshComponent);
+	SceneActor->RegisterAllComponents();
 }
 
 bool FChaosClothPreviewScene::IsComponentSelected(const UPrimitiveComponent* InComponent)
@@ -143,31 +243,42 @@ bool FChaosClothPreviewScene::IsComponentSelected(const UPrimitiveComponent* InC
 	return false;
 }
 
-void FChaosClothPreviewScene::CreateClothActor(UChaosClothAsset* Asset)
+void FChaosClothPreviewScene::SetClothAsset(UChaosClothAsset* Asset)
 {
 	if (ClothComponent)
 	{
 		ClothComponent->SelectionOverrideDelegate.Unbind();
 	}
 
-	if (ClothActor)
+	if (SceneActor)
 	{
-		ClothActor->UnregisterAllComponents();
+		SceneActor->UnregisterAllComponents();
 	}
 
-	ClothActor = GetWorld()->SpawnActor<AActor>(AActor::StaticClass());
+	SceneActor = GetWorld()->SpawnActor<AActor>(AActor::StaticClass());
 
-	check(ClothActor);
+	check(SceneActor);
 
-	ClothComponent = NewObject<UChaosClothComponent>(ClothActor);
+	ClothComponent = NewObject<UChaosClothComponent>(SceneActor);
 	ClothComponent->SetClothAsset(Asset);
 
 	ClothComponent->SelectionOverrideDelegate = UPrimitiveComponent::FSelectionOverride::CreateRaw(this, &FChaosClothPreviewScene::IsComponentSelected);
 
-	ClothActor->SetRootComponent(ClothComponent);
-	ClothActor->RegisterAllComponents();
+	SceneActor->SetRootComponent(ClothComponent);
+	SceneActor->RegisterAllComponents();
+
+	ReattachSkeletalMeshAndAnimation();
 }
 
+UAnimSingleNodeInstance* FChaosClothPreviewScene::GetPreviewAnimInstance()
+{
+	return PreviewAnimInstance;
+}
+
+const UAnimSingleNodeInstance* const FChaosClothPreviewScene::GetPreviewAnimInstance() const
+{
+	return PreviewAnimInstance;
+}
 
 #undef LOCTEXT_NAMESPACE
 
diff --git a/Engine/Plugins/Experimental/ChaosClothAssetEditor/Source/ChaosClothAssetEditor/Private/ChaosClothAsset/ClothEditorPreviewScene.h b/Engine/Plugins/Experimental/ChaosClothAssetEditor/Source/ChaosClothAssetEditor/Private/ChaosClothAsset/ClothEditorPreviewScene.h
--- a/Engine/Plugins/Experimental/ChaosClothAssetEditor/Source/ChaosClothAssetEditor/Private/ChaosClothAsset/ClothEditorPreviewScene.h
+++ b/Engine/Plugins/Experimental/ChaosClothAssetEditor/Source/ChaosClothAssetEditor/Private/ChaosClothAsset/ClothEditorPreviewScene.h
@@ -10,10 +10,9 @@ class UChaosClothAsset;
 class USkeletalMesh;
 class UChaosClothComponent;
 class FChaosClothPreviewScene;
-class UPhysicsAsset;
-class UAnimSequence;
-class ASkeletalMeshActor;
 class FAssetEditorModeManager;
+class UAnimationAsset;
+class UAnimSingleNodeInstance;
 
 ///
 /// The UChaosClothPreviewSceneDescription is a description of the Preview scene contents, intended to be editable in an FAdvancedPreviewSettingsWidget
@@ -33,9 +32,8 @@ public:
 	UPROPERTY(EditAnywhere, Category = "SkeletalMesh")
 	FTransform SkeletalMeshTransform;
 
-	// TODO: Add anything else to the scene, e.g.:
-	//UPROPERTY(EditAnywhere, Category = "Animation")
-	//TObjectPtr<UAnimSequence> AnimationSequence;
+	UPROPERTY(EditAnywhere, Category = "SkeletalMesh")
+	TObjectPtr<UAnimationAsset> AnimationAsset;
 
 private:
 
@@ -61,33 +59,47 @@ public:
 	const UChaosClothPreviewSceneDescription* GetPreviewSceneDescription() const { return PreviewSceneDescription; }
 	UChaosClothPreviewSceneDescription* GetPreviewSceneDescription() { return PreviewSceneDescription; }
 
-	void CreateClothActor(UChaosClothAsset* Asset);
+	void SetClothAsset(UChaosClothAsset* Asset);
 
 	// Update Scene in response to the SceneDescription changing
 	void SceneDescriptionPropertyChanged(struct FPropertyChangedEvent& PropertyChangedEvent);
 
-	// Preview simulation mesh
-	TObjectPtr<AActor> ClothActor;
-	TObjectPtr<UChaosClothComponent> ClothComponent;
+	UAnimSingleNodeInstance* GetPreviewAnimInstance();
+	const UAnimSingleNodeInstance* const GetPreviewAnimInstance() const;
 
-	// Skeletal Mesh
-	TObjectPtr<ASkeletalMeshActor> SkeletalMeshActor;
+	UChaosClothComponent* GetClothComponent();
+	const UChaosClothComponent* GetClothComponent() const;
+	
+	const USkeletalMeshComponent* GetSkeletalMeshComponent() const;
 
 	void SetModeManager(TSharedPtr<FAssetEditorModeManager> InClothPreviewEditorModeManager);
-
 	const TSharedPtr<const FAssetEditorModeManager> GetClothPreviewEditorModeManager() const;
 
 private:
 
+	// (Re)attach the cloth component to the skeletal mesh component, if it exists. 
+	// Create the PreviewAnimationInstance if the AnimationAsset and SkeletalMesh both exist.
+	void ReattachSkeletalMeshAndAnimation();
+
 	void SkeletalMeshTransformChanged(USceneComponent* UpdatedComponent, EUpdateTransformFlags UpdateTransformFlags, ETeleportType Teleport);
 
-	void CreateSkeletalMeshActor();
+	void CreateSkeletalMeshComponent();
+	void DeleteSkeletalMeshComponent();
 
 	bool IsComponentSelected(const UPrimitiveComponent* InComponent);
 
 	TObjectPtr<UChaosClothPreviewSceneDescription> PreviewSceneDescription;
 
 	TSharedPtr<FAssetEditorModeManager> ClothPreviewEditorModeManager;
+
+	TObjectPtr<UAnimSingleNodeInstance> PreviewAnimInstance;
+
+	TObjectPtr<AActor> SceneActor;
+
+	TObjectPtr<UChaosClothComponent> ClothComponent;
+
+	TObjectPtr<USkeletalMeshComponent> SkeletalMeshComponent;
+
 };
 
 
diff --git a/Engine/Plugins/Experimental/ChaosClothAssetEditor/Source/ChaosClothAssetEditor/Private/ChaosClothAsset/ClothEditorToolkit.cpp b/Engine/Plugins/Experimental/ChaosClothAssetEditor/Source/ChaosClothAssetEditor/Private/ChaosClothAsset/ClothEditorToolkit.cpp
--- a/Engine/Plugins/Experimental/ChaosClothAssetEditor/Source/ChaosClothAssetEditor/Private/ChaosClothAsset/ClothEditorToolkit.cpp
+++ b/Engine/Plugins/Experimental/ChaosClothAssetEditor/Source/ChaosClothAssetEditor/Private/ChaosClothAsset/ClothEditorToolkit.cpp
@@ -440,7 +440,7 @@ void FChaosClothAssetEditorToolkit::PostInitAssetEditor()
 	ViewportClient->ReceivedFocus(ViewportClient->Viewport);
 
 	// Set up 3D viewport
-	ClothPreviewViewportClient->SetClothComponent(ClothPreviewScene->ClothComponent);
+	ClothPreviewViewportClient->SetClothComponent(ClothPreviewScene->GetClothComponent());
 	ClothPreviewViewportClient->SetClothEdMode(ClothMode);
 	ClothPreviewViewportClient->SetClothEditorToolkit(StaticCastSharedRef<FChaosClothAssetEditorToolkit>(this->AsShared()));
 
@@ -712,7 +712,7 @@ void FChaosClothAssetEditorToolkit::InitDetailsViewPanel()
 void FChaosClothAssetEditorToolkit::PopulateOutliner()
 {
 	TSharedPtr<FManagedArrayCollection> ClothCollection;
-	if (const TObjectPtr<UChaosClothComponent> ClothComponent = ClothPreviewScene->ClothComponent)
+	if (const TObjectPtr<UChaosClothComponent> ClothComponent = ClothPreviewScene->GetClothComponent())
 	{
 		if (UChaosClothAsset* const ClothAsset = ClothComponent->GetClothAsset())
 		{
@@ -980,7 +980,7 @@ void FChaosClothAssetEditorToolkit::OnClothAssetChanged()
 
 	if (UChaosClothAsset* const ClothAsset = Cast<UChaosClothAsset>(ObjectsToEdit[0]))
 	{
-		ClothPreviewScene->CreateClothActor(ClothAsset);
+		ClothPreviewScene->SetClothAsset(ClothAsset);
 
 		ensure(ClothAsset->HasAnyFlags(RF_Transactional));		// Ensure all objects are transactable for undo/redo in the details panel
 		SetEditingObject(ClothAsset);
diff --git a/Engine/Plugins/Experimental/ChaosClothAssetEditor/Source/ChaosClothAssetEditor/Private/ChaosClothAsset/SClothAnimationScrubPanel.cpp b/Engine/Plugins/Experimental/ChaosClothAssetEditor/Source/ChaosClothAssetEditor/Private/ChaosClothAsset/SClothAnimationScrubPanel.cpp
new file mode 100644
--- /dev/null
+++ b/Engine/Plugins/Experimental/ChaosClothAssetEditor/Source/ChaosClothAssetEditor/Private/ChaosClothAsset/SClothAnimationScrubPanel.cpp
@@ -0,0 +1,293 @@
+// Copyright Epic Games, Inc. All Rights Reserved.
+
+#include "ChaosClothAsset/SClothAnimationScrubPanel.h"
+#include "Widgets/SBoxPanel.h"
+#include "ChaosClothAsset/ClothEditorPreviewScene.h"
+#include "SScrubControlPanel.h"
+#include "Animation/AnimSingleNodeInstance.h"
+#include "Animation/AnimSequence.h"
+#include "Animation/BlendSpace.h"
+
+#define LOCTEXT_NAMESPACE "ClothAnimationScrubPanel"
+
+void SClothAnimationScrubPanel::Construct( const SClothAnimationScrubPanel::FArguments& InArgs, const TWeakPtr<FChaosClothPreviewScene> InPreviewScene)
+{
+	PreviewSceneWeakPtr = InPreviewScene;
+
+	this->ChildSlot
+	[
+		SNew(SHorizontalBox)
+		.AddMetaData<FTagMetaData>(TEXT("ClothAnimScrub.Scrub"))
+		+SHorizontalBox::Slot()
+		.HAlign(HAlign_Fill) 
+		.VAlign(VAlign_Center)
+		.FillWidth(1)
+		.Padding(0.0f)
+		[
+			SAssignNew(ScrubControlPanel, SScrubControlPanel)
+			.IsEnabled(true)
+			.Value(this, &SClothAnimationScrubPanel::GetScrubValue)
+			.NumOfKeys(this, &SClothAnimationScrubPanel::GetNumberOfKeys)
+			.SequenceLength(this, &SClothAnimationScrubPanel::GetSequenceLength)
+			.DisplayDrag(this, &SClothAnimationScrubPanel::GetDisplayDrag)
+			.OnValueChanged(this, &SClothAnimationScrubPanel::OnValueChanged)
+			.OnBeginSliderMovement(this, &SClothAnimationScrubPanel::OnBeginSliderMovement)
+			.OnClickedForwardPlay(this, &SClothAnimationScrubPanel::OnClick_Forward)
+			.OnClickedForwardStep(this, &SClothAnimationScrubPanel::OnClick_Forward_Step)
+			.OnClickedForwardEnd(this, &SClothAnimationScrubPanel::OnClick_Forward_End)
+			.OnClickedBackwardPlay(this, &SClothAnimationScrubPanel::OnClick_Backward)
+			.OnClickedBackwardStep(this, &SClothAnimationScrubPanel::OnClick_Backward_Step)
+			.OnClickedBackwardEnd(this, &SClothAnimationScrubPanel::OnClick_Backward_End)
+			.OnClickedToggleLoop(this, &SClothAnimationScrubPanel::OnClick_ToggleLoop)
+			.OnGetLooping(this, &SClothAnimationScrubPanel::IsLoopStatusOn)
+			.OnGetPlaybackMode(this, &SClothAnimationScrubPanel::GetPlaybackMode)
+			.ViewInputMin(InArgs._ViewInputMin)
+			.ViewInputMax(InArgs._ViewInputMax)
+			.bDisplayAnimScrubBarEditing(false)
+			.bAllowZoom(false)
+			.IsRealtimeStreamingMode(false)
+		]
+	];
+}
+
+FReply SClothAnimationScrubPanel::OnClick_Forward_Step()
+{
+	if (UAnimSingleNodeInstance* const PreviewInstance = GetPreviewAnimationInstance())
+	{
+		PreviewInstance->SetPlaying(false);
+		PreviewInstance->StepForward();
+	}
+
+	return FReply::Handled();
+}
+
+FReply SClothAnimationScrubPanel::OnClick_Forward_End()
+{
+	if (UAnimSingleNodeInstance* const PreviewInstance = GetPreviewAnimationInstance())
+	{
+		PreviewInstance->SetPlaying(false);
+		PreviewInstance->SetPosition(PreviewInstance->GetLength(), false);
+	}
+
+	return FReply::Handled();
+}
+
+FReply SClothAnimationScrubPanel::OnClick_Backward_Step()
+{
+	if (UAnimSingleNodeInstance* const PreviewInstance = GetPreviewAnimationInstance())
+	{
+		PreviewInstance->SetPlaying(false);
+		PreviewInstance->StepBackward();
+	}
+
+	return FReply::Handled();
+}
+
+FReply SClothAnimationScrubPanel::OnClick_Backward_End()
+{
+	if (UAnimSingleNodeInstance* const PreviewInstance = GetPreviewAnimationInstance())
+	{
+		PreviewInstance->SetPlaying(false);
+		PreviewInstance->SetPosition(0.f, false);
+	}
+
+	return FReply::Handled();
+}
+
+FReply SClothAnimationScrubPanel::OnClick_Forward()
+{
+	if (UAnimSingleNodeInstance* const PreviewInstance = GetPreviewAnimationInstance())
+	{
+		const bool bIsReverse = PreviewInstance->IsReverse();
+		const bool bIsPlaying = PreviewInstance->IsPlaying();
+
+		// if current bIsReverse and bIsPlaying, we'd like to just turn off reverse
+		if (bIsReverse && bIsPlaying)
+		{
+			PreviewInstance->SetReverse(false);
+		}
+		// already playing, simply pause
+		else if (bIsPlaying) 
+		{
+			PreviewInstance->SetPlaying(false);
+		}
+		// if not playing, play forward
+		else 
+		{
+			//if we're at the end of the animation, jump back to the beginning before playing
+			if ( GetScrubValue() >= GetSequenceLength() )
+			{
+				PreviewInstance->SetPosition(0.0f, false);
+			}
+
+			PreviewInstance->SetReverse(false);
+			PreviewInstance->SetPlaying(true);
+		}
+	}
+
+	return FReply::Handled();
+}
+
+FReply SClothAnimationScrubPanel::OnClick_Backward()
+{
+	if (UAnimSingleNodeInstance* const PreviewInstance = GetPreviewAnimationInstance())
+	{
+		const bool bIsReverse = PreviewInstance->IsReverse();
+		const bool bIsPlaying = PreviewInstance->IsPlaying();
+
+		// if currently playing forward, just simply turn on reverse
+		if (!bIsReverse && bIsPlaying)
+		{
+			PreviewInstance->SetReverse(true);
+		}
+		else if (bIsPlaying)
+		{
+			PreviewInstance->SetPlaying(false);
+		}
+		else
+		{
+			// if we're at the beginning of the animation, jump back to the end before playing
+			if ( GetScrubValue() <= 0.0f )
+			{
+				PreviewInstance->SetPosition(GetSequenceLength(), false);
+			}
+
+			PreviewInstance->SetPlaying(true);
+			PreviewInstance->SetReverse(true);
+		}
+	}
+
+	return FReply::Handled();
+}
+
+FReply SClothAnimationScrubPanel::OnClick_ToggleLoop()
+{
+	if (UAnimSingleNodeInstance* const PreviewInstance = GetPreviewAnimationInstance())
+	{
+		bool bIsLooping = PreviewInstance->IsLooping();
+		PreviewInstance->SetLooping(!bIsLooping);
+	}
+	return FReply::Handled();
+}
+
+bool SClothAnimationScrubPanel::IsLoopStatusOn() const
+{
+	const UAnimSingleNodeInstance* const PreviewInstance = GetPreviewAnimationInstance();
+	return (PreviewInstance && PreviewInstance->IsLooping());
+}
+
+EPlaybackMode::Type SClothAnimationScrubPanel::GetPlaybackMode() const
+{
+	if (const UAnimSingleNodeInstance* const PreviewInstance = GetPreviewAnimationInstance())
+	{
+		if (PreviewInstance->IsPlaying())
+		{
+			return PreviewInstance->IsReverse() ? EPlaybackMode::PlayingReverse : EPlaybackMode::PlayingForward;
+		}
+		return EPlaybackMode::Stopped;
+	}
+	
+	return EPlaybackMode::Stopped;
+}
+
+void SClothAnimationScrubPanel::OnValueChanged(float NewValue)
+{
+	if (UAnimSingleNodeInstance* const PreviewInstance = GetPreviewAnimationInstance())
+	{
+		PreviewInstance->SetPosition(NewValue);
+	}
+}
+
+void SClothAnimationScrubPanel::OnBeginSliderMovement()
+{
+	if (UAnimSingleNodeInstance* const PreviewInstance = GetPreviewAnimationInstance())
+	{
+		PreviewInstance->SetPlaying(false);
+	}
+}
+
+uint32 SClothAnimationScrubPanel::GetNumberOfKeys() const
+{
+	if (const TSharedPtr<FChaosClothPreviewScene> PreviewScene = PreviewSceneWeakPtr.Pin())
+	{
+		if (UAnimSingleNodeInstance* const PreviewInstance = PreviewScene->GetPreviewAnimInstance())	// non-const because UAnimSingleNodeInstance::GetLength() is non-const
+		{
+			const float Length = PreviewInstance->GetLength();
+
+			// if anim sequence, use correct num frames
+			int32 NumKeys = (int32)(Length / 0.0333f);
+
+			if (PreviewInstance->CurrentAsset)
+			{
+				if (PreviewInstance->CurrentAsset->IsA(UAnimSequenceBase::StaticClass()))
+				{
+					NumKeys = CastChecked<UAnimSequenceBase>(PreviewInstance->CurrentAsset)->GetNumberOfSampledKeys();
+				}
+				else if (PreviewInstance->CurrentAsset->IsA(UBlendSpace::StaticClass()))
+				{
+					// Blendspaces dont display frame notches, so just return 0 here
+					NumKeys = 0;
+				}
+			}
+			return NumKeys;
+		}
+	}
+
+	return 1;
+}
+
+float SClothAnimationScrubPanel::GetSequenceLength() const
+{
+	if (const TSharedPtr<FChaosClothPreviewScene> PreviewScene = PreviewSceneWeakPtr.Pin())
+	{
+		if (UAnimSingleNodeInstance* const PreviewInstance = PreviewScene->GetPreviewAnimInstance())	// non-const because UAnimSingleNodeInstance::GetLength() is non-const
+		{
+			return PreviewInstance->GetLength();
+		}
+	}
+
+	return 0.f;
+}
+
+float SClothAnimationScrubPanel::GetScrubValue() const
+{
+	if (const UAnimSingleNodeInstance* const PreviewInstance = GetPreviewAnimationInstance())
+	{
+		return PreviewInstance->GetCurrentTime(); 
+	}
+
+	return 0.f;
+}
+
+bool SClothAnimationScrubPanel::GetDisplayDrag() const
+{
+	const UAnimSingleNodeInstance* const PreviewInstance = GetPreviewAnimationInstance();
+	if (PreviewInstance && PreviewInstance->CurrentAsset)
+	{
+		return true;
+	}
+
+	return false;
+}
+
+UAnimSingleNodeInstance* SClothAnimationScrubPanel::GetPreviewAnimationInstance()
+{
+	if (const TSharedPtr<FChaosClothPreviewScene> PreviewScene = PreviewSceneWeakPtr.Pin())
+	{
+		return PreviewScene->GetPreviewAnimInstance();
+	}
+
+	return nullptr;
+}
+
+const UAnimSingleNodeInstance* SClothAnimationScrubPanel::GetPreviewAnimationInstance() const
+{
+	if (const TSharedPtr<const FChaosClothPreviewScene> PreviewScene = PreviewSceneWeakPtr.Pin())
+	{
+		return PreviewScene->GetPreviewAnimInstance();
+	}
+
+	return nullptr;
+}
+
+#undef LOCTEXT_NAMESPACE
diff --git a/Engine/Plugins/Experimental/ChaosClothAssetEditor/Source/ChaosClothAssetEditor/Private/ChaosClothAsset/SClothAnimationScrubPanel.h b/Engine/Plugins/Experimental/ChaosClothAssetEditor/Source/ChaosClothAssetEditor/Private/ChaosClothAsset/SClothAnimationScrubPanel.h
new file mode 100644
--- /dev/null
+++ b/Engine/Plugins/Experimental/ChaosClothAssetEditor/Source/ChaosClothAssetEditor/Private/ChaosClothAsset/SClothAnimationScrubPanel.h
@@ -0,0 +1,53 @@
+// Copyright Epic Games, Inc. All Rights Reserved.
+
+#pragma once
+
+#include "Widgets/SCompoundWidget.h"
+#include "ITransportControl.h"  // EPlaybackMode::Type
+
+class SScrubControlPanel;
+class FChaosClothPreviewScene;
+class UAnimSingleNodeInstance;
+
+///
+/// A simple panel containing animation controls to operate on a UAnimSingleNodeInstance
+///
+class SClothAnimationScrubPanel : public SCompoundWidget
+{
+private:
+
+	SLATE_BEGIN_ARGS(SClothAnimationScrubPanel)	{}
+		SLATE_ATTRIBUTE( float, ViewInputMin )
+		SLATE_ATTRIBUTE( float, ViewInputMax )
+	SLATE_END_ARGS()
+
+	void Construct( const FArguments& InArgs, const TWeakPtr<FChaosClothPreviewScene> InPreviewScene );
+
+	// notifiers 
+	FReply OnClick_Forward_Step();
+	FReply OnClick_Forward_End();
+	FReply OnClick_Backward_Step();
+	FReply OnClick_Backward_End();
+	FReply OnClick_Forward();
+	FReply OnClick_Backward();
+	FReply OnClick_ToggleLoop();
+
+	void OnValueChanged(float NewValue);
+	void OnBeginSliderMovement();
+
+	EPlaybackMode::Type GetPlaybackMode() const;
+	bool IsLoopStatusOn() const;
+	float GetScrubValue() const;
+
+	UAnimSingleNodeInstance* GetPreviewAnimationInstance();
+	const UAnimSingleNodeInstance* GetPreviewAnimationInstance() const;
+
+	uint32 GetNumberOfKeys() const;
+	float GetSequenceLength() const;
+
+	bool GetDisplayDrag() const;
+
+	TWeakPtr<FChaosClothPreviewScene> PreviewSceneWeakPtr;
+
+	TSharedPtr<SScrubControlPanel> ScrubControlPanel;
+};
diff --git a/Engine/Plugins/Experimental/ChaosClothAssetEditor/Source/ChaosClothAssetEditor/Private/ChaosClothAsset/SClothEditor3DViewport.cpp b/Engine/Plugins/Experimental/ChaosClothAssetEditor/Source/ChaosClothAssetEditor/Private/ChaosClothAsset/SClothEditor3DViewport.cpp
--- a/Engine/Plugins/Experimental/ChaosClothAssetEditor/Source/ChaosClothAssetEditor/Private/ChaosClothAsset/SClothEditor3DViewport.cpp
+++ b/Engine/Plugins/Experimental/ChaosClothAssetEditor/Source/ChaosClothAssetEditor/Private/ChaosClothAsset/SClothEditor3DViewport.cpp
@@ -5,10 +5,59 @@
 #include "ChaosClothAsset/SClothEditor3DViewportToolBar.h"
 #include "ChaosClothAsset/ClothEditorCommands.h"
 #include "ChaosClothAsset/ClothEditorMode.h"
+#include "ChaosClothAsset/ClothEditorPreviewScene.h"
+#include "ChaosClothAsset/SClothAnimationScrubPanel.h"
+#include "Animation/AnimSingleNodeInstance.h"
 #include "EditorModeTools.h"
+#include "Widgets/Layout/SBorder.h"
 
 #define LOCTEXT_NAMESPACE "SChaosClothAssetEditor3DViewport"
 
+void SChaosClothAssetEditor3DViewport::Construct(const FArguments& InArgs, const FAssetEditorViewportConstructionArgs& InViewportConstructionArgs)
+{
+	SAssetEditorViewport::FArguments ParentArgs;
+	ParentArgs._EditorViewportClient = InArgs._EditorViewportClient;
+	if (InArgs._ViewportSize.IsSet())
+	{
+		ParentArgs._ViewportSize = InArgs._ViewportSize;
+	}
+	SAssetEditorViewport::Construct(ParentArgs, InViewportConstructionArgs);
+
+	ViewportOverlay->AddSlot()
+	[
+		SNew(SHorizontalBox)
+		+ SHorizontalBox::Slot()
+		.HAlign(HAlign_Fill)
+		.VAlign(VAlign_Bottom)
+		.FillWidth(1)
+		.Padding(10.0f, 0.0f)
+		[
+			SNew(SBorder)
+			.BorderImage(FAppStyle::Get().GetBrush("EditorViewport.OverlayBrush"))
+			.Visibility_Raw(this, &SChaosClothAssetEditor3DViewport::GetAnimControlVisibility)
+			.Padding(10.0f, 2.0f)
+			[
+				SNew(SClothAnimationScrubPanel, GetPreviewScene())
+				.ViewInputMin(this, &SChaosClothAssetEditor3DViewport::GetViewMinInput)
+				.ViewInputMax(this, &SChaosClothAssetEditor3DViewport::GetViewMaxInput)
+			]
+		]
+	];
+}
+
+TWeakPtr<FChaosClothPreviewScene> SChaosClothAssetEditor3DViewport::GetPreviewScene()
+{
+	const TSharedPtr<FChaosClothAssetEditor3DViewportClient> ClothViewportClient = StaticCastSharedPtr<FChaosClothAssetEditor3DViewportClient>(Client);
+	return ClothViewportClient->GetClothPreviewScene();
+}
+
+TWeakPtr<const FChaosClothPreviewScene> SChaosClothAssetEditor3DViewport::GetPreviewScene() const
+{
+	const TSharedPtr<const FChaosClothAssetEditor3DViewportClient> ClothViewportClient = StaticCastSharedPtr<FChaosClothAssetEditor3DViewportClient>(Client);
+	return ClothViewportClient->GetClothPreviewScene();
+}
+
+
 void SChaosClothAssetEditor3DViewport::BindCommands()
 {
 	SAssetEditorViewport::BindCommands();
@@ -99,8 +148,31 @@ TSharedPtr<FExtender> SChaosClothAssetEditor3DViewport::GetExtenders() const
 	return Result;
 }
 
-void SChaosClothAssetEditor3DViewport::OnFloatingButtonClicked()
+float SChaosClothAssetEditor3DViewport::GetViewMinInput() const
+{
+	return 0.0f;
+}
+
+float SChaosClothAssetEditor3DViewport::GetViewMaxInput() const
+{
+	// (these are non-const because UAnimSingleNodeInstance::GetLength() is non-const)
+	const TSharedPtr<FChaosClothAssetEditor3DViewportClient> ClothViewportClient = StaticCastSharedPtr<FChaosClothAssetEditor3DViewportClient>(Client);
+	const TSharedPtr<FChaosClothPreviewScene> Scene = ClothViewportClient->GetClothPreviewScene().Pin();
+	if (Scene)
+	{
+		if (UAnimSingleNodeInstance* const PreviewInstance = Scene->GetPreviewAnimInstance())
+		{
+			return PreviewInstance->GetLength();
+		}
+	}
+
+	return 0.0f;
+}
+
+EVisibility SChaosClothAssetEditor3DViewport::GetAnimControlVisibility() const
 {
+	const TSharedPtr<const FChaosClothPreviewScene> Scene = GetPreviewScene().Pin();
+	return (Scene && Scene->GetSkeletalMeshComponent() && Scene->GetPreviewAnimInstance()) ? EVisibility::Visible : EVisibility::Hidden;
 }
 
 #undef LOCTEXT_NAMESPACE
diff --git a/Engine/Plugins/Experimental/ChaosClothAssetEditor/Source/ChaosClothAssetEditor/Private/ChaosClothAsset/SClothEditor3DViewport.h b/Engine/Plugins/Experimental/ChaosClothAssetEditor/Source/ChaosClothAssetEditor/Private/ChaosClothAsset/SClothEditor3DViewport.h
--- a/Engine/Plugins/Experimental/ChaosClothAssetEditor/Source/ChaosClothAssetEditor/Private/ChaosClothAsset/SClothEditor3DViewport.h
+++ b/Engine/Plugins/Experimental/ChaosClothAssetEditor/Source/ChaosClothAssetEditor/Private/ChaosClothAsset/SClothEditor3DViewport.h
@@ -7,6 +7,9 @@
 #include "SCommonEditorViewportToolbarBase.h"
 #include "SBaseCharacterFXEditorViewport.h"
 
+class FChaosClothPreviewScene;
+class SClothAnimationScrubPanel;
+
 /**
  * Viewport used for 3D preview in cloth editor. Has a custom toolbar overlay at the top.
  */
@@ -14,6 +17,14 @@ class CHAOSCLOTHASSETEDITOR_API SChaosClothAssetEditor3DViewport : public SBaseC
 {
 
 public:
+
+	SLATE_BEGIN_ARGS(SChaosClothAssetEditor3DViewport) {}
+		SLATE_ATTRIBUTE(FVector2D, ViewportSize);
+		SLATE_ARGUMENT(TSharedPtr<FEditorViewportClient>, EditorViewportClient)
+	SLATE_END_ARGS()
+
+	void Construct(const FArguments& InArgs, const FAssetEditorViewportConstructionArgs& InViewportConstructionArgs);
+
 	// SAssetEditorViewport
 	virtual void BindCommands() override;
 	virtual TSharedPtr<SWidget> MakeViewportToolbar() override;
@@ -23,6 +34,16 @@ public:
 	// ICommonEditorViewportToolbarInfoProvider
 	virtual TSharedRef<class SEditorViewport> GetViewportWidget() override;
 	virtual TSharedPtr<FExtender> GetExtenders() const override;
-	virtual void OnFloatingButtonClicked() override;
+	virtual void OnFloatingButtonClicked() override {}
+
+private:
+
+	TWeakPtr<FChaosClothPreviewScene> GetPreviewScene();
+	TWeakPtr<const FChaosClothPreviewScene> GetPreviewScene() const;
+
+	float GetViewMinInput() const;
+	float GetViewMaxInput() const;
+	EVisibility GetAnimControlVisibility() const;
+		 
 
 };
