ClothAssetEditor: add SkeletalMeshImportNode

diff --git a/Engine/Plugins/Experimental/ChaosClothAssetEditor/Source/ChaosClothAssetDataflowNodes/Private/ChaosClothAsset/DataflowNodes.cpp b/Engine/Plugins/Experimental/ChaosClothAssetEditor/Source/ChaosClothAssetDataflowNodes/Private/ChaosClothAsset/DataflowNodes.cpp
--- a/Engine/Plugins/Experimental/ChaosClothAssetEditor/Source/ChaosClothAssetDataflowNodes/Private/ChaosClothAsset/DataflowNodes.cpp
+++ b/Engine/Plugins/Experimental/ChaosClothAssetEditor/Source/ChaosClothAssetDataflowNodes/Private/ChaosClothAsset/DataflowNodes.cpp
@@ -4,6 +4,7 @@
 #include "ChaosClothAsset/TerminalNode.h"
 #include "ChaosClothAsset/DatasmithImportNode.h"
 #include "ChaosClothAsset/ImportNode.h"
+#include "ChaosClothAsset/SkeletalMeshImportNode.h"
 #include "ChaosClothAsset/StaticMeshImportNode.h"
 #include "ChaosClothAsset/PatternSelectionNode.h"
 #include "ChaosClothAsset/DeleteRenderMeshNode.h"
@@ -27,6 +28,7 @@ namespace UE::Chaos::ClothAsset::DataflowNodes
 		DATAFLOW_NODE_REGISTER_CREATION_FACTORY(FChaosClothAssetTerminalNode);
 		DATAFLOW_NODE_REGISTER_CREATION_FACTORY(FChaosClothAssetDatasmithImportNode);
 		DATAFLOW_NODE_REGISTER_CREATION_FACTORY(FChaosClothAssetImportNode);
+		DATAFLOW_NODE_REGISTER_CREATION_FACTORY(FChaosClothAssetSkeletalMeshImportNode);
 		DATAFLOW_NODE_REGISTER_CREATION_FACTORY(FChaosClothAssetStaticMeshImportNode);
 		DATAFLOW_NODE_REGISTER_CREATION_FACTORY(FChaosClothAssetPatternSelectionNode);
 		DATAFLOW_NODE_REGISTER_CREATION_FACTORY(FChaosClothAssetDeleteRenderMeshNode);
diff --git a/Engine/Plugins/Experimental/ChaosClothAssetEditor/Source/ChaosClothAssetDataflowNodes/Private/ChaosClothAsset/SkeletalMeshImportNode.cpp b/Engine/Plugins/Experimental/ChaosClothAssetEditor/Source/ChaosClothAssetDataflowNodes/Private/ChaosClothAsset/SkeletalMeshImportNode.cpp
new file mode 100644
--- /dev/null
+++ b/Engine/Plugins/Experimental/ChaosClothAssetEditor/Source/ChaosClothAssetDataflowNodes/Private/ChaosClothAsset/SkeletalMeshImportNode.cpp
@@ -0,0 +1,196 @@
+// Copyright Epic Games, Inc. All Rights Reserved.
+
+#include "ChaosClothAsset/SkeletalMeshImportNode.h"
+#include "ChaosClothAsset/DataflowNodes.h"
+#include "ChaosClothAsset/CollectionClothFacade.h"
+#include "Animation/Skeleton.h"
+#include "BoneWeights.h"
+#include "Engine/SkeletalMesh.h"
+#include "Engine/SkinnedAssetCommon.h"
+#include "GeometryCollection/ManagedArrayCollection.h"
+#include "MeshDescriptionToDynamicMesh.h"
+#include "PhysicsEngine/PhysicsAsset.h"
+#include "Rendering/SkeletalMeshModel.h"
+
+#include UE_INLINE_GENERATED_CPP_BY_NAME(SkeletalMeshImportNode)
+
+#define LOCTEXT_NAMESPACE "ChaosClothAssetSkeletalMeshImportNode"
+
+FChaosClothAssetSkeletalMeshImportNode::FChaosClothAssetSkeletalMeshImportNode(const Dataflow::FNodeParameters& InParam, FGuid InGuid)
+	: FDataflowNode(InParam, InGuid)
+{
+	RegisterOutputConnection(&Collection);
+}
+
+void FChaosClothAssetSkeletalMeshImportNode::Evaluate(Dataflow::FContext& Context, const FDataflowOutput* Out) const
+{
+	using namespace UE::Chaos::ClothAsset;
+	using namespace UE::Chaos::ClothAsset::DataflowNodes;
+
+	if (Out->IsA<FManagedArrayCollection>(&Collection))
+	{
+		const TSharedRef<FManagedArrayCollection> ClothCollection = MakeShared<FManagedArrayCollection>();
+		if (SkeletalMesh)
+		{
+			const FSkeletalMeshModel* ImportedModel = SkeletalMesh->GetImportedModel();
+			const bool bIsValidLOD = ImportedModel && ImportedModel->LODModels.IsValidIndex(LODIndex);
+			if (!ensureAlways(bIsValidLOD))
+			{
+				SetValue<FManagedArrayCollection>(Context, *ClothCollection, &Collection);
+				return;
+			}
+			const FSkeletalMeshLODModel &LODModel = ImportedModel->LODModels[LODIndex];
+
+			const bool bIsValidSection = SectionIndex < LODModel.Sections.Num();
+			if (!ensureAlways(bIsValidSection))
+			{
+				SetValue<FManagedArrayCollection>(Context, *ClothCollection, &Collection);	
+				return;
+			}
+
+			const FSkelMeshSection &Section = LODModel.Sections[SectionIndex];
+			FCollectionClothFacade ClothFacade(ClothCollection);
+			ClothFacade.DefineSchema();
+			FCollectionClothLodFacade ClothLodFacade = ClothFacade.AddGetLod();
+			FCollectionClothPatternFacade ClothPatternFacade = ClothLodFacade.AddGetPattern();
+
+			using FBoneData = TTuple<int32, TArray<int32>, TArray<float>>;
+			auto ExtractBones = [](const FSkelMeshSection& Section, const FSoftSkinVertex& SoftVertex)
+			{
+				const int32 NumBones = Section.MaxBoneInfluences;
+				TArray<int32> BoneIndices;
+				TArray<float> BoneWeights;
+				BoneIndices.SetNum(NumBones);
+				BoneWeights.SetNum(NumBones);
+				for (int32 BoneIndex = 0; BoneIndex < NumBones; ++BoneIndex)
+				{
+					BoneIndices[BoneIndex] = (int32)Section.BoneMap[(int32)SoftVertex.InfluenceBones[BoneIndex]];
+					BoneWeights[BoneIndex] = (float)SoftVertex.InfluenceWeights[BoneIndex] * UE::AnimationCore::InvMaxRawBoneWeightFloat;
+				}
+				return FBoneData(NumBones, BoneIndices, BoneWeights);
+			};
+
+			auto GetSectionFace = [](const FSkelMeshSection& Section, const TArray<uint32>& IndexBuffer, int32 FaceIndex)
+			{
+				const int32 IndexOffset = Section.BaseIndex + FaceIndex * 3;
+				FIntVector3 TriangleIndices;
+				TriangleIndices.X = IndexBuffer[IndexOffset + 0] - Section.BaseVertexIndex;
+				TriangleIndices.Y = IndexBuffer[IndexOffset + 1] - Section.BaseVertexIndex;
+				TriangleIndices.Z = IndexBuffer[IndexOffset + 2] - Section.BaseVertexIndex;
+				return TriangleIndices;
+			};
+
+			if (bImportSimMesh)
+			{
+				ClothPatternFacade.SetNumSimVertices(Section.NumVertices);
+				ClothPatternFacade.SetNumSimFaces(Section.NumTriangles);
+
+				TArrayView<FVector2f> SimPosition = ClothPatternFacade.GetSimPosition();
+				TArrayView<FVector3f> SimRestPosition = ClothPatternFacade.GetSimRestPosition();
+				TArrayView<FVector3f> SimRestNormal = ClothPatternFacade.GetSimRestNormal();
+				TArrayView<int32> SimBoneInfluences = ClothPatternFacade.GetSimNumBoneInfluences();
+				TArrayView<TArray<int32>> SimBoneIndices = ClothPatternFacade.GetSimBoneIndices();
+				TArrayView<TArray<float>> SimBoneWeights = ClothPatternFacade.GetSimBoneWeights();
+				TArrayView<FIntVector3> SimIndices = ClothPatternFacade.GetSimIndices();
+
+				const bool bIsValidUV = UVChannel >= 0 && UVChannel < MAX_TEXCOORDS;
+				if (!bIsValidUV || UVChannel == INDEX_NONE)
+				{
+					UE_LOG(LogChaos, Error, TEXT("Invalid UV channel %d"), UVChannel);
+				}
+
+				for (int32 VertexIndex = 0; VertexIndex < Section.NumVertices; ++VertexIndex)
+				{
+					const FSoftSkinVertex& SoftVertex = Section.SoftVertices[VertexIndex];
+					SimPosition[VertexIndex] = bIsValidUV ? SoftVertex.UVs[UVChannel] : FVector2f(SoftVertex.Position[1], SoftVertex.Position[2]);
+					SimRestPosition[VertexIndex] = SoftVertex.Position;
+					SimRestNormal[VertexIndex] = -SoftVertex.TangentZ;
+
+					FBoneData BoneData = ExtractBones(Section, SoftVertex);
+					SimBoneInfluences[VertexIndex] = BoneData.Get<0>();
+					SimBoneIndices[VertexIndex] = MoveTemp(BoneData.Get<1>());
+					SimBoneWeights[VertexIndex] = MoveTemp(BoneData.Get<2>());
+				}
+
+				for (uint32 FaceIndex = 0; FaceIndex < Section.NumTriangles; ++FaceIndex)
+				{
+					SimIndices[FaceIndex] = GetSectionFace(Section, LODModel.IndexBuffer, FaceIndex);
+				}
+
+				ensureMsgf(false, TEXT("TODO: BuildSeams needs to be implemented"));
+				ClothLodFacade.SetNumTetherBatches(0);
+				ClothLodFacade.SetNumSeams(0);
+			}
+
+			if (bImportRenderMesh)
+			{
+				ClothPatternFacade.SetNumRenderVertices(Section.NumVertices);
+				ClothPatternFacade.SetNumRenderFaces(Section.NumTriangles);
+
+				TArrayView<FVector3f> RenderPosition = ClothPatternFacade.GetRenderPosition();
+				TArrayView<FVector3f> RenderNormal = ClothPatternFacade.GetRenderNormal();
+				TArrayView<FVector3f> RenderTangentU = ClothPatternFacade.GetRenderTangentU();
+				TArrayView<FVector3f> RenderTangentV = ClothPatternFacade.GetRenderTangentV();
+				TArrayView<TArray<FVector2f>> RenderUVs = ClothPatternFacade.GetRenderUVs();
+				TArrayView<FLinearColor> RenderColor = ClothPatternFacade.GetRenderColor();
+				TArrayView<int32> RenderNumBoneInfluences = ClothPatternFacade.GetRenderNumBoneInfluences();
+				TArrayView<TArray<int32>> RenderBoneIndices = ClothPatternFacade.GetRenderBoneIndices();
+				TArrayView<TArray<float>> RenderBoneWeights = ClothPatternFacade.GetRenderBoneWeights();
+				TArrayView<FIntVector3> RenderIndices = ClothPatternFacade.GetRenderIndices();
+				TArrayView<int32> RenderMaterialIndex = ClothPatternFacade.GetRenderMaterialIndex();
+
+				for (int32 VertexIndex = 0; VertexIndex < Section.NumVertices; ++VertexIndex)
+				{
+					const FSoftSkinVertex& SoftVertex = Section.SoftVertices[VertexIndex];
+
+					RenderPosition[VertexIndex] = SoftVertex.Position;
+					RenderNormal[VertexIndex] = SoftVertex.TangentZ;
+					RenderTangentU[VertexIndex] = SoftVertex.TangentX;
+					RenderTangentV[VertexIndex] = SoftVertex.TangentY;
+
+					TArray<FVector2f> UVs;
+					UVs.SetNum(MAX_TEXCOORDS);
+					for (int32 TexCoordIndex = 0; TexCoordIndex < MAX_TEXCOORDS; ++TexCoordIndex)
+					{
+						UVs[TexCoordIndex] = SoftVertex.UVs[TexCoordIndex];
+					}
+					RenderUVs[VertexIndex] = UVs;
+
+					RenderColor[VertexIndex] = FLinearColor(SoftVertex.Color);
+
+					FBoneData BoneData = ExtractBones(Section, SoftVertex);
+					RenderNumBoneInfluences[VertexIndex] = BoneData.Get<0>();
+					RenderBoneIndices[VertexIndex] = MoveTemp(BoneData.Get<1>());
+					RenderBoneWeights[VertexIndex] = MoveTemp(BoneData.Get<2>());
+				}
+
+				for (uint32 FaceIndex = 0; FaceIndex < Section.NumTriangles; ++FaceIndex)
+				{
+					RenderIndices[FaceIndex] = GetSectionFace(Section, LODModel.IndexBuffer, FaceIndex);
+					RenderMaterialIndex[FaceIndex] = 0;
+				}
+			}
+
+			ClothLodFacade.SetNumMaterials(1);
+			TArrayView<FString> RenderMaterialPathNames = ClothLodFacade.GetRenderMaterialPathName();
+			const TArray<FSkeletalMaterial>& Materials = SkeletalMesh->GetMaterials();
+			check(SectionIndex < Materials.Num());
+			if (const UMaterialInterface* Interface = Materials[SectionIndex].MaterialInterface)
+			{
+				RenderMaterialPathNames[0] = Interface->GetPathName();
+			}
+
+			if (const UPhysicsAsset* PhysicsAsset = SkeletalMesh->GetPhysicsAsset())
+			{
+				ClothLodFacade.SetPhysicsAssetPathName(PhysicsAsset->GetPathName());
+			}
+			if (const USkeleton* Skeleton = SkeletalMesh->GetSkeleton())
+			{
+				ClothLodFacade.SetSkeletonAssetPathName(Skeleton->GetPathName());
+			}
+		}
+		SetValue<FManagedArrayCollection>(Context, *ClothCollection, &Collection);	
+	}
+}
+
+#undef LOCTEXT_NAMESPACE
diff --git a/Engine/Plugins/Experimental/ChaosClothAssetEditor/Source/ChaosClothAssetDataflowNodes/Public/ChaosClothAsset/SkeletalMeshImportNode.h b/Engine/Plugins/Experimental/ChaosClothAssetEditor/Source/ChaosClothAssetDataflowNodes/Public/ChaosClothAsset/SkeletalMeshImportNode.h
new file mode 100644
--- /dev/null
+++ b/Engine/Plugins/Experimental/ChaosClothAssetEditor/Source/ChaosClothAssetDataflowNodes/Public/ChaosClothAsset/SkeletalMeshImportNode.h
@@ -0,0 +1,42 @@
+// Copyright Epic Games, Inc. All Rights Reserved.
+
+#pragma once 
+
+#include "Dataflow/DataflowNode.h"
+#include "GeometryCollection/ManagedArrayCollection.h"
+#include "SkeletalMeshImportNode.generated.h"
+
+class USkeletalMesh;
+
+USTRUCT(Meta = (DataflowCloth))
+struct FChaosClothAssetSkeletalMeshImportNode : public FDataflowNode
+{
+	GENERATED_USTRUCT_BODY()
+	DATAFLOW_NODE_DEFINE_INTERNAL(FChaosClothAssetSkeletalMeshImportNode, "SkeletalMeshImport", "Cloth", "Cloth Skeletal Mesh Import")
+public:
+
+	UPROPERTY(Meta = (DataflowOutput, DisplayName = "Collection"))
+	FManagedArrayCollection Collection;
+
+	UPROPERTY(EditAnywhere, Category = "Skeletal Mesh")
+	TObjectPtr<const USkeletalMesh> SkeletalMesh;
+
+	UPROPERTY(VisibleAnywhere, Category = "Skeletal Mesh")
+	int32 LODIndex = 0;
+
+	UPROPERTY(EditAnywhere, Category = "Skeletal Mesh", Meta = (ClampMin = "0"))
+	int32 SectionIndex = 0;
+
+	UPROPERTY(EditAnywhere, Category = "Skeletal Mesh")
+	bool bImportSimMesh = true;
+
+	UPROPERTY(EditAnywhere, Category = "Skeletal Mesh")
+	bool bImportRenderMesh = true;
+
+	UPROPERTY(EditAnywhere, Category = "Skeletal Mesh")
+	int32 UVChannel = INDEX_NONE;
+
+	FChaosClothAssetSkeletalMeshImportNode(const Dataflow::FNodeParameters& InParam, FGuid InGuid = FGuid::NewGuid());
+
+	virtual void Evaluate(Dataflow::FContext& Context, const FDataflowOutput* Out) const override;
+};
