ClothAssetBuilderEditor: prevent a crash when accidentally linking a cloth collection with no simulation mesh to the terminal node

diff --git a/Engine/Plugins/Experimental/ChaosClothAssetEditor/Source/ChaosClothAssetTools/Private/ChaosClothAsset/ClothAssetBuilderEditor.cpp b/Engine/Plugins/Experimental/ChaosClothAssetEditor/Source/ChaosClothAssetTools/Private/ChaosClothAsset/ClothAssetBuilderEditor.cpp
--- a/Engine/Plugins/Experimental/ChaosClothAssetEditor/Source/ChaosClothAssetTools/Private/ChaosClothAsset/ClothAssetBuilderEditor.cpp
+++ b/Engine/Plugins/Experimental/ChaosClothAssetEditor/Source/ChaosClothAssetTools/Private/ChaosClothAsset/ClothAssetBuilderEditor.cpp
@@ -87,6 +87,7 @@ void UClothAssetBuilderEditor::BuildLod(FSkeletalMeshLODModel& LODModel, const U
 	const ClothingMeshUtils::ClothMeshDesc SourceMesh(
 		GetSimPositions(ClothAsset, LodIndex),
 		GetSimIndices(ClothAsset, LodIndex));  // Let it calculate the averaged normals as to match the simulation data output
+	const bool bIsValidSourceMesh = (SourceMesh.GetPositions().Num() > 0);
 
 	const int32 NumLodSimVertices = GetNumVertices(ClothAsset, LodIndex);
 
@@ -254,35 +255,38 @@ void UClothAssetBuilderEditor::BuildLod(FSkeletalMeshLODModel& LODModel, const U
 		Section.CalcUse16BitBoneIndex();
 
 		// Setup clothing data
-		Section.ClothMappingDataLODs.SetNum(1);  // TODO: LODBias maps for raytracing
+		if (bIsValidSourceMesh)
+		{
+			Section.ClothMappingDataLODs.SetNum(1);  // TODO: LODBias maps for raytracing
 
-		Section.ClothingData.AssetLodIndex = LodIndex;
-		Section.ClothingData.AssetGuid = ClothAsset.AssetGuid;  // There is only one cloth asset,
-		Section.CorrespondClothAssetIndex = 0;       // this one
+			Section.ClothingData.AssetLodIndex = LodIndex;
+			Section.ClothingData.AssetGuid = ClothAsset.AssetGuid;  // There is only one cloth asset,
+			Section.CorrespondClothAssetIndex = 0;       // this one
 
-		TArray<FVector3f> RenderPositions;
-		TArray<FVector3f> RenderNormals;
-		TArray<FVector3f> RenderTangents;
-		for (const FSoftSkinVertex& SoftVert : Section.SoftVertices)
-		{
-			RenderPositions.Add(SoftVert.Position);
-			RenderNormals.Add(SoftVert.TangentZ);
-			RenderTangents.Add(SoftVert.TangentX);
-		}
+			TArray<FVector3f> RenderPositions;
+			TArray<FVector3f> RenderNormals;
+			TArray<FVector3f> RenderTangents;
+			for (const FSoftSkinVertex& SoftVert : Section.SoftVertices)
+			{
+				RenderPositions.Add(SoftVert.Position);
+				RenderNormals.Add(SoftVert.TangentZ);
+				RenderTangents.Add(SoftVert.TangentX);
+			}
 
-		const ClothingMeshUtils::ClothMeshDesc TargetMesh(RenderPositions, 
-			RenderNormals, 
-			RenderTangents,
-			LODModel.IndexBuffer);
-
-		ClothingMeshUtils::GenerateMeshToMeshVertData(
-			Section.ClothMappingDataLODs[0],
-			TargetMesh,
-			SourceMesh,
-			&MaxDistances,
-			ClothAsset.bSmoothTransition,
-			ClothAsset.bUseMultipleInfluences,
-			ClothAsset.SkinningKernelRadius);
+			const ClothingMeshUtils::ClothMeshDesc TargetMesh(RenderPositions, 
+				RenderNormals, 
+				RenderTangents,
+				LODModel.IndexBuffer);
+
+			ClothingMeshUtils::GenerateMeshToMeshVertData(
+				Section.ClothMappingDataLODs[0],
+				TargetMesh,
+				SourceMesh,
+				&MaxDistances,
+				ClothAsset.bSmoothTransition,
+				ClothAsset.bUseMultipleInfluences,
+				ClothAsset.SkinningKernelRadius);
+		}
 
 		// Save the original indices for the newly added vertices
 		LODModel.MeshToImportVertexMap.Append(MoveTemp(UniqueIndices));
