Cloth Weight Map Paint Tool:
- Add a new Smooth brush that mimics the Smooth mode in the SkeletalMesh Editor's Cloth Paint tool
- Add a Strength scalar to the existing Paint brush and differntiate between the painted value and the brush strength
- Set Flow to zero and enable IgnoreZeroMoevments in both brushes so that they are only activated when the mouse moves

diff --git a/Engine/Plugins/Experimental/ChaosClothAssetEditor/Source/ChaosClothAssetEditorTools/Private/ChaosClothAsset/ClothWeightMapPaintBrushOps.cpp b/Engine/Plugins/Experimental/ChaosClothAssetEditor/Source/ChaosClothAssetEditorTools/Private/ChaosClothAsset/ClothWeightMapPaintBrushOps.cpp
new file mode 100644
--- /dev/null
+++ b/Engine/Plugins/Experimental/ChaosClothAssetEditor/Source/ChaosClothAssetEditorTools/Private/ChaosClothAsset/ClothWeightMapPaintBrushOps.cpp
@@ -0,0 +1,94 @@
+// Copyright Epic Games, Inc. All Rights Reserved.
+
+#include "ChaosClothAsset/ClothWeightMapPaintBrushOps.h"
+#include "Async/ParallelFor.h"
+#include "DynamicMesh/DynamicMesh3.h"
+
+//
+// Paint Brush
+//
+
+void FWeightMapPaintBrushOp::ApplyStampByVertices(
+	const FDynamicMesh3* Mesh,
+	const FSculptBrushStamp& Stamp,
+	const TArray<int32>& Vertices,
+	TArray<double>& NewAttributesOut)
+{
+	const FVector3d& StampPos = Stamp.LocalFrame.Origin;
+
+	UWeightMapPaintBrushOpProps* const Props = GetPropertySetAs<UWeightMapPaintBrushOpProps>();
+	const double TargetValue = (double)Props->GetAttribute();
+
+	const int32 NumVertices = Vertices.Num();
+	ParallelFor(NumVertices, [&](int32 k)
+	{
+		const int32 VertIdx = Vertices[k];
+		const FVector3d OrigPos = Mesh->GetVertex(VertIdx);
+		const double ExistingValue = NewAttributesOut[k];
+		const double ValueDiff = TargetValue - ExistingValue;
+		NewAttributesOut[k] = FMath::Clamp(ExistingValue + Stamp.Power * ValueDiff, 0.0, 1.0);
+	});
+}
+
+
+//
+// Smooth Brush
+//
+
+void FWeightMapSmoothBrushOp::ApplyStampByVertices(
+	const FDynamicMesh3* Mesh,
+	const FSculptBrushStamp& Stamp,
+	const TArray<int32>& Vertices,
+	TArray<double>& VertexWeightValues)
+{
+	// Converted from FClothPaintTool_Smooth::SmoothVertices
+	
+	const int32 NumVertices = Vertices.Num();
+
+	TMap<int32, int32> VertexToBufferIndexMap;
+	for (int32 BufferIndex = 0; BufferIndex < NumVertices; ++BufferIndex)
+	{
+		const int32 VertexIndex = Vertices[BufferIndex];
+		VertexToBufferIndexMap.Add(VertexIndex, BufferIndex);
+	}
+
+	// Compute average values of one-rings for all vertices
+	TArray<double> OneRingAverages;
+	OneRingAverages.SetNumUninitialized(NumVertices);
+
+	for (int32 BufferIndex = 0; BufferIndex < NumVertices; ++BufferIndex)
+	{
+		const int32 VertexIndex = Vertices[BufferIndex];
+		double Accumulator = 0.0f;
+		int NumNeighbors = 0;
+
+		for (const int32 NeighborIndex : Mesh->VtxVerticesItr(VertexIndex))
+		{
+			if (VertexToBufferIndexMap.Contains(NeighborIndex))
+			{
+				const int32 NeighborBufferIndex = VertexToBufferIndexMap[NeighborIndex];
+				Accumulator += VertexWeightValues[NeighborBufferIndex];
+				++NumNeighbors;
+			}
+		}
+
+		if (NumNeighbors > 0)
+		{
+			OneRingAverages[BufferIndex] = Accumulator / NumNeighbors;
+		}
+		else
+		{
+			// Don't change the vertex value if it has no neighbors
+			OneRingAverages[BufferIndex] = VertexWeightValues[BufferIndex];
+		}
+	}
+
+	// Blend vertex value with its average one-ring value
+	for (int32 BufferIndex = 0; BufferIndex < NumVertices; ++BufferIndex)
+	{
+		const int32 VertexIndex = Vertices[BufferIndex];
+		const double Diff = OneRingAverages[BufferIndex] - VertexWeightValues[BufferIndex];
+		VertexWeightValues[BufferIndex] += Stamp.Power * Diff;
+	}
+}
+
diff --git a/Engine/Plugins/Experimental/ChaosClothAssetEditor/Source/ChaosClothAssetEditorTools/Private/ChaosClothAsset/ClothWeightMapPaintBrushOps.h b/Engine/Plugins/Experimental/ChaosClothAssetEditor/Source/ChaosClothAssetEditorTools/Private/ChaosClothAsset/ClothWeightMapPaintBrushOps.h
--- a/Engine/Plugins/Experimental/ChaosClothAssetEditor/Source/ChaosClothAssetEditorTools/Private/ChaosClothAsset/ClothWeightMapPaintBrushOps.h
+++ b/Engine/Plugins/Experimental/ChaosClothAssetEditor/Source/ChaosClothAssetEditorTools/Private/ChaosClothAsset/ClothWeightMapPaintBrushOps.h
@@ -3,9 +3,6 @@
 #pragma once
 
 #include "Sculpting/MeshBrushOpBase.h"
-#include "DynamicMesh/DynamicMesh3.h"
-#include "MeshWeights.h"
-#include "Async/ParallelFor.h"
 #include "ClothWeightMapPaintBrushOps.generated.h"
 
 
@@ -36,8 +33,9 @@ public:
 
 
 
-
-
+//
+// Erase Brush
+// 
 
 UCLASS()
 class CHAOSCLOTHASSETEDITORTOOLS_API UWeightMapEraseBrushOpProps : public UMeshSculptBrushOpProps
@@ -81,45 +79,82 @@ public:
 };
 
 
-
+//
+// Paint Brush
+// 
 
 UCLASS()
 class CHAOSCLOTHASSETEDITORTOOLS_API UWeightMapPaintBrushOpProps : public UMeshSculptBrushOpProps
 {
 	GENERATED_BODY()
+
 public:
+
 	/** The Attribute that will be assigned to triangles within the brush region */
-	UPROPERTY(EditAnywhere, Category = PaintBrush, meta = (DisplayName = "Attribute", UIMin = 0))
+	UPROPERTY(EditAnywhere, Category = PaintBrush, meta = (DisplayName = "Value", UIMin = 0))
 	double AttributeValue = 1;
 
+	/** Strength of the Brush */
+	UPROPERTY(EditAnywhere, Category = PaintBrush, meta = (DisplayName = "Strength", UIMin = "0.0", UIMax = "10.", ClampMin = "0.0", ClampMax = "10."))
+	float Strength = 0.5;
+
 	virtual double GetAttribute() { return AttributeValue; }
-};
 
+	virtual float GetStrength() override { return Strength; }
+	virtual void SetStrength(float NewStrength) override { Strength = NewStrength;  }
+};
 
 
 class FWeightMapPaintBrushOp : public FMeshVertexWeightMapEditBrushOp
 {
 public:
 
+	virtual bool IgnoreZeroMovements() const override
+	{
+		return true;
+	}
+
 	virtual void ApplyStampByVertices(
 		const FDynamicMesh3* Mesh,
 		const FSculptBrushStamp& Stamp,
 		const TArray<int32>& Vertices,
-		TArray<double>& NewAttributesOut) override
-	{
-		const FVector3d& StampPos = Stamp.LocalFrame.Origin;
+		TArray<double>& NewAttributesOut) override;
 
-		UWeightMapPaintBrushOpProps* Props = GetPropertySetAs<UWeightMapPaintBrushOpProps>();
-		double SetToAttribute = (double)Props->GetAttribute();
+};
 
-		// TODO: Add something here to get the old value so we can add (clamped) the AttributeValue from it.
+//
+// Smooth Brush
+// 
 
-		// TODO: Handle the stamp's properties for fall off, etc..
+UCLASS()
+class CHAOSCLOTHASSETEDITORTOOLS_API UWeightMapSmoothBrushOpProps : public UMeshSculptBrushOpProps
+{
+	GENERATED_BODY()
 
-		int32 NumVertices = Vertices.Num();
-		for ( int32 k = 0; k < NumVertices; ++k)
-		{
-			NewAttributesOut[k] = FMath::Clamp(SetToAttribute, 0.0, 1.0);
-		}
+public:
+
+	UPROPERTY(EditAnywhere, Category = PaintBrush, meta = (DisplayName = "Strength", UIMin = "0.0", UIMax = "10.", ClampMin = "0.0", ClampMax = "10."))
+	float Strength = 0.5;
+
+	virtual float GetStrength() override { return Strength; }
+	virtual void SetStrength(float NewStrength) override { Strength = NewStrength; }
+
+};
+
+
+class FWeightMapSmoothBrushOp : public FMeshVertexWeightMapEditBrushOp
+{
+public:
+
+	virtual bool IgnoreZeroMovements() const override
+	{
+		return true;
 	}
+
+	virtual void ApplyStampByVertices(
+		const FDynamicMesh3* Mesh,
+		const FSculptBrushStamp& Stamp,
+		const TArray<int32>& Vertices,
+		TArray<double>& NewAttributesOut) override;
+
 };
diff --git a/Engine/Plugins/Experimental/ChaosClothAssetEditor/Source/ChaosClothAssetEditorTools/Private/ChaosClothAsset/ClothWeightMapPaintTool.cpp b/Engine/Plugins/Experimental/ChaosClothAssetEditor/Source/ChaosClothAssetEditorTools/Private/ChaosClothAsset/ClothWeightMapPaintTool.cpp
--- a/Engine/Plugins/Experimental/ChaosClothAssetEditor/Source/ChaosClothAssetEditorTools/Private/ChaosClothAsset/ClothWeightMapPaintTool.cpp
+++ b/Engine/Plugins/Experimental/ChaosClothAssetEditor/Source/ChaosClothAssetEditorTools/Private/ChaosClothAsset/ClothWeightMapPaintTool.cpp
@@ -174,7 +174,8 @@ void UClothEditorWeightMapPaintTool::Setup()
 		[this](float NewSize) { UMeshSculptToolBase::BrushProperties->BrushSize.AdaptiveSize = NewSize; });
 	FilterProperties->RestoreProperties(this);
 	FilterProperties->BrushSize = UMeshSculptToolBase::BrushProperties->BrushSize.AdaptiveSize;
-	FilterProperties->StrengthValue = 1.0;
+	FilterProperties->Strength = 1.0;
+	FilterProperties->AttributeValue = 1.0;
 	AddToolPropertySource(FilterProperties);
 
 	InitializeIndicator();
@@ -184,6 +185,7 @@ void UClothEditorWeightMapPaintTool::Setup()
 	UMeshSculptToolBase::BrushProperties->bShowPerBrushProps = false;
 	UMeshSculptToolBase::BrushProperties->bShowFalloff = true;
 	UMeshSculptToolBase::BrushProperties->bShowLazyness = false;
+	UMeshSculptToolBase::BrushProperties->FlowRate = 0.0f;
 	CalculateBrushRadius();
 
 	PaintBrushOpProperties = NewObject<UWeightMapPaintBrushOpProps>(this);
@@ -191,6 +193,11 @@ void UClothEditorWeightMapPaintTool::Setup()
 		MakeUnique<FLambdaMeshSculptBrushOpFactory>([this]() { return MakeUnique<FWeightMapPaintBrushOp>(); }),
 		PaintBrushOpProperties);
 
+	SmoothBrushOpProperties = NewObject<UWeightMapSmoothBrushOpProps>(this);
+	RegisterBrushType((int32)EClothEditorWeightMapPaintBrushType::Smooth, LOCTEXT("SmoothBrushType", "Smooth"),
+		MakeUnique<FLambdaMeshSculptBrushOpFactory>([this]() { return MakeUnique<FWeightMapSmoothBrushOp>(); }),
+		SmoothBrushOpProperties);
+
 	// secondary brushes
 	EraseBrushOpProperties = NewObject<UWeightMapEraseBrushOpProps>(this);
 
@@ -425,11 +432,16 @@ void UClothEditorWeightMapPaintTool::OnBeginStroke(const FRay& WorldRay)
 
 	if (PaintBrushOpProperties)
 	{
-		PaintBrushOpProperties->AttributeValue = FilterProperties->StrengthValue;
+		PaintBrushOpProperties->AttributeValue = FilterProperties->AttributeValue;
+		PaintBrushOpProperties->Strength = FilterProperties->Strength;
 	}
 	if (EraseBrushOpProperties)
 	{
-		EraseBrushOpProperties->AttributeValue = FilterProperties->StrengthValue;
+		EraseBrushOpProperties->AttributeValue = 0.0;
+	}
+	if (SmoothBrushOpProperties)
+	{
+		SmoothBrushOpProperties->Strength = FilterProperties->Strength;
 	}
 
 	// initialize first "Last Stamp", so that we can assume all stamps in stroke have a valid previous stamp
@@ -608,7 +620,8 @@ bool UClothEditorWeightMapPaintTool::UpdateStampPosition(const FRay& WorldRay)
 	CurrentStamp.PrevWorldFrame = LastStamp.WorldFrame;
 
 	FVector3d MoveDelta = CurrentStamp.LocalFrame.Origin - CurrentStamp.PrevLocalFrame.Origin;
-	if (UseBrushOp->IgnoreZeroMovements() && MoveDelta.SquaredLength() < FMathd::ZeroTolerance)
+
+	if (UseBrushOp->IgnoreZeroMovements() && MoveDelta.SquaredLength() < 0.1 * CurrentBrushRadius)
 	{
 		return false;
 	}
@@ -887,7 +900,7 @@ void UClothEditorWeightMapPaintTool::OnPolyLassoFinished(const FCameraPolyLasso&
 	FMeshVertexSelection VertexSelection(Mesh);
 	VertexSelection.SelectByVertexID([&](int32 vid) { return TempROIBuffer[vid] == 1; });
 
-	double SetWeightValue = GetInEraseStroke() ? 0.0 : FilterProperties->StrengthValue;
+	double SetWeightValue = GetInEraseStroke() ? 0.0 : FilterProperties->AttributeValue;
 	SetVerticesToWeightMap(VertexSelection.AsSet(), SetWeightValue, GetInEraseStroke());
 }
 
@@ -1104,6 +1117,11 @@ void UClothEditorWeightMapPaintTool::DrawHUD(FCanvas* Canvas, IToolsContextRende
 	}
 }
 
+void UClothEditorWeightMapPaintTool::UpdateStampPendingState()
+{
+	if (InStroke() == false) return;
+	bIsStampPending = true;
+}
 
 void UClothEditorWeightMapPaintTool::OnTick(float DeltaTime)
 {
@@ -1149,7 +1167,7 @@ void UClothEditorWeightMapPaintTool::OnTick(float DeltaTime)
 				double HitWeightValue = GetCurrentWeightValueUnderBrush();
 				if (bPendingPickWeight)
 				{
-					FilterProperties->StrengthValue = HitWeightValue;
+					FilterProperties->AttributeValue = HitWeightValue;
 					NotifyOfPropertyChangeByTool(FilterProperties);
 				}
 			}
@@ -1239,7 +1257,7 @@ void UClothEditorWeightMapPaintTool::FloodFillCurrentWeightAction()
 
 	BeginChange();
 
-	float SetWeightValue = FilterProperties->StrengthValue;
+	const float SetWeightValue = FilterProperties->AttributeValue;
 	const FDynamicMesh3* Mesh = DynamicMeshComponent->GetMesh();
 	TempROIBuffer.SetNum(0, false);
 	for (int32 vid : Mesh->VertexIndicesItr())
@@ -1578,7 +1596,7 @@ void UClothEditorWeightMapPaintTool::UpdateSubToolType(EClothEditorWeightMapPain
 
 void UClothEditorWeightMapPaintTool::UpdateBrushType(EClothEditorWeightMapPaintBrushType BrushType)
 {
-	static const FText BaseMessage = LOCTEXT("OnStartTool", "Hold Shift to Erase. [/] and S/D change Size (+Shift to small-step). Shift+Q for New Group, Shift+G to pick Group, Shift+F to Freeze Group.");
+	static const FText BaseMessage = LOCTEXT("OnStartTool", "Hold Shift to Erase. [/] and S/D change Size (+Shift to small-step)");
 	FTextBuilder Builder;
 	Builder.AppendLine(BaseMessage);
 
diff --git a/Engine/Plugins/Experimental/ChaosClothAssetEditor/Source/ChaosClothAssetEditorTools/Public/ChaosClothAsset/ClothWeightMapPaintTool.h b/Engine/Plugins/Experimental/ChaosClothAssetEditor/Source/ChaosClothAssetEditorTools/Public/ChaosClothAsset/ClothWeightMapPaintTool.h
--- a/Engine/Plugins/Experimental/ChaosClothAssetEditor/Source/ChaosClothAssetEditorTools/Public/ChaosClothAsset/ClothWeightMapPaintTool.h
+++ b/Engine/Plugins/Experimental/ChaosClothAssetEditor/Source/ChaosClothAssetEditorTools/Public/ChaosClothAsset/ClothWeightMapPaintTool.h
@@ -26,6 +26,7 @@
 class UMeshElementsVisualizer;
 class UWeightMapEraseBrushOpProps;
 class UWeightMapPaintBrushOpProps;
+class UWeightMapSmoothBrushOpProps;
 class UClothEditorContextObject;
 
 DECLARE_STATS_GROUP(TEXT("WeightMapPaintTool"), STATGROUP_WeightMapPaintTool, STATCAT_Advanced);
@@ -84,6 +85,9 @@ enum class EClothEditorWeightMapPaintBrushType : uint8
 	/** Paint weights */
 	Paint UMETA(DisplayName = "Paint"),
 
+	/** Smooth existing weights */
+	Smooth UMETA(DisplayName = "Smooth"),
+
 	/** Erase weights */
 	Erase UMETA(Hidden, DisplayName = "Erase"),
 
@@ -132,9 +136,12 @@ public:
 		HideEditConditionToggle, EditConditionHides, EditCondition = "SubToolType != EClothEditorWeightMapPaintInteractionType::PolyLasso"))
 	float BrushSize = 0.25f;
 
-	/** The weight value that will be assigned to vertices */
 	UPROPERTY(EditAnywhere, Category = ActionType, meta = (UIMin = 0, ClampMin = 0, UIMax = 1, ClampMax = 1))
-	double StrengthValue = 1;
+	double AttributeValue = 1;
+
+	UPROPERTY(EditAnywhere, Category = ActionType, meta = (UIMin = 0, ClampMin = 0, UIMax = 1, ClampMax = 1))
+	double Strength = 1;
+
 
 	/** The Region affected by the current operation will be bounded by edge angles larger than this threshold */
 	UPROPERTY(EditAnywhere, Category = Filters, meta = (UIMin = "0.0", UIMax = "180.0", EditCondition = "SubToolType != EClothEditorWeightMapPaintInteractionType::PolyLasso && BrushAreaMode == EClothEditorWeightMapPaintBrushAreaType::Connected"))
@@ -240,6 +247,8 @@ public:
 	virtual void Shutdown(EToolShutdownType ShutdownType) override;
 	virtual void DrawHUD(FCanvas* Canvas, IToolsContextRenderAPI* RenderAPI) override;
 
+	virtual void UpdateStampPendingState() override;
+
 	virtual void OnTick(float DeltaTime) override;
 
 	virtual bool HasCancel() const override { return true; }
@@ -268,6 +277,9 @@ private:
 	UPROPERTY()
 	TObjectPtr<UWeightMapPaintBrushOpProps> PaintBrushOpProperties;
 
+	UPROPERTY()
+	TObjectPtr<UWeightMapSmoothBrushOpProps> SmoothBrushOpProperties;
+
 	UPROPERTY()
 	TObjectPtr<UWeightMapEraseBrushOpProps> EraseBrushOpProperties;
 
