Cloth Paint Tool:
- Add a Gradient Interaction Type. The user selects "high" and "low" vertices and a gradient is automatically computed between them

diff --git a/Engine/Plugins/Experimental/ChaosClothAssetEditor/Source/ChaosClothAssetEditorTools/Private/ChaosClothAsset/ClothWeightMapPaintTool.cpp b/Engine/Plugins/Experimental/ChaosClothAssetEditor/Source/ChaosClothAssetEditorTools/Private/ChaosClothAsset/ClothWeightMapPaintTool.cpp
--- a/Engine/Plugins/Experimental/ChaosClothAssetEditor/Source/ChaosClothAssetEditorTools/Private/ChaosClothAsset/ClothWeightMapPaintTool.cpp
+++ b/Engine/Plugins/Experimental/ChaosClothAssetEditor/Source/ChaosClothAssetEditorTools/Private/ChaosClothAsset/ClothWeightMapPaintTool.cpp
@@ -41,6 +41,7 @@
 #include "Dataflow/DataflowGraphEditor.h"
 #include "Widgets/Notifications/SNotificationList.h"
 #include "Framework/Notifications/NotificationManager.h"
+#include "Selection/PolygonSelectionMechanic.h"
 
 #include UE_INLINE_GENERATED_CPP_BY_NAME(ClothWeightMapPaintTool)
 
@@ -158,6 +159,28 @@ void UClothEditorWeightMapPaintTool::Setup()
 	PolyLassoMechanic->SpacingTolerance = 10.0f;
 	PolyLassoMechanic->OnDrawPolyLassoFinished.AddUObject(this, &UClothEditorWeightMapPaintTool::OnPolyLassoFinished);
 
+
+	// Set up vertex selection mechanic
+	PolygonSelectionMechanic = NewObject<UPolygonSelectionMechanic>(this);
+	PolygonSelectionMechanic->bAddSelectionFilterPropertiesToParentTool = false;
+	PolygonSelectionMechanic->Setup(this);
+	PolygonSelectionMechanic->SetIsEnabled(false);
+	PolygonSelectionMechanic->OnSelectionChanged.AddUObject(this, &UClothEditorWeightMapPaintTool::OnSelectionModified);
+
+	// disable CTRL to remove from selection
+	PolygonSelectionMechanic->SetShouldRemoveFromSelectionFunc([]() { return false; });
+
+	PolygonSelectionMechanic->Properties->bSelectEdges = false;
+	PolygonSelectionMechanic->Properties->bSelectFaces = false;
+	PolygonSelectionMechanic->Properties->bSelectVertices = true;
+	
+	constexpr bool bAutoBuild = true;
+	GradientSelectionTopology = MakeUnique<UE::Geometry::FTriangleGroupTopology>(DynamicMeshComponent->GetMesh(), bAutoBuild);
+	
+	MeshSpatial = MakeUnique<UE::Geometry::FDynamicMeshAABBTree3>(DynamicMeshComponent->GetMesh(), bAutoBuild);
+	PolygonSelectionMechanic->Initialize(DynamicMeshComponent, GradientSelectionTopology.Get(), [this]() { return MeshSpatial.Get(); });
+	
+
 	AddWeightMapProperties = NewObject<UClothEditorAddWeightMapProperties>(this);
 	AddWeightMapProperties->Initialize(this);
 	AddToolPropertySource(AddWeightMapProperties);
@@ -332,6 +355,12 @@ void UClothEditorWeightMapPaintTool::Shutdown(EToolShutdownType ShutdownType)
 		PreviewMeshActor = nullptr;
 	}
 
+	if (PolygonSelectionMechanic)
+	{
+		PolygonSelectionMechanic->Shutdown();
+		PolygonSelectionMechanic = nullptr;
+	}
+
 	UMeshSculptToolBase::Shutdown(ShutdownType);
 }
 
@@ -905,6 +934,96 @@ void UClothEditorWeightMapPaintTool::OnPolyLassoFinished(const FCameraPolyLasso&
 }
 
 
+void UClothEditorWeightMapPaintTool::ComputeGradient()
+{
+	if (!ensure(ActiveWeightMap))
+	{
+		UE_LOG(LogTemp, Warning, TEXT("No active weight map"));
+		return;
+	}
+
+	BeginChange();
+
+	const FDynamicMesh3* const Mesh = DynamicMeshComponent->GetMesh();
+	TempROIBuffer.SetNum(0, false);
+	for (int32 vid : Mesh->VertexIndicesItr())
+	{
+		TempROIBuffer.Add(vid);
+	}
+	for (int32 vid : TempROIBuffer)
+	{
+		ActiveWeightEditChangeTracker->SaveVertexOneRingTriangles(vid, true);
+	}
+
+
+	for (const int32 VertexIndex : TempROIBuffer)
+	{
+		const FVector3d Vert = Mesh->GetVertex(VertexIndex);
+
+		// (Copied from FClothPaintTool_Gradient::ApplyGradient)
+
+		// Get distances
+		// TODO: Look into surface distance instead of 3D distance? May be necessary for some complex shapes
+		float DistanceToLowSq = MAX_flt;
+		for (const int32& LowIndex : LowValueGradientVertexSelection.SelectedCornerIDs)
+		{
+			const FVector3d LowPoint = Mesh->GetVertex(LowIndex);
+			const float DistanceSq = (LowPoint - Vert).SizeSquared();
+			if (DistanceSq < DistanceToLowSq)
+			{
+				DistanceToLowSq = DistanceSq;
+			}
+		}
+
+		float DistanceToHighSq = MAX_flt;
+		for (const int32& HighIndex : HighValueGradientVertexSelection.SelectedCornerIDs)
+		{
+			const FVector3d HighPoint = Mesh->GetVertex(HighIndex);
+			const float DistanceSq = (HighPoint - Vert).SizeSquared();
+			if (DistanceSq < DistanceToHighSq)
+			{
+				DistanceToHighSq = DistanceSq;
+			}
+		}
+
+		const float Value = FMath::LerpStable(0.0f, 1.0f, DistanceToLowSq / (DistanceToLowSq + DistanceToHighSq));
+		ActiveWeightMap->SetValue(VertexIndex, &Value);
+	}
+
+
+	DynamicMeshComponent->FastNotifyVertexAttributesUpdated(EMeshRenderAttributeFlags::VertexColors);
+	GetToolManager()->PostInvalidation();
+	EndChange();
+
+}
+
+
+void UClothEditorWeightMapPaintTool::OnSelectionModified()
+{
+	const bool bToolTypeIsGradient = (FilterProperties->SubToolType == EClothEditorWeightMapPaintInteractionType::Gradient);
+	if (bToolTypeIsGradient && PolygonSelectionMechanic)
+	{
+		const FGroupTopologySelection NewSelection = PolygonSelectionMechanic->GetActiveSelection();
+
+		const bool bSelectingLowValueGradientVertices = GetCtrlToggle();
+		if (bSelectingLowValueGradientVertices)
+		{
+			LowValueGradientVertexSelection = NewSelection;
+		}
+		else
+		{
+			HighValueGradientVertexSelection = NewSelection;
+		}
+
+		if (LowValueGradientVertexSelection.SelectedCornerIDs.Num() > 0 && HighValueGradientVertexSelection.SelectedCornerIDs.Num() > 0)
+		{
+			ComputeGradient();
+		}
+
+		constexpr bool bBroadcast = false;
+		PolygonSelectionMechanic->SetSelection(FGroupTopologySelection(), bBroadcast);
+	}
+}
 
 
 void UClothEditorWeightMapPaintTool::SetVerticesToWeightMap(const TSet<int32>& Vertices, double WeightValue, bool bIsErase)
@@ -1115,6 +1234,47 @@ void UClothEditorWeightMapPaintTool::DrawHUD(FCanvas* Canvas, IToolsContextRende
 		PolyLassoMechanic->LineColor = GetInEraseStroke() ? FLinearColor::Red : FLinearColor::Green;
 		PolyLassoMechanic->DrawHUD(Canvas, RenderAPI);
 	}
+
+	if (PolygonSelectionMechanic)
+	{
+		PolygonSelectionMechanic->DrawHUD(Canvas, RenderAPI);
+	}
+}
+
+void UClothEditorWeightMapPaintTool::Render(IToolsContextRenderAPI* RenderAPI)
+{
+	UMeshSculptToolBase::Render(RenderAPI);
+
+	if (PolygonSelectionMechanic)
+	{
+		PolygonSelectionMechanic->RenderMarquee(RenderAPI);
+
+		const FViewCameraState RenderCameraState = RenderAPI->GetCameraState();
+		GradientSelectionRenderer.BeginFrame(RenderAPI, RenderCameraState);
+
+		const FTransform Transform = DynamicMeshComponent->GetComponentTransform();
+		GradientSelectionRenderer.SetTransform(Transform);
+			
+		GradientSelectionRenderer.SetPointParameters(FLinearColor::Red, 1.0);
+		PolygonSelectionMechanic->UMeshTopologySelectionMechanic::GetTopologySelector()->DrawSelection(LowValueGradientVertexSelection, &GradientSelectionRenderer, &RenderCameraState);
+
+		GradientSelectionRenderer.SetPointParameters(FLinearColor::Green, 1.0);
+		PolygonSelectionMechanic->UMeshTopologySelectionMechanic::GetTopologySelector()->DrawSelection(HighValueGradientVertexSelection, &GradientSelectionRenderer, &RenderCameraState);
+
+		// Now the current unsaved selection
+		if (GetCtrlToggle())
+		{
+			GradientSelectionRenderer.SetPointParameters(FLinearColor::Red, 1.0);
+		}
+		else
+		{
+			GradientSelectionRenderer.SetPointParameters(FLinearColor::Green, 1.0);
+		}
+
+		PolygonSelectionMechanic->UMeshTopologySelectionMechanic::GetTopologySelector()->DrawSelection(PolygonSelectionMechanic->GetActiveSelection(), &GradientSelectionRenderer, &RenderCameraState);
+
+		GradientSelectionRenderer.EndFrame();
+	}
 }
 
 void UClothEditorWeightMapPaintTool::UpdateStampPendingState()
@@ -1128,9 +1288,14 @@ void UClothEditorWeightMapPaintTool::OnTick(float DeltaTime)
 	UMeshSculptToolBase::OnTick(DeltaTime);
 	MeshElementsDisplay->OnTick(DeltaTime);
 
-	bool bIsLasso = (FilterProperties->SubToolType == EClothEditorWeightMapPaintInteractionType::PolyLasso);
+	const bool bIsLasso = (FilterProperties->SubToolType == EClothEditorWeightMapPaintInteractionType::PolyLasso);
 	PolyLassoMechanic->SetIsEnabled(bIsLasso);
 
+	const bool bIsGradient = (FilterProperties->SubToolType == EClothEditorWeightMapPaintInteractionType::Gradient);
+	PolygonSelectionMechanic->SetIsEnabled(bIsGradient);
+
+	check(!(bIsLasso && bIsGradient));
+
 	ConfigureIndicator(false);
 	SetIndicatorVisibility(bIsLasso == false);
 
@@ -1591,6 +1756,12 @@ void UClothEditorWeightMapPaintTool::UpdateSubToolType(EClothEditorWeightMapPain
 
 	SetToolPropertySourceEnabled(FilterProperties, true);
 	SetBrushOpPropsVisibility(false);
+
+	if (NewType != EClothEditorWeightMapPaintInteractionType::Gradient)
+	{
+		LowValueGradientVertexSelection.Clear();
+		HighValueGradientVertexSelection.Clear();
+	}
 }
 
 
diff --git a/Engine/Plugins/Experimental/ChaosClothAssetEditor/Source/ChaosClothAssetEditorTools/Public/ChaosClothAsset/ClothWeightMapPaintTool.h b/Engine/Plugins/Experimental/ChaosClothAssetEditor/Source/ChaosClothAssetEditorTools/Public/ChaosClothAsset/ClothWeightMapPaintTool.h
--- a/Engine/Plugins/Experimental/ChaosClothAssetEditor/Source/ChaosClothAssetEditorTools/Public/ChaosClothAsset/ClothWeightMapPaintTool.h
+++ b/Engine/Plugins/Experimental/ChaosClothAssetEditor/Source/ChaosClothAssetEditorTools/Public/ChaosClothAsset/ClothWeightMapPaintTool.h
@@ -20,6 +20,8 @@
 #include "DynamicMesh/DynamicMeshChangeTracker.h"
 #include "DynamicMesh/MeshNormals.h"
 #include "TransformTypes.h"
+#include "ToolDataVisualizer.h"
+#include "GroupTopology.h"
 
 #include "ClothWeightMapPaintTool.generated.h"
 
@@ -28,6 +30,7 @@ class UWeightMapEraseBrushOpProps;
 class UWeightMapPaintBrushOpProps;
 class UWeightMapSmoothBrushOpProps;
 class UClothEditorContextObject;
+class UPolygonSelectionMechanic;
 
 DECLARE_STATS_GROUP(TEXT("WeightMapPaintTool"), STATGROUP_WeightMapPaintTool, STATCAT_Advanced);
 DECLARE_CYCLE_STAT(TEXT("WeightMapPaintTool_UpdateROI"), WeightMapPaintTool_UpdateROI, STATGROUP_WeightMapPaintTool);
@@ -68,6 +71,7 @@ enum class EClothEditorWeightMapPaintInteractionType : uint8
 	Brush,
 	Fill,
 	PolyLasso,
+	Gradient,
 
 	LastValue UMETA(Hidden)
 };
@@ -246,6 +250,7 @@ public:
 	virtual void Setup() override;
 	virtual void Shutdown(EToolShutdownType ShutdownType) override;
 	virtual void DrawHUD(FCanvas* Canvas, IToolsContextRenderAPI* RenderAPI) override;
+	virtual void Render(IToolsContextRenderAPI* RenderAPI) override;
 
 	virtual void UpdateStampPendingState() override;
 
@@ -350,6 +355,24 @@ protected:
 	void OnPolyLassoFinished(const FCameraPolyLasso& Lasso, bool bCanceled);
 
 
+	// 
+	// Gradient Support
+	//
+	UPROPERTY()
+	TObjectPtr<UPolygonSelectionMechanic> PolygonSelectionMechanic;
+
+	TUniquePtr<UE::Geometry::FDynamicMeshAABBTree3> MeshSpatial = nullptr;
+
+	TUniquePtr<UE::Geometry::FTriangleGroupTopology> GradientSelectionTopology = nullptr;
+
+	FToolDataVisualizer GradientSelectionRenderer;
+
+	UE::Geometry::FGroupTopologySelection LowValueGradientVertexSelection;
+	UE::Geometry::FGroupTopologySelection HighValueGradientVertexSelection;
+
+	void ComputeGradient();
+	void OnSelectionModified();
+
 	//
 	// Internals
 	//
diff --git a/Engine/Plugins/Runtime/MeshModelingToolset/Source/ModelingComponents/Private/Selection/MeshTopologySelectionMechanic.cpp b/Engine/Plugins/Runtime/MeshModelingToolset/Source/ModelingComponents/Private/Selection/MeshTopologySelectionMechanic.cpp
--- a/Engine/Plugins/Runtime/MeshModelingToolset/Source/ModelingComponents/Private/Selection/MeshTopologySelectionMechanic.cpp
+++ b/Engine/Plugins/Runtime/MeshModelingToolset/Source/ModelingComponents/Private/Selection/MeshTopologySelectionMechanic.cpp
@@ -599,6 +599,13 @@ void UMeshTopologySelectionMechanic::UpdateMarqueeEnabled()
 			|| Properties->bSelectEdgeLoops || Properties->bSelectEdgeRings));
 }
 
+void UMeshTopologySelectionMechanic::RenderMarquee(IToolsContextRenderAPI* RenderAPI)
+{
+	if (MarqueeMechanic)
+	{
+		MarqueeMechanic->Render(RenderAPI);
+	}
+}
 
 void UMeshTopologySelectionMechanic::BeginChange()
 {
diff --git a/Engine/Plugins/Runtime/MeshModelingToolset/Source/ModelingComponents/Public/Selection/MeshTopologySelectionMechanic.h b/Engine/Plugins/Runtime/MeshModelingToolset/Source/ModelingComponents/Public/Selection/MeshTopologySelectionMechanic.h
--- a/Engine/Plugins/Runtime/MeshModelingToolset/Source/ModelingComponents/Public/Selection/MeshTopologySelectionMechanic.h
+++ b/Engine/Plugins/Runtime/MeshModelingToolset/Source/ModelingComponents/Public/Selection/MeshTopologySelectionMechanic.h
@@ -339,6 +339,12 @@ public:
 	 */
 	bool IsCurrentlyMarqueeDragging() { return bCurrentlyMarqueeDragging; }
 
+	/**
+	* Render only the MarqueeMechanic, without rendering the current selection
+	*/
+	void RenderMarquee(IToolsContextRenderAPI* RenderAPI);
+	
+
 	// TODO: Is it worth issuing separate callbacks in normal selection changes and in FChange ones, to
 	// allow the user to bundle in some FChanges into the normal callback?
 
diff --git a/Engine/Source/Runtime/InteractiveToolsFramework/Public/BaseTools/MeshSurfacePointTool.h b/Engine/Source/Runtime/InteractiveToolsFramework/Public/BaseTools/MeshSurfacePointTool.h
--- a/Engine/Source/Runtime/InteractiveToolsFramework/Public/BaseTools/MeshSurfacePointTool.h
+++ b/Engine/Source/Runtime/InteractiveToolsFramework/Public/BaseTools/MeshSurfacePointTool.h
@@ -125,10 +125,10 @@ public:
 	/** @return current state of the shift toggle */
 	virtual bool GetShiftToggle() const { return bShiftToggle; }
 
-	/** Called by registered InputBehaviors to set the state of the "shift" button (or device equivalent) */
+	/** Called by registered InputBehaviors to set the state of the "ctrl" button (or device equivalent) */
 	virtual void SetCtrlToggle(bool bCtrlDown);
 
-	/** @return current state of the shift toggle */
+	/** @return current state of the ctrl toggle */
 	virtual bool GetCtrlToggle() const { return bCtrlToggle; }
 
 	/** @return current input device pressure in range 0-1 */
