Cloth Editor: Fix occasional crash in DataflowComponent rendering. We were adding RenderTargets to the Dataflow object when the node selection changed,  but never removing them. This CL updates the RenderTargets when the selection changes, and also when nodes are deleted from the graph.

diff --git a/Engine/Plugins/Experimental/ChaosClothAssetEditor/Source/ChaosClothAssetEditor/Private/ChaosClothAsset/ClothEditorMode.cpp b/Engine/Plugins/Experimental/ChaosClothAssetEditor/Source/ChaosClothAssetEditor/Private/ChaosClothAsset/ClothEditorMode.cpp
--- a/Engine/Plugins/Experimental/ChaosClothAssetEditor/Source/ChaosClothAssetEditor/Private/ChaosClothAsset/ClothEditorMode.cpp
+++ b/Engine/Plugins/Experimental/ChaosClothAssetEditor/Source/ChaosClothAssetEditor/Private/ChaosClothAsset/ClothEditorMode.cpp
@@ -320,6 +320,13 @@ void UChaosClothAssetEditorMode::Exit()
 	{
 		WireframeDisplay->Disconnect();
 	}
+
+	if (DataflowComponent)
+	{
+		DataflowComponent->UnregisterComponent();
+		DataflowComponent->DestroyComponent();
+	}
+
 	WireframesToTick.Reset();
 	DynamicMeshComponents.Reset();
 	PropertyObjectsToTick.Empty();
diff --git a/Engine/Plugins/Experimental/ChaosClothAssetEditor/Source/ChaosClothAssetEditor/Private/ChaosClothAsset/ClothEditorToolkit.cpp b/Engine/Plugins/Experimental/ChaosClothAssetEditor/Source/ChaosClothAssetEditor/Private/ChaosClothAsset/ClothEditorToolkit.cpp
--- a/Engine/Plugins/Experimental/ChaosClothAssetEditor/Source/ChaosClothAssetEditor/Private/ChaosClothAsset/ClothEditorToolkit.cpp
+++ b/Engine/Plugins/Experimental/ChaosClothAssetEditor/Source/ChaosClothAssetEditor/Private/ChaosClothAsset/ClothEditorToolkit.cpp
@@ -809,6 +809,7 @@ TSharedRef<SDataflowGraphEditor> FChaosClothAssetEditorToolkit::CreateGraphEdito
 		.EvaluateGraph(EvalLambda);
 
 	NewGraphEditor->OnSelectionChangedMulticast.AddSP(this, &FChaosClothAssetEditorToolkit::OnNodeSelectionChanged);
+	NewGraphEditor->OnNodeDeletedMulticast.AddSP(this, &FChaosClothAssetEditorToolkit::OnNodeDeleted);
 
 	return NewGraphEditor;
 }
@@ -935,9 +936,14 @@ void FChaosClothAssetEditorToolkit::OnNodeSelectionChanged(const TSet<UObject*>&
 
 	TSharedPtr<FManagedArrayCollection> Collection = nullptr;
 
+	// Get any selected node with a ClothCollection output
+	// Also, set the selected node(s) to be the Dataflow's RenderTargets
+	// TODO: decide if we want selection to be the mechanism for toggling DataflowComponent rendering, or the switch on the Node
+
 	if (Dataflow)
 	{
-		// Get any selected node with a ClothCollection output
+		Dataflow->RenderTargets.Reset();
+
 		for (const UObject* const Selected : NewSelection)
 		{
 			if (const UDataflowEdNode* const Node = Cast<UDataflowEdNode>(Selected))
@@ -960,9 +966,30 @@ void FChaosClothAssetEditorToolkit::OnNodeSelectionChanged(const TSet<UObject*>&
 		ClothMode->SetSelectedClothCollection(Collection);
 		ClothMode->RefocusRestSpaceViewportClient();
 	}
+}
+
+
+void FChaosClothAssetEditorToolkit::OnNodeDeleted(const TSet<UObject*>& DeletedNodes) const
+{
+	if (Dataflow)
+	{
+		Dataflow->RenderTargets.SetNum(Algo::RemoveIf(Dataflow->RenderTargets, [&DeletedNodes](const UDataflowEdNode* RenderTarget)
+		{
+			return DeletedNodes.Contains(RenderTarget);
+		}));
+	}
+
+
+	UChaosClothAssetEditorMode* const ClothMode = CastChecked<UChaosClothAssetEditorMode>(EditorModeManager->GetActiveScriptableMode(UChaosClothAssetEditorMode::EM_ChaosClothAssetEditorModeId));
+	if (ClothMode)
+	{
+		ClothMode->SetSelectedClothCollection(nullptr);
+		ClothMode->RefocusRestSpaceViewportClient();
+	}
 
 }
 
+
 //~ Ends DataflowEditorActions
 
 void FChaosClothAssetEditorToolkit::OnClothAssetChanged()
diff --git a/Engine/Plugins/Experimental/ChaosClothAssetEditor/Source/ChaosClothAssetEditor/Private/ChaosClothAsset/ClothEditorToolkit.h b/Engine/Plugins/Experimental/ChaosClothAssetEditor/Source/ChaosClothAssetEditor/Private/ChaosClothAsset/ClothEditorToolkit.h
--- a/Engine/Plugins/Experimental/ChaosClothAssetEditor/Source/ChaosClothAssetEditor/Private/ChaosClothAsset/ClothEditorToolkit.h
+++ b/Engine/Plugins/Experimental/ChaosClothAssetEditor/Source/ChaosClothAssetEditor/Private/ChaosClothAsset/ClothEditorToolkit.h
@@ -115,7 +115,7 @@ private:
 	bool OnNodeVerifyTitleCommit(const FText& NewText, UEdGraphNode* GraphNode, FText& OutErrorMessage) const;
 	void OnNodeTitleCommitted(const FText& InNewText, ETextCommit::Type InCommitType, UEdGraphNode* GraphNode) const;
 	void OnNodeSelectionChanged(const TSet<UObject*>& NewSelection) const;
-
+	void OnNodeDeleted(const TSet<UObject*>& DeletedNodes) const;
 
 	/** Scene in which the 3D sim space preview meshes live. Ownership shared with AdvancedPreviewSettingsWidget*/
 	TSharedPtr<FChaosClothPreviewScene> ClothPreviewScene;
