Cloth Editor:
- remove ClothComponent alias in 3DViewportClient
- fix SetClothAsset on the PreviewScene to not purge the SkeletalMeshComponent if it exists

diff --git a/Engine/Plugins/Experimental/ChaosClothAssetEditor/Source/ChaosClothAssetDataflowNodes/Public/ChaosClothAsset/TransferSkinWeightsNode.h b/Engine/Plugins/Experimental/ChaosClothAssetEditor/Source/ChaosClothAssetDataflowNodes/Public/ChaosClothAsset/TransferSkinWeightsNode.h
--- a/Engine/Plugins/Experimental/ChaosClothAssetEditor/Source/ChaosClothAssetDataflowNodes/Public/ChaosClothAsset/TransferSkinWeightsNode.h
+++ b/Engine/Plugins/Experimental/ChaosClothAssetEditor/Source/ChaosClothAssetDataflowNodes/Public/ChaosClothAsset/TransferSkinWeightsNode.h
@@ -25,7 +25,7 @@ public:
 	TObjectPtr<USkeletalMesh> SkeletalMesh;
 
 	UPROPERTY(EditAnywhere, Category = "Source Mesh", Meta = (Dataflowinput, DisplayName = "SkeletalMeshLOD"))
-	int32 SkeletalMeshLOD;
+	int32 SkeletalMeshLOD = 0;
 
 	// The relative transform between the SkeletalMesh and ClothCollection
 	UPROPERTY(EditAnywhere, Category = "Source Mesh", Meta = (Dataflowinput, DisplayName = "Transform"))
diff --git a/Engine/Plugins/Experimental/ChaosClothAssetEditor/Source/ChaosClothAssetEditor/Private/ChaosClothAsset/ClothEditor3DViewportClient.cpp b/Engine/Plugins/Experimental/ChaosClothAssetEditor/Source/ChaosClothAssetEditor/Private/ChaosClothAsset/ClothEditor3DViewportClient.cpp
--- a/Engine/Plugins/Experimental/ChaosClothAssetEditor/Source/ChaosClothAssetEditor/Private/ChaosClothAsset/ClothEditor3DViewportClient.cpp
+++ b/Engine/Plugins/Experimental/ChaosClothAssetEditor/Source/ChaosClothAssetEditor/Private/ChaosClothAsset/ClothEditor3DViewportClient.cpp
@@ -148,18 +148,15 @@ void FChaosClothAssetEditor3DViewportClient::EnableRenderMeshWireframe(bool bEna
 {
 	bRenderMeshWireframe = bEnable;
 
-	if (ClothComponent)
+	if (const TSharedPtr<FChaosClothPreviewScene> PinnedClothPreviewScene = ClothPreviewScene.Pin())
 	{
-		ClothComponent->SetForceWireframe(bRenderMeshWireframe);
+		if (UChaosClothComponent* const ClothComponent = PinnedClothPreviewScene->GetClothComponent())
+		{
+			ClothComponent->SetForceWireframe(bRenderMeshWireframe);
+		}
 	}
 }
 
-void FChaosClothAssetEditor3DViewportClient::SetClothComponent(TObjectPtr<UChaosClothComponent> InClothComponent)
-{
-	ClearSelectedComponents();
-	ClothComponent = InClothComponent;
-}
-
 void FChaosClothAssetEditor3DViewportClient::SetClothEdMode(TObjectPtr<UChaosClothAssetEditorMode> InClothEdMode)
 {
 	ClothEdMode = InClothEdMode;
diff --git a/Engine/Plugins/Experimental/ChaosClothAssetEditor/Source/ChaosClothAssetEditor/Private/ChaosClothAsset/ClothEditor3DViewportClient.h b/Engine/Plugins/Experimental/ChaosClothAssetEditor/Source/ChaosClothAssetEditor/Private/ChaosClothAsset/ClothEditor3DViewportClient.h
--- a/Engine/Plugins/Experimental/ChaosClothAssetEditor/Source/ChaosClothAssetEditor/Private/ChaosClothAsset/ClothEditor3DViewportClient.h
+++ b/Engine/Plugins/Experimental/ChaosClothAssetEditor/Source/ChaosClothAssetEditor/Private/ChaosClothAsset/ClothEditor3DViewportClient.h
@@ -40,7 +40,6 @@ public:
 	void EnableRenderMeshWireframe(bool bEnable);
 	bool RenderMeshWireframeEnabled() const { return bRenderMeshWireframe; }
 
-	void SetClothComponent(TObjectPtr<UChaosClothComponent> ClothComponent);
 	void SetClothEdMode(TObjectPtr<UChaosClothAssetEditorMode> ClothEdMode);
 	void SetClothEditorToolkit(TWeakPtr<const FChaosClothAssetEditorToolkit> ClothToolkit);
 
@@ -75,8 +74,6 @@ private:
 
 	TWeakPtr<FChaosClothPreviewScene> ClothPreviewScene;
 
-	TObjectPtr<UChaosClothComponent> ClothComponent;
-
 	TObjectPtr<UChaosClothAssetEditorMode> ClothEdMode;
 
 	TWeakPtr<const FChaosClothAssetEditorToolkit> ClothToolkit;
diff --git a/Engine/Plugins/Experimental/ChaosClothAssetEditor/Source/ChaosClothAssetEditor/Private/ChaosClothAsset/ClothEditorPreviewScene.cpp b/Engine/Plugins/Experimental/ChaosClothAssetEditor/Source/ChaosClothAssetEditor/Private/ChaosClothAsset/ClothEditorPreviewScene.cpp
--- a/Engine/Plugins/Experimental/ChaosClothAssetEditor/Source/ChaosClothAssetEditor/Private/ChaosClothAsset/ClothEditorPreviewScene.cpp
+++ b/Engine/Plugins/Experimental/ChaosClothAssetEditor/Source/ChaosClothAssetEditor/Private/ChaosClothAsset/ClothEditorPreviewScene.cpp
@@ -32,6 +32,8 @@ FChaosClothPreviewScene::FChaosClothPreviewScene(FPreviewScene::ConstructionValu
 	PreviewSceneDescription = NewObject<UChaosClothPreviewSceneDescription>();
 	PreviewSceneDescription->SetPreviewScene(this);
 
+	SceneActor = GetWorld()->SpawnActor<AActor>(AActor::StaticClass());
+
 	if (PreviewSceneDescription->SkeletalMeshAsset)
 	{
 		CreateSkeletalMeshComponent();
@@ -245,26 +247,35 @@ bool FChaosClothPreviewScene::IsComponentSelected(const UPrimitiveComponent* InC
 
 void FChaosClothPreviewScene::SetClothAsset(UChaosClothAsset* Asset)
 {
+	// Clean up old cloth component if it exists
+
+	check(Asset);
+	check(SceneActor);
+
+	bool bShouldClothBeRootComponent = false;
 	if (ClothComponent)
 	{
-		ClothComponent->SelectionOverrideDelegate.Unbind();
-	}
+		if (SceneActor->GetRootComponent() == ClothComponent)
+		{
+			bShouldClothBeRootComponent = true;
+			SceneActor->SetRootComponent(nullptr);
+		}
 
-	if (SceneActor)
-	{
-		SceneActor->UnregisterAllComponents();
+		ClothComponent->SelectionOverrideDelegate.Unbind();
+		SceneActor->RemoveOwnedComponent(ClothComponent);
+		ClothComponent->DestroyComponent();
 	}
 
-	SceneActor = GetWorld()->SpawnActor<AActor>(AActor::StaticClass());
-
-	check(SceneActor);
-
 	ClothComponent = NewObject<UChaosClothComponent>(SceneActor);
 	ClothComponent->SetClothAsset(Asset);
 
 	ClothComponent->SelectionOverrideDelegate = UPrimitiveComponent::FSelectionOverride::CreateRaw(this, &FChaosClothPreviewScene::IsComponentSelected);
 
-	SceneActor->SetRootComponent(ClothComponent);
+	if (bShouldClothBeRootComponent)
+	{
+		SceneActor->SetRootComponent(ClothComponent);
+	}
+
 	SceneActor->RegisterAllComponents();
 
 	ReattachSkeletalMeshAndAnimation();
diff --git a/Engine/Plugins/Experimental/ChaosClothAssetEditor/Source/ChaosClothAssetEditor/Private/ChaosClothAsset/ClothEditorToolkit.cpp b/Engine/Plugins/Experimental/ChaosClothAssetEditor/Source/ChaosClothAssetEditor/Private/ChaosClothAsset/ClothEditorToolkit.cpp
--- a/Engine/Plugins/Experimental/ChaosClothAssetEditor/Source/ChaosClothAssetEditor/Private/ChaosClothAsset/ClothEditorToolkit.cpp
+++ b/Engine/Plugins/Experimental/ChaosClothAssetEditor/Source/ChaosClothAssetEditor/Private/ChaosClothAsset/ClothEditorToolkit.cpp
@@ -440,7 +440,6 @@ void FChaosClothAssetEditorToolkit::PostInitAssetEditor()
 	ViewportClient->ReceivedFocus(ViewportClient->Viewport);
 
 	// Set up 3D viewport
-	ClothPreviewViewportClient->SetClothComponent(ClothPreviewScene->GetClothComponent());
 	ClothPreviewViewportClient->SetClothEdMode(ClothMode);
 	ClothPreviewViewportClient->SetClothEditorToolkit(StaticCastSharedRef<FChaosClothAssetEditorToolkit>(this->AsShared()));
 
