Add cloth visualization to new cloth asset editor. This is just bringing over a copy of the existing debug draw features for skeletal mesh cloth. We expect to iterate and improve on this in the future beyond what is available in the legacy editor.

ClothEditorSimulationVisualization is a copy of ChaosSimulationEditorExtender, but with interface updates to match our current use (e.g., IClothingSimulation -> UChaosClothComponent). The Visualization is owned by the FChaosClothAssetEditorToolkit. A WeakPtr is shared with FChaosClothAssetEditor3DViewportClient.

Additional changes to make this work:
Added ability to disable and enable simulation
Added an InvalidateViews call when ticking the ToolKit. This ensures the views update while slate is throttling (e.g., when menus are open). This is what the other persona viewports do, so I thought it was OK?

diff --git a/Engine/Plugins/Experimental/ChaosClothAsset/Source/ChaosClothAssetEngine/Private/ChaosClothAsset/ClothComponent.cpp b/Engine/Plugins/Experimental/ChaosClothAsset/Source/ChaosClothAssetEngine/Private/ChaosClothAsset/ClothComponent.cpp
--- a/Engine/Plugins/Experimental/ChaosClothAsset/Source/ChaosClothAssetEngine/Private/ChaosClothAsset/ClothComponent.cpp
+++ b/Engine/Plugins/Experimental/ChaosClothAsset/Source/ChaosClothAssetEngine/Private/ChaosClothAsset/ClothComponent.cpp
@@ -18,7 +18,7 @@ UChaosClothComponent::UChaosClothComponent(const FObjectInitializer& ObjectIniti
 	: Super(ObjectInitializer)
 	, bUseAttachedParentAsPoseComponent(1)  // By default use the parent component as leader pose component
 	, bWaitForParallelTask(0)
-	, bDisableSimulation(0)
+	, bEnableSimulation(1)
 	, bSuspendSimulation(0)
 	, bBindToLeaderComponent(0)
 {
@@ -129,6 +129,13 @@ bool UChaosClothComponent::IsSimulationSuspended() const
 	return bSuspendSimulation || !ClothSimulationProxy.IsValid() || (CVarClothPhysics && !CVarClothPhysics->GetBool());
 }
 
+bool UChaosClothComponent::IsSimulationEnabled() const
+{
+	static IConsoleVariable* const CVarClothPhysics = IConsoleManager::Get().FindConsoleVariable(TEXT("p.ClothPhysics"));
+	// If the console variable doesn't exist, default to simulation enabled.
+	return bEnableSimulation && ClothSimulationProxy.IsValid() && (!CVarClothPhysics || CVarClothPhysics->GetBool());
+}
+
 void UChaosClothComponent::ResetConfigProperties()
 {
 	using namespace UE::Chaos::ClothAsset;
@@ -161,7 +168,7 @@ void UChaosClothComponent::OnUnregister()
 
 bool UChaosClothComponent::IsComponentTickEnabled() const
 {
-	return !bDisableSimulation && Super::IsComponentTickEnabled();
+	return bEnableSimulation && Super::IsComponentTickEnabled();
 }
 
 void UChaosClothComponent::TickComponent(float DeltaTime, enum ELevelTick TickType, FActorComponentTickFunction* ThisTickFunction)
@@ -211,7 +218,7 @@ void UChaosClothComponent::GetUpdateClothSimulationData_AnyThread(TMap<int32, FC
 		OutBlendWeight = BlendWeight;
 		OutClothSimulData = LeaderPoseClothComponent->ClothSimulationProxy->GetCurrentSimulationData_AnyThread();
 	}
-	else if (!bDisableSimulation && !bBindToLeaderComponent && ClothSimulationProxy)
+	else if (bEnableSimulation && !bBindToLeaderComponent && ClothSimulationProxy)
 	{
 		OutBlendWeight = BlendWeight;
 		OutClothSimulData = ClothSimulationProxy->GetCurrentSimulationData_AnyThread();
diff --git a/Engine/Plugins/Experimental/ChaosClothAsset/Source/ChaosClothAssetEngine/Public/ChaosClothAsset/ClothComponent.h b/Engine/Plugins/Experimental/ChaosClothAsset/Source/ChaosClothAssetEngine/Public/ChaosClothAsset/ClothComponent.h
--- a/Engine/Plugins/Experimental/ChaosClothAsset/Source/ChaosClothAssetEngine/Public/ChaosClothAsset/ClothComponent.h
+++ b/Engine/Plugins/Experimental/ChaosClothAsset/Source/ChaosClothAssetEngine/Public/ChaosClothAsset/ClothComponent.h
@@ -74,6 +74,14 @@ public:
 	UFUNCTION(BlueprintCallable, Category = "Components|Simulation", Meta = (Keywords = "Chaos Cloth Simulation Suspend"))
 	bool IsSimulationSuspended() const;
 
+	/** Set whether or not to enable simulation. */
+	UFUNCTION(BlueprintCallable, Category = "Components|Simulation", Meta = (UnsafeDuringActorConstruction, Keywords = "Chaos Cloth Simulation Enable"))
+	void SetEnableSimulation(bool bEnable) { bEnableSimulation = bEnable; }
+
+	/** Return whether or not the simulation is currently enabled. */
+	UFUNCTION(BlueprintCallable, Category = "Components|Simulation", Meta = (Keywords = "Chaos Cloth Simulation Enable"))
+	bool IsSimulationEnabled() const;
+
 	/** Reset all cloth simulation config properties to the values stored in the original cloth asset. */
 	UFUNCTION(BlueprintCallable, Category = "Components|Simulation", Meta = (Keywords = "Chaos Cloth Config Property"))
 	void ResetConfigProperties();
@@ -85,6 +93,8 @@ public:
 	 */
 	TSharedPtr<const FManagedArrayCollection> GetPropertyCollection() const { return TSharedPtr<const FManagedArrayCollection>(PropertyCollection); }
 
+	const UE::Chaos::ClothAsset::FClothSimulationProxy* GetClothSimulationProxy() const { return ClothSimulationProxy.Get(); }
+
 	/** Pose the cloth component using component space transforms. */
 	void Pose(const TArray<FTransform>& InComponentSpaceTransforms);
 
@@ -140,9 +150,9 @@ private:
 	UPROPERTY(EditAnywhere, Category = ClothComponent)
 	uint8 bWaitForParallelTask : 1;
 
-	/** Whether to disable the simulation and use the skinned pose instead. */
+	/** Whether to enable the simulation or use the skinned pose instead. */
 	UPROPERTY()
-	uint8 bDisableSimulation : 1;
+	uint8 bEnableSimulation : 1;
 
 	/** Whether to suspend the simulation and use the last simulated pose. */
 	UPROPERTY()
diff --git a/Engine/Plugins/Experimental/ChaosClothAsset/Source/ChaosClothAssetEngine/Public/ChaosClothAsset/ClothSimulationProxy.h b/Engine/Plugins/Experimental/ChaosClothAsset/Source/ChaosClothAssetEngine/Public/ChaosClothAsset/ClothSimulationProxy.h
--- a/Engine/Plugins/Experimental/ChaosClothAsset/Source/ChaosClothAssetEngine/Public/ChaosClothAsset/ClothSimulationProxy.h
+++ b/Engine/Plugins/Experimental/ChaosClothAsset/Source/ChaosClothAssetEngine/Public/ChaosClothAsset/ClothSimulationProxy.h
@@ -59,6 +59,8 @@ namespace UE::Chaos::ClothAsset
 
 		FBoxSphereBounds CalculateBounds_AnyThread() const;
 
+		const ::Chaos::FClothVisualization* GetVisualization() const { return Visualization.Get(); }
+
 	protected:
 		void Tick();
 		void WriteSimulationData();
diff --git a/Engine/Plugins/Experimental/ChaosClothAssetEditor/Source/ChaosClothAssetEditor/Private/ChaosClothAsset/ClothEditor3DViewportClient.cpp b/Engine/Plugins/Experimental/ChaosClothAssetEditor/Source/ChaosClothAssetEditor/Private/ChaosClothAsset/ClothEditor3DViewportClient.cpp
--- a/Engine/Plugins/Experimental/ChaosClothAssetEditor/Source/ChaosClothAssetEditor/Private/ChaosClothAsset/ClothEditor3DViewportClient.cpp
+++ b/Engine/Plugins/Experimental/ChaosClothAssetEditor/Source/ChaosClothAssetEditor/Private/ChaosClothAsset/ClothEditor3DViewportClient.cpp
@@ -3,6 +3,7 @@
 #include "ChaosClothAsset/ClothEditor3DViewportClient.h"
 #include "ChaosClothAsset/ClothEditorMode.h"
 #include "ChaosClothAsset/ClothEditorToolkit.h"
+#include "ChaosClothAsset/ClothEditorSimulationVisualization.h"
 #include "EditorModeManager.h"
 #include "ChaosClothAsset/ClothComponent.h"
 #include "Dataflow/DataflowEdNode.h"
@@ -21,9 +22,11 @@
 
 FChaosClothAssetEditor3DViewportClient::FChaosClothAssetEditor3DViewportClient(FEditorModeTools* InModeTools,
 	TSharedPtr<FChaosClothPreviewScene> InPreviewScene,
+	TSharedPtr<FClothEditorSimulationVisualization> InVisualization,
 	const TWeakPtr<SEditorViewport>& InEditorViewportWidget)
 	: FEditorViewportClient(InModeTools, InPreviewScene.Get(), InEditorViewportWidget),
 	  ClothPreviewScene(InPreviewScene)
+	, ClothEditorSimulationVisualization(InVisualization)
 {
 	// We want our near clip plane to be quite close so that we can zoom in further.
 	OverrideNearClipPlane(KINDA_SMALL_NUMBER);
@@ -148,12 +151,9 @@ void FChaosClothAssetEditor3DViewportClient::EnableRenderMeshWireframe(bool bEna
 {
 	bRenderMeshWireframe = bEnable;
 
-	if (const TSharedPtr<FChaosClothPreviewScene> PinnedClothPreviewScene = ClothPreviewScene.Pin())
+	if (UChaosClothComponent* const ClothComponent = GetPreviewClothComponent())
 	{
-		if (UChaosClothComponent* const ClothComponent = PinnedClothPreviewScene->GetClothComponent())
-		{
-			ClothComponent->SetForceWireframe(bRenderMeshWireframe);
-		}
+		ClothComponent->SetForceWireframe(bRenderMeshWireframe);
 	}
 }
 
@@ -209,9 +209,44 @@ bool FChaosClothAssetEditor3DViewportClient::IsSimulationSuspended() const
 	return false;
 }
 
+void FChaosClothAssetEditor3DViewportClient::SetEnableSimulation(bool bEnable)
+{
+	if (ClothEdMode)
+	{
+		ClothEdMode->SetEnableSimulation(bEnable);
+	}
+}
+
+bool FChaosClothAssetEditor3DViewportClient::IsSimulationEnabled() const
+{
+	if (ClothEdMode)
+	{
+		return ClothEdMode->IsSimulationEnabled();
+	}
+
+	return false;
+}
+
 void FChaosClothAssetEditor3DViewportClient::Draw(const FSceneView* View, FPrimitiveDrawInterface* PDI)
 {
 	FEditorViewportClient::Draw(View, PDI);
+	TSharedPtr<FClothEditorSimulationVisualization> Visualization = ClothEditorSimulationVisualization.Pin();
+	UChaosClothComponent* const ClothComponent = GetPreviewClothComponent();
+	if(Visualization && ClothComponent)
+	{
+		Visualization->DebugDrawSimulation(ClothComponent, PDI);
+	}
+}
+
+void FChaosClothAssetEditor3DViewportClient::DrawCanvas(FViewport& InViewport, FSceneView& View, FCanvas& Canvas)
+{
+	FEditorViewportClient::DrawCanvas(InViewport, View, Canvas);
+	TSharedPtr<FClothEditorSimulationVisualization> Visualization = ClothEditorSimulationVisualization.Pin();
+	UChaosClothComponent* const ClothComponent = GetPreviewClothComponent();
+	if (Visualization && ClothComponent)
+	{
+		Visualization->DebugDrawSimulationTexts(ClothComponent, &Canvas, &View);
+	}
 }
 
 FBox FChaosClothAssetEditor3DViewportClient::PreviewBoundingBox() const
@@ -234,6 +269,24 @@ TWeakPtr<const FChaosClothPreviewScene> FChaosClothAssetEditor3DViewportClient::
 	return ClothPreviewScene;
 }
 
+UChaosClothComponent* FChaosClothAssetEditor3DViewportClient::GetPreviewClothComponent()
+{
+	if (ClothPreviewScene.IsValid())
+	{
+		return ClothPreviewScene.Pin()->GetClothComponent();
+	}
+	return nullptr;
+}
+
+const UChaosClothComponent* FChaosClothAssetEditor3DViewportClient::GetPreviewClothComponent() const
+{
+	if (ClothPreviewScene.IsValid())
+	{
+		return ClothPreviewScene.Pin()->GetClothComponent();
+	}
+	return nullptr;
+}
+
 void FChaosClothAssetEditor3DViewportClient::ProcessClick(FSceneView& View, HHitProxy* HitProxy, FKey Key, EInputEvent Event, uint32 HitX, uint32 HitY)
 {
 	FEditorViewportClient::ProcessClick(View, HitProxy, Key, Event, HitX, HitY);
diff --git a/Engine/Plugins/Experimental/ChaosClothAssetEditor/Source/ChaosClothAssetEditor/Private/ChaosClothAsset/ClothEditor3DViewportClient.h b/Engine/Plugins/Experimental/ChaosClothAssetEditor/Source/ChaosClothAssetEditor/Private/ChaosClothAsset/ClothEditor3DViewportClient.h
--- a/Engine/Plugins/Experimental/ChaosClothAssetEditor/Source/ChaosClothAssetEditor/Private/ChaosClothAsset/ClothEditor3DViewportClient.h
+++ b/Engine/Plugins/Experimental/ChaosClothAssetEditor/Source/ChaosClothAssetEditor/Private/ChaosClothAsset/ClothEditor3DViewportClient.h
@@ -12,6 +12,7 @@ class FChaosClothAssetEditorToolkit;
 class UTransformProxy;
 class UCombinedTransformGizmo;
 class FChaosClothPreviewScene;
+class FClothEditorSimulationVisualization;
 
 /**
  * Viewport client for the 3d sim preview in the cloth editor. Currently same as editor viewport
@@ -21,7 +22,8 @@ class CHAOSCLOTHASSETEDITOR_API FChaosClothAssetEditor3DViewportClient : public
 {
 public:
 
-	FChaosClothAssetEditor3DViewportClient(FEditorModeTools* InModeTools, TSharedPtr<FChaosClothPreviewScene> InPreviewScene,
+	FChaosClothAssetEditor3DViewportClient(FEditorModeTools* InModeTools, TSharedPtr<FChaosClothPreviewScene> InPreviewScene, 
+		TSharedPtr<FClothEditorSimulationVisualization> InVisualization,
 		const TWeakPtr<SEditorViewport>& InEditorViewportWidget = nullptr);
 
 	// Call this after construction to initialize callbacks when settings change
@@ -48,11 +50,19 @@ public:
 	void SuspendSimulation();
 	void ResumeSimulation();
 	bool IsSimulationSuspended() const;
+	void SetEnableSimulation(bool bEnable);
+	bool IsSimulationEnabled() const;
 
 	FBox PreviewBoundingBox() const;
 
 	TWeakPtr<FChaosClothPreviewScene> GetClothPreviewScene();
 	TWeakPtr<const FChaosClothPreviewScene> GetClothPreviewScene() const;
+	UChaosClothComponent* GetPreviewClothComponent();
+	const UChaosClothComponent* GetPreviewClothComponent() const;
+	TWeakPtr<FClothEditorSimulationVisualization> GetSimulationVisualization() {
+		return ClothEditorSimulationVisualization;
+	}
+	TWeakPtr<const FChaosClothAssetEditorToolkit> GetClothToolKit() const { return ClothToolkit; }
 
 private:
 
@@ -66,6 +76,7 @@ private:
 	virtual void Tick(float DeltaSeconds) override;
 	virtual void ProcessClick(FSceneView& View, HHitProxy* HitProxy, FKey Key, EInputEvent Event, uint32 HitX, uint32 HitY) override;
 	virtual void Draw(const FSceneView* View, FPrimitiveDrawInterface* PDI) override;
+	virtual void DrawCanvas(FViewport& InViewport, FSceneView& View, FCanvas& Canvas) override;
 
 	void OnAssetViewerSettingsChanged(const FName& InPropertyName);
 	void SetAdvancedShowFlagsForScene(const bool bAdvancedShowFlags);
@@ -77,6 +88,8 @@ private:
 	TObjectPtr<UChaosClothAssetEditorMode> ClothEdMode;
 
 	TWeakPtr<const FChaosClothAssetEditorToolkit> ClothToolkit;
+
+	TWeakPtr<FClothEditorSimulationVisualization> ClothEditorSimulationVisualization;
 	
 	bool bSimMeshWireframe = true;
 	bool bRenderMeshWireframe = false;
diff --git a/Engine/Plugins/Experimental/ChaosClothAssetEditor/Source/ChaosClothAssetEditor/Private/ChaosClothAsset/ClothEditorMode.cpp b/Engine/Plugins/Experimental/ChaosClothAssetEditor/Source/ChaosClothAssetEditor/Private/ChaosClothAsset/ClothEditorMode.cpp
--- a/Engine/Plugins/Experimental/ChaosClothAssetEditor/Source/ChaosClothAssetEditor/Private/ChaosClothAsset/ClothEditorMode.cpp
+++ b/Engine/Plugins/Experimental/ChaosClothAssetEditor/Source/ChaosClothAssetEditor/Private/ChaosClothAsset/ClothEditorMode.cpp
@@ -720,6 +720,24 @@ bool UChaosClothAssetEditorMode::IsSimulationSuspended() const
 	return false;
 }
 
+void UChaosClothAssetEditorMode::SetEnableSimulation(bool bEnable)
+{
+	if (PreviewScene && PreviewScene->GetClothComponent())
+	{
+		PreviewScene->GetClothComponent()->SetEnableSimulation(bEnable);
+	}
+}
+
+bool UChaosClothAssetEditorMode::IsSimulationEnabled() const
+{
+	if (PreviewScene && PreviewScene->GetClothComponent())
+	{
+		return PreviewScene->GetClothComponent()->IsSimulationEnabled();
+	}
+
+	return false;
+}
+
 UDataflowComponent* UChaosClothAssetEditorMode::GetDataflowComponent() const
 {
 	return DataflowComponent;
diff --git a/Engine/Plugins/Experimental/ChaosClothAssetEditor/Source/ChaosClothAssetEditor/Private/ChaosClothAsset/ClothEditorMode.h b/Engine/Plugins/Experimental/ChaosClothAssetEditor/Source/ChaosClothAssetEditor/Private/ChaosClothAsset/ClothEditorMode.h
--- a/Engine/Plugins/Experimental/ChaosClothAssetEditor/Source/ChaosClothAssetEditor/Private/ChaosClothAsset/ClothEditorMode.h
+++ b/Engine/Plugins/Experimental/ChaosClothAssetEditor/Source/ChaosClothAssetEditor/Private/ChaosClothAsset/ClothEditorMode.h
@@ -65,6 +65,8 @@ public:
 	void SuspendSimulation();
 	void ResumeSimulation();
 	bool IsSimulationSuspended() const;
+	void SetEnableSimulation(bool bEnabled);
+	bool IsSimulationEnabled() const;
 
 	UDataflowComponent* GetDataflowComponent() const;
 
diff --git a/Engine/Plugins/Experimental/ChaosClothAssetEditor/Source/ChaosClothAssetEditor/Private/ChaosClothAsset/ClothEditorSimulationVisualization.cpp b/Engine/Plugins/Experimental/ChaosClothAssetEditor/Source/ChaosClothAssetEditor/Private/ChaosClothAsset/ClothEditorSimulationVisualization.cpp
new file mode 100644
--- /dev/null
+++ b/Engine/Plugins/Experimental/ChaosClothAssetEditor/Source/ChaosClothAssetEditor/Private/ChaosClothAsset/ClothEditorSimulationVisualization.cpp
@@ -0,0 +1,210 @@
+// Copyright Epic Games, Inc. All Rights Reserved.
+
+#include "ClothEditorSimulationVisualization.h"
+#include "ChaosClothAsset/ClothEditor3DViewportClient.h"
+#include "ChaosClothAsset/ClothSimulationProxy.h"
+#include "ChaosClothAsset/ClothComponent.h"
+#include "ChaosCloth/ChaosClothVisualization.h"
+#include "Framework/MultiBox/MultiBoxBuilder.h"
+#include "Rendering/SkeletalMeshRenderData.h"
+
+
+#define LOCTEXT_NAMESPACE "ClothEditorSimulationVisualization"
+
+namespace Private
+{
+struct FVisualizationOption
+{
+	// Actual option entries
+	static const FVisualizationOption OptionData[];
+	static const uint32 Count;
+
+	// Chaos debug draw function
+	typedef void (::Chaos::FClothVisualization::* FDebugDrawFunction)(FPrimitiveDrawInterface*) const;
+	typedef void (::Chaos::FClothVisualization::* FDebugDrawTextsFunction)(FCanvas*, const FSceneView*) const;
+	FDebugDrawFunction DebugDrawFunction;
+	FDebugDrawTextsFunction DebugDrawTextsFunction;
+
+	FText DisplayName;         // Text for menu entries.
+	FText ToolTip;             // Text for menu tooltips.
+	bool bDisablesSimulation;  // Whether or not this option requires the simulation to be disabled.
+	bool bHidesClothSections;  // Hides the cloth section to avoid zfighting with the debug geometry.
+
+	// Console override
+	IConsoleVariable* const ConsoleVariable;
+
+	FVisualizationOption(FDebugDrawFunction InDebugDrawFunction, const TCHAR* ConsoleName, const FText& InDisplayName, const FText& InToolTip, bool bInDisablesSimulation = false, bool bInHidesClothSections = false)
+		: DebugDrawFunction(InDebugDrawFunction)
+		, DisplayName(InDisplayName)
+		, ToolTip(InToolTip)
+		, bDisablesSimulation(bInDisablesSimulation)
+		, bHidesClothSections(bInHidesClothSections)
+		, ConsoleVariable(IConsoleManager::Get().RegisterConsoleVariable(ConsoleName, false, *InToolTip.ToString(), 0))
+	{}
+
+	FVisualizationOption(FDebugDrawTextsFunction InDebugDrawTextsFunction, const TCHAR* ConsoleName, const FText& InDisplayName, const FText& InToolTip, bool bInDisablesSimulation = false, bool bInHidesClothSections = false)
+		: DebugDrawTextsFunction(InDebugDrawTextsFunction)
+		, DisplayName(InDisplayName)
+		, ToolTip(InToolTip)
+		, bDisablesSimulation(bInDisablesSimulation)
+		, bHidesClothSections(bInHidesClothSections)
+		, ConsoleVariable(IConsoleManager::Get().RegisterConsoleVariable(ConsoleName, false, *InToolTip.ToString(), 0))
+	{}
+
+	~FVisualizationOption()
+	{
+		IConsoleManager::Get().UnregisterConsoleObject(ConsoleVariable);
+	}
+
+	bool IsConsoleVariableEnabled() const { return ConsoleVariable->AsVariableBool()->GetValueOnGameThread(); }
+};
+
+const FVisualizationOption FVisualizationOption::OptionData[] =
+{
+	FVisualizationOption(&::Chaos::FClothVisualization::DrawPhysMeshShaded      , TEXT("p.ChaosClothEditor.DebugDrawPhysMeshShaded"), LOCTEXT("ChaosVisName_PhysMesh"            , "Physical Mesh (Flat Shaded)"), LOCTEXT("ChaosVisName_PhysMeshShaded_ToolTip"      , "Draws the current physical result as a doubled sided flat shaded mesh"), /*bDisablesSimulation =*/false, /*bHidesClothSections=*/true),
+	FVisualizationOption(&::Chaos::FClothVisualization::DrawPhysMeshWired       , TEXT("p.ChaosClothEditor.DebugDrawPhysMeshWired"), LOCTEXT("ChaosVisName_PhysMeshWire"        , "Physical Mesh (Wireframe)"), LOCTEXT("ChaosVisName_PhysMeshWired_ToolTip"       , "Draws the current physical mesh result in wireframe")),
+	FVisualizationOption(&::Chaos::FClothVisualization::DrawAnimMeshWired       , TEXT("p.ChaosClothEditor.DebugDrawAnimMeshWired"), LOCTEXT("ChaosVisName_AnimMeshWire"        , "Animated Mesh (Wireframe)"), LOCTEXT("ChaosVisName_AnimMeshWired_ToolTip"       , "Draws the current animated mesh input in wireframe")),
+	FVisualizationOption(&::Chaos::FClothVisualization::DrawParticleIndices     , TEXT("p.ChaosClothEditor.DebugDrawParticleIndices"), LOCTEXT("ChaosVisName_ParticleIndices"     , "Particle Indices"), LOCTEXT("ChaosVisName_ParticleIndices_ToolTip"     , "Draws the particle indices as instantiated by the solver")),
+	FVisualizationOption(&::Chaos::FClothVisualization::DrawElementIndices      , TEXT("p.ChaosClothEditor.DebugDrawElementIndices"), LOCTEXT("ChaosVisName_ElementIndices"      , "Element Indices"), LOCTEXT("ChaosVisName_ElementIndices_ToolTip"      , "Draws the element's (triangle or other) indices as instantiated by the solver")),
+	FVisualizationOption(&::Chaos::FClothVisualization::DrawPointNormals        , TEXT("p.ChaosClothEditor.DebugDrawPointNormals"), LOCTEXT("ChaosVisName_PointNormals"        , "Physical Mesh Normals"), LOCTEXT("ChaosVisName_PointNormals_ToolTip"        , "Draws the current point normals for the simulation mesh")),
+	FVisualizationOption(&::Chaos::FClothVisualization::DrawPointVelocities     , TEXT("p.ChaosClothEditor.DebugDrawPointVelocities"), LOCTEXT("ChaosVisName_PointVelocities"     , "Point Velocities"), LOCTEXT("ChaosVisName_PointVelocities_ToolTip"     , "Draws the current point velocities for the simulation mesh")),
+	FVisualizationOption(&::Chaos::FClothVisualization::DrawAnimNormals         , TEXT("p.ChaosClothEditor.DebugDrawAnimNormals"), LOCTEXT("ChaosVisName_AnimNormals"         , "Animated Mesh Normals"), LOCTEXT("ChaosVisName_AnimNormals_ToolTip"         , "Draws the current point normals for the animated mesh")),
+	FVisualizationOption(&::Chaos::FClothVisualization::DrawCollision           , TEXT("p.ChaosClothEditor.DebugDrawCollision"), LOCTEXT("ChaosVisName_Collision"           , "Collisions"), LOCTEXT("ChaosVisName_Collision_ToolTip"           , "Draws the collision bodies the simulation is currently using")),
+	FVisualizationOption(&::Chaos::FClothVisualization::DrawBackstops           , TEXT("p.ChaosClothEditor.DebugDrawBackstops"), LOCTEXT("ChaosVisName_Backstop"            , "Backstops"), LOCTEXT("ChaosVisName_Backstop_ToolTip"            , "Draws the backstop radius and position for each simulation particle")),
+	FVisualizationOption(&::Chaos::FClothVisualization::DrawBackstopDistances   , TEXT("p.ChaosClothEditor.DebugDrawBackstopDistances"), LOCTEXT("ChaosVisName_BackstopDistance"    , "Backstop Distances"), LOCTEXT("ChaosVisName_BackstopDistance_ToolTip"    , "Draws the backstop distance offset for each simulation particle"), /*bDisablesSimulation =*/true),
+	FVisualizationOption(&::Chaos::FClothVisualization::DrawMaxDistances        , TEXT("p.ChaosClothEditor.DebugDrawMaxDistances"), LOCTEXT("ChaosVisName_MaxDistance"         , "Max Distances"), LOCTEXT("ChaosVisName_MaxDistance_ToolTip"         , "Draws the current max distances for the sim particles as a line along its normal"), /*bDisablesSimulation =*/true),
+	FVisualizationOption(&::Chaos::FClothVisualization::DrawMaxDistanceValues   , TEXT("p.ChaosClothEditor.DebugDrawMaxDistanceValues"), LOCTEXT("ChaosVisName_MaxDistanceValue"    , "Max Distances As Numbers"), LOCTEXT("ChaosVisName_MaxDistanceValue_ToolTip"    , "Draws the current max distances as numbers")),
+	FVisualizationOption(&::Chaos::FClothVisualization::DrawAnimDrive           , TEXT("p.ChaosClothEditor.DebugDrawAnimDrive"), LOCTEXT("ChaosVisName_AnimDrive"           , "Anim Drive"), LOCTEXT("ChaosVisName_AnimDrive_Tooltip"           , "Draws the current skinned reference mesh for the simulation which anim drive will attempt to reach if enabled")),
+	FVisualizationOption(&::Chaos::FClothVisualization::DrawEdgeConstraint      , TEXT("p.ChaosClothEditor.DebugDrawEdgeConstraint"), LOCTEXT("ChaosVisName_EdgeConstraint"      , "Edge Constraint"), LOCTEXT("ChaosVisName_EdgeConstraint_Tooltip"      , "Draws the edge spring constraints")),
+	FVisualizationOption(&::Chaos::FClothVisualization::DrawBendingConstraint   , TEXT("p.ChaosClothEditor.DebugDrawBendingConstraint"), LOCTEXT("ChaosVisName_BendingConstraint"   , "Bending Constraint"), LOCTEXT("ChaosVisName_BendingConstraint_Tooltip"   , "Draws the bending spring constraints")),
+	FVisualizationOption(&::Chaos::FClothVisualization::DrawLongRangeConstraint , TEXT("p.ChaosClothEditor.DebugDrawLongRangeConstraint"), LOCTEXT("ChaosVisName_LongRangeConstraint" , "Long Range Constraint"), LOCTEXT("ChaosVisName_LongRangeConstraint_Tooltip" , "Draws the long range attachment constraint distances")),
+	FVisualizationOption(&::Chaos::FClothVisualization::DrawWindAndPressureForces, TEXT("p.ChaosClothEditor.DebugDrawWindAndPressureForces"), LOCTEXT("ChaosVisName_WindAndPressureForces", "Wind Aerodynamic And Pressure Forces"), LOCTEXT("ChaosVisName_WindAndPressure_Tooltip", "Draws the Wind drag and lift and pressure forces")),
+	FVisualizationOption(&::Chaos::FClothVisualization::DrawLocalSpace          , TEXT("p.ChaosClothEditor.DebugDrawLocalSpace"), LOCTEXT("ChaosVisName_LocalSpace"          , "Local Space Reference Bone"), LOCTEXT("ChaosVisName_LocalSpace_Tooltip"          , "Draws the local space reference bone")),
+	FVisualizationOption(&::Chaos::FClothVisualization::DrawSelfCollision       , TEXT("p.ChaosClothEditor.DebugDrawSelfCollision"), LOCTEXT("ChaosVisName_SelfCollision"       , "Self Collision"), LOCTEXT("ChaosVisName_SelfCollision_Tooltip"       , "Draws the self collision thickness/debugging information")),
+	FVisualizationOption(&::Chaos::FClothVisualization::DrawSelfIntersection    , TEXT("p.ChaosClothEditor.DebugDrawSelfIntersection"), LOCTEXT("ChaosVisName_SelfIntersection"    , "Self Intersection"), LOCTEXT("ChaosVisName_SelfIntersection_Tooltip"    , "Draws the self intersection contour/region information")),
+};
+const uint32 FVisualizationOption::Count = sizeof(OptionData) / sizeof(FVisualizationOption);
+
+} // namespace Private
+
+FClothEditorSimulationVisualization::FClothEditorSimulationVisualization()
+	: Flags(false, Private::FVisualizationOption::Count)
+{
+}
+
+void FClothEditorSimulationVisualization::ExtendViewportShowMenu(FMenuBuilder& MenuBuilder, TSharedRef<FChaosClothAssetEditor3DViewportClient> ViewportClient)
+{
+	MenuBuilder.BeginSection(TEXT("ChaosSimulation_Visualizations"), LOCTEXT("VisualizationSection", "Visualizations"));
+	{
+		for (uint32 OptionIndex = 0; OptionIndex < Private::FVisualizationOption::Count; ++OptionIndex)
+		{
+			// Handler for visualization entry being clicked
+			const FExecuteAction ExecuteAction = FExecuteAction::CreateLambda([this, OptionIndex, ViewportClient]()
+			{
+				Flags[OptionIndex] = !Flags[OptionIndex];
+
+			if (UChaosClothComponent* const ClothComponent = ViewportClient->GetPreviewClothComponent())
+			{
+				// If we need to toggle the disabled or visibility states, handle it
+				// Disable simulation
+				const bool bShouldDisableSimulation = ShouldDisableSimulation();
+				if (ClothComponent->IsSimulationEnabled() == bShouldDisableSimulation)
+				{
+					ClothComponent->SetEnableSimulation(!bShouldDisableSimulation);
+				}
+				// Hide cloth section
+				if (Private::FVisualizationOption::OptionData[OptionIndex].bHidesClothSections)
+				{
+					const bool bIsClothSectionsVisible = !Flags[OptionIndex];
+					ShowClothSections(ClothComponent, bIsClothSectionsVisible);
+				}
+			}
+			});
+
+			// Checkstate function for visualization entries
+			const FIsActionChecked IsActionChecked = FIsActionChecked::CreateLambda([this, OptionIndex]()
+			{
+				return Flags[OptionIndex];
+			});
+
+			const FUIAction Action(ExecuteAction, FCanExecuteAction(), IsActionChecked);
+
+			MenuBuilder.AddMenuEntry(Private::FVisualizationOption::OptionData[OptionIndex].DisplayName, Private::FVisualizationOption::OptionData[OptionIndex].ToolTip, FSlateIcon(), Action, NAME_None, EUserInterfaceActionType::ToggleButton);
+		}
+	}
+	MenuBuilder.EndSection();
+}
+
+void FClothEditorSimulationVisualization::DebugDrawSimulation(const UChaosClothComponent* ClothComponent, FPrimitiveDrawInterface* PDI)
+{
+	const ::Chaos::FClothVisualization* const Visualization = ClothComponent && ClothComponent->GetClothSimulationProxy() ? ClothComponent->GetClothSimulationProxy()->GetVisualization() : nullptr;
+	if (!ensure(Visualization))
+	{
+		return;
+	}
+
+	for (int32 OptionIndex = 0; OptionIndex < Private::FVisualizationOption::Count; ++OptionIndex)
+	{
+		if (Private::FVisualizationOption::OptionData[OptionIndex].DebugDrawFunction &&
+			(Private::FVisualizationOption::OptionData[OptionIndex].IsConsoleVariableEnabled() || Flags[OptionIndex]))
+		{
+			(Visualization->*(Private::FVisualizationOption::OptionData[OptionIndex].DebugDrawFunction))(PDI);
+		}
+	}
+}
+
+void FClothEditorSimulationVisualization::DebugDrawSimulationTexts(const UChaosClothComponent* ClothComponent, FCanvas* Canvas, const FSceneView* SceneView)
+{
+	const ::Chaos::FClothVisualization* const Visualization = ClothComponent && ClothComponent->GetClothSimulationProxy() ? ClothComponent->GetClothSimulationProxy()->GetVisualization() : nullptr;
+	if (!ensure(Visualization))
+	{
+		return;
+	}
+	for (int32 OptionIndex = 0; OptionIndex < Private::FVisualizationOption::Count; ++OptionIndex)
+	{
+		if (Private::FVisualizationOption::OptionData[OptionIndex].DebugDrawTextsFunction &&
+			(Private::FVisualizationOption::OptionData[OptionIndex].IsConsoleVariableEnabled() || Flags[OptionIndex]))
+		{
+			(Visualization->*(Private::FVisualizationOption::OptionData[OptionIndex].DebugDrawTextsFunction))(Canvas, SceneView);
+		}
+	}
+}
+
+bool FClothEditorSimulationVisualization::ShouldDisableSimulation() const
+{
+	for (uint32 OptionIndex = 0; OptionIndex < Private::FVisualizationOption::Count; ++OptionIndex)
+	{
+		if (Flags[OptionIndex])
+		{
+			const Private::FVisualizationOption& Data = Private::FVisualizationOption::OptionData[OptionIndex];
+
+			if (Data.bDisablesSimulation)
+			{
+				return true;
+			}
+		}
+	}
+	return false;
+}
+
+void FClothEditorSimulationVisualization::ShowClothSections(UChaosClothComponent* ClothComponent, bool bIsClothSectionsVisible) const
+{
+	if (FSkeletalMeshRenderData* const SkeletalMeshRenderData = ClothComponent->GetSkeletalMeshRenderData())
+	{
+		for (int32 LODIndex = 0; LODIndex < SkeletalMeshRenderData->LODRenderData.Num(); ++LODIndex)
+		{
+			FSkeletalMeshLODRenderData& SkeletalMeshLODRenderData = SkeletalMeshRenderData->LODRenderData[LODIndex];
+
+			for (int32 SectionIndex = 0; SectionIndex < SkeletalMeshLODRenderData.RenderSections.Num(); ++SectionIndex)
+			{
+				FSkelMeshRenderSection& SkelMeshRenderSection = SkeletalMeshLODRenderData.RenderSections[SectionIndex];
+
+				if (SkelMeshRenderSection.HasClothingData())
+				{
+					ClothComponent->ShowMaterialSection(SkelMeshRenderSection.MaterialIndex, SectionIndex, bIsClothSectionsVisible, LODIndex);
+				}
+			}
+		}
+	}
+}
+#undef LOCTEXT_NAMESPACE
diff --git a/Engine/Plugins/Experimental/ChaosClothAssetEditor/Source/ChaosClothAssetEditor/Private/ChaosClothAsset/ClothEditorSimulationVisualization.h b/Engine/Plugins/Experimental/ChaosClothAssetEditor/Source/ChaosClothAssetEditor/Private/ChaosClothAsset/ClothEditorSimulationVisualization.h
new file mode 100644
--- /dev/null
+++ b/Engine/Plugins/Experimental/ChaosClothAssetEditor/Source/ChaosClothAssetEditor/Private/ChaosClothAsset/ClothEditorSimulationVisualization.h
@@ -0,0 +1,32 @@
+// Copyright Epic Games, Inc. All Rights Reserved.
+
+#pragma once
+#include "Containers/BitArray.h"
+#include "Templates/SharedPointer.h"
+
+class FChaosClothAssetEditor3DViewportClient;
+class FMenuBuilder;
+class UChaosClothComponent;
+class FPrimitiveDrawInterface;
+class FCanvas;
+class FSceneView;
+
+class FClothEditorSimulationVisualization
+{
+public:
+	FClothEditorSimulationVisualization();
+
+	void ExtendViewportShowMenu(FMenuBuilder& MenuBuilder, TSharedRef<FChaosClothAssetEditor3DViewportClient> ViewportClient);
+	void DebugDrawSimulation(const UChaosClothComponent* ClothComponent, FPrimitiveDrawInterface* PDI);
+	void DebugDrawSimulationTexts(const UChaosClothComponent* ClothComponent, FCanvas* Canvas, const FSceneView* SceneView);
+
+private:
+	/** Return whether or not - given the current enabled options - the simulation should be disabled. */
+	bool ShouldDisableSimulation() const;
+	/** Show/hide all cloth sections for the specified mesh compoment. */
+	void ShowClothSections(UChaosClothComponent* ClothComponent, bool bIsClothSectionsVisible) const;
+
+private:
+	/** Flags used to store the checked status for the visualization options. */
+	TBitArray<> Flags;
+};
\ No newline at end of file
diff --git a/Engine/Plugins/Experimental/ChaosClothAssetEditor/Source/ChaosClothAssetEditor/Private/ChaosClothAsset/ClothEditorToolkit.cpp b/Engine/Plugins/Experimental/ChaosClothAssetEditor/Source/ChaosClothAssetEditor/Private/ChaosClothAsset/ClothEditorToolkit.cpp
--- a/Engine/Plugins/Experimental/ChaosClothAssetEditor/Source/ChaosClothAssetEditor/Private/ChaosClothAsset/ClothEditorToolkit.cpp
+++ b/Engine/Plugins/Experimental/ChaosClothAssetEditor/Source/ChaosClothAssetEditor/Private/ChaosClothAsset/ClothEditorToolkit.cpp
@@ -9,6 +9,7 @@
 #include "ChaosClothAsset/ClothEditor3DViewportClient.h"
 #include "ChaosClothAsset/SClothEditorRestSpaceViewport.h"
 #include "ChaosClothAsset/ClothEditorRestSpaceViewportClient.h"
+#include "ChaosClothAsset/ClothEditorSimulationVisualization.h"
 #include "ChaosClothAsset/ClothComponent.h"
 #include "ChaosClothAsset/ClothAsset.h"
 #include "ChaosClothAsset/ClothEditorModeToolkit.h"
@@ -36,6 +37,7 @@
 #include "AdvancedPreviewSceneModule.h"
 #include "Widgets/Layout/SSpacer.h"
 #include "DynamicMesh/DynamicMesh3.h"
+#include "Toolkits/AssetEditorToolkitMenuContext.h"
 
 #define LOCTEXT_NAMESPACE "ChaosClothAssetEditorToolkit"
 
@@ -178,9 +180,11 @@ FChaosClothAssetEditorToolkit::FChaosClothAssetEditorToolkit(UAssetEditor* InOwn
 	ClothPreviewEditorModeManager->SetPreviewScene(ClothPreviewScene.Get());
 	ClothPreviewScene->SetModeManager(ClothPreviewEditorModeManager);
 
+	ClothEditorSimulationVisualization = MakeShared<FClothEditorSimulationVisualization>();
+
 	//ClothPreviewInputRouter = ClothPreviewEditorModeManager->GetInteractiveToolsContext()->InputRouter;
 	ClothPreviewTabContent = MakeShareable(new FEditorViewportTabContent());
-	ClothPreviewViewportClient = MakeShared<FChaosClothAssetEditor3DViewportClient>(ClothPreviewEditorModeManager.Get(), ClothPreviewScene);
+	ClothPreviewViewportClient = MakeShared<FChaosClothAssetEditor3DViewportClient>(ClothPreviewEditorModeManager.Get(), ClothPreviewScene, ClothEditorSimulationVisualization);
 	ClothPreviewViewportClient->RegisterDelegates();
 
 	ClothPreviewViewportDelegate = [this](FAssetEditorViewportConstructionArgs InArgs)
@@ -242,6 +246,8 @@ void FChaosClothAssetEditorToolkit::Tick(float DeltaTime)
 			OnClothAssetChanged();
 		}
 	}
+
+	InvalidateViews();
 }
 
 TStatId FChaosClothAssetEditorToolkit::GetStatId() const
@@ -458,6 +464,15 @@ void FChaosClothAssetEditorToolkit::PostInitAssetEditor()
 	ClothMode->SetDataflowGraphEditor(GraphEditor);
 }
 
+void FChaosClothAssetEditorToolkit::InitToolMenuContext(FToolMenuContext& MenuContext)
+{
+	FAssetEditorToolkit::InitToolMenuContext(MenuContext);
+
+	UAssetEditorToolkitMenuContext* const ClothEditorContext = NewObject<UAssetEditorToolkitMenuContext>();
+	ClothEditorContext->Toolkit = SharedThis(this);
+	MenuContext.AddObject(ClothEditorContext);
+}
+
 //~ End FAssetEditorToolkit overrides
 
 
@@ -1027,5 +1042,11 @@ void FChaosClothAssetEditorToolkit::OnClothAssetChanged()
 	}
 }
 
+void FChaosClothAssetEditorToolkit::InvalidateViews()
+{
+	ViewportClient->Invalidate();
+	ClothPreviewViewportClient->Invalidate();
+}
+
 
 #undef LOCTEXT_NAMESPACE
diff --git a/Engine/Plugins/Experimental/ChaosClothAssetEditor/Source/ChaosClothAssetEditor/Private/ChaosClothAsset/ClothEditorToolkit.h b/Engine/Plugins/Experimental/ChaosClothAssetEditor/Source/ChaosClothAssetEditor/Private/ChaosClothAsset/ClothEditorToolkit.h
--- a/Engine/Plugins/Experimental/ChaosClothAssetEditor/Source/ChaosClothAssetEditor/Private/ChaosClothAsset/ClothEditorToolkit.h
+++ b/Engine/Plugins/Experimental/ChaosClothAssetEditor/Source/ChaosClothAssetEditor/Private/ChaosClothAsset/ClothEditorToolkit.h
@@ -18,6 +18,7 @@ class UEdGraphNode;
 class UChaosClothComponent;
 class SChaosClothAssetEditorRestSpaceViewport;
 class SChaosClothAssetEditor3DViewport;
+class FClothEditorSimulationVisualization;
 
 namespace Dataflow
 {
@@ -76,6 +77,7 @@ private:
 	virtual void RemoveViewportOverlayWidget(TSharedRef<SWidget> InViewportOverlayWidget) override;
 	virtual bool OnRequestClose() override;
 	virtual void PostInitAssetEditor() override;
+	virtual void InitToolMenuContext(FToolMenuContext& MenuContext) override;
 
 	// IAssetEditorInstance
 	// TODO: If this returns true then the editor cannot re-open after it's closed. Figure out why.
@@ -103,6 +105,7 @@ private:
 
 	void PopulateOutliner();
 	void OnClothAssetChanged();
+	void InvalidateViews();
 
 	// Dataflow
 	void EvaluateNode(FDataflowNode* Node, FDataflowOutput* Out);
@@ -124,6 +127,7 @@ private:
 	AssetEditorViewportFactoryFunction ClothPreviewViewportDelegate;
 	TSharedPtr<FChaosClothAssetEditor3DViewportClient> ClothPreviewViewportClient;
 	TSharedPtr<FAssetEditorModeManager> ClothPreviewEditorModeManager;
+	TSharedPtr<FClothEditorSimulationVisualization> ClothEditorSimulationVisualization;
 
 	TSharedPtr<SChaosClothAssetEditorRestSpaceViewport> RestSpaceViewportWidget;
 	TSharedPtr<SChaosClothAssetEditor3DViewport> PreviewViewportWidget;
diff --git a/Engine/Plugins/Experimental/ChaosClothAssetEditor/Source/ChaosClothAssetEditor/Private/ChaosClothAsset/SClothEditor3DViewportToolBar.cpp b/Engine/Plugins/Experimental/ChaosClothAssetEditor/Source/ChaosClothAssetEditor/Private/ChaosClothAsset/SClothEditor3DViewportToolBar.cpp
--- a/Engine/Plugins/Experimental/ChaosClothAssetEditor/Source/ChaosClothAssetEditor/Private/ChaosClothAsset/SClothEditor3DViewportToolBar.cpp
+++ b/Engine/Plugins/Experimental/ChaosClothAssetEditor/Source/ChaosClothAssetEditor/Private/ChaosClothAsset/SClothEditor3DViewportToolBar.cpp
@@ -2,6 +2,8 @@
 
 #include "ChaosClothAsset/SClothEditor3DViewportToolBar.h"
 #include "ChaosClothAsset/SClothEditor3DViewport.h"
+#include "ChaosClothAsset/ClothEditor3DViewportClient.h"
+#include "ChaosClothAsset/ClothEditorSimulationVisualization.h"
 #include "ChaosClothAsset/ClothEditorCommands.h"
 #include "ChaosClothAsset/ClothEditorStyle.h"
 #include "EditorViewportCommands.h"
@@ -141,4 +143,21 @@ TSharedRef<SWidget> SChaosClothAssetEditor3DViewportToolBar::MakeToolBar(const T
 	return ToolbarBuilder.MakeWidget();
 }
 
+void SChaosClothAssetEditor3DViewportToolBar::ExtendOptionsMenu(FMenuBuilder& OptionsMenuBuilder) const
+{
+	constexpr bool bOpenSubMenuOnClick = false;
+	constexpr bool bShouldCloseWindowAfterMenuSelection = false;
+	OptionsMenuBuilder.AddSubMenu(
+		LOCTEXT("ChaosClothAssetEditor_SimulationVisualization", "Simulation Visualization"),
+		LOCTEXT("ChaosClothAssetEditor_SimulationVisualizationToolTip", "Options to control simulation visualization"),
+		FNewMenuDelegate::CreateLambda([this](FMenuBuilder& MenuBuilder)
+	{
+		TSharedRef<FChaosClothAssetEditor3DViewportClient> ViewportClient = StaticCastSharedPtr<FChaosClothAssetEditor3DViewportClient>(ChaosClothAssetEditor3DViewportPtr.Pin()->GetViewportClient()).ToSharedRef();
+		if (FClothEditorSimulationVisualization* const Visualization = ViewportClient->GetSimulationVisualization().Pin().Get())
+		{
+			Visualization->ExtendViewportShowMenu(MenuBuilder, ViewportClient);
+		}
+	}), bOpenSubMenuOnClick, FSlateIcon(), bShouldCloseWindowAfterMenuSelection);
+}
+
 #undef LOCTEXT_NAMESPACE
diff --git a/Engine/Plugins/Experimental/ChaosClothAssetEditor/Source/ChaosClothAssetEditor/Private/ChaosClothAsset/SClothEditor3DViewportToolBar.h b/Engine/Plugins/Experimental/ChaosClothAssetEditor/Source/ChaosClothAssetEditor/Private/ChaosClothAsset/SClothEditor3DViewportToolBar.h
--- a/Engine/Plugins/Experimental/ChaosClothAssetEditor/Source/ChaosClothAssetEditor/Private/ChaosClothAsset/SClothEditor3DViewportToolBar.h
+++ b/Engine/Plugins/Experimental/ChaosClothAssetEditor/Source/ChaosClothAssetEditor/Private/ChaosClothAsset/SClothEditor3DViewportToolBar.h
@@ -26,6 +26,8 @@ private:
 	TSharedRef<SWidget> MakeDisplayToolBar(const TSharedPtr<FExtender> InExtenders);
 	TSharedRef<SWidget> MakeToolBar(const TSharedPtr<FExtender> InExtenders);
 
+	virtual void ExtendOptionsMenu(FMenuBuilder& OptionsMenuBuilder) const override;
+
 	/** The viewport that we are in */
 	TWeakPtr<class SChaosClothAssetEditor3DViewport> ChaosClothAssetEditor3DViewportPtr;
 
