Chaos Cloth - Fixed the transfer skin weight node for when the render mesh is different from the sim mesh.

#rnx

diff --git a/Engine/Plugins/Experimental/ChaosClothAssetEditor/Source/ChaosClothAssetDataflowNodes/Private/ChaosClothAsset/TransferSkinWeightsNode.cpp b/Engine/Plugins/Experimental/ChaosClothAssetEditor/Source/ChaosClothAssetDataflowNodes/Private/ChaosClothAsset/TransferSkinWeightsNode.cpp
--- a/Engine/Plugins/Experimental/ChaosClothAssetEditor/Source/ChaosClothAssetDataflowNodes/Private/ChaosClothAsset/TransferSkinWeightsNode.cpp
+++ b/Engine/Plugins/Experimental/ChaosClothAssetEditor/Source/ChaosClothAssetDataflowNodes/Private/ChaosClothAsset/TransferSkinWeightsNode.cpp
@@ -112,50 +112,69 @@ void FChaosClothAssetTransferSkinWeightsNode::Evaluate(Dataflow::FContext& Conte
 				FCollectionClothLodFacade ClothLodFacade = ClothFacade.GetLod(TargetLODIdx);
 
 				// Cloth collection data arrays we are writing to
-				TArrayView<int32> SimNumBoneInfluences = ClothLodFacade.GetSimNumBoneInfluences();
-				TArrayView<TArray<int32>> SimBoneIndices = ClothLodFacade.GetSimBoneIndices();
-				TArrayView<TArray<float>> SimBoneWeights = ClothLodFacade.GetSimBoneWeights();
-
-				TArrayView<int32> RenderNumBoneInfluences = ClothLodFacade.GetRenderNumBoneInfluences();
-				TArrayView<TArray<int32>> RenderBoneIndices = ClothLodFacade.GetRenderBoneIndices();
-				TArrayView<TArray<float>> RenderBoneWeights = ClothLodFacade.GetRenderBoneWeights();
-
-				const TArrayView<FVector3f> SimPositions = ClothLodFacade.GetSimRestPosition();
+				const TArrayView<int32> SimNumBoneInfluences = ClothLodFacade.GetSimNumBoneInfluences();
+				const TArrayView<TArray<int32>> SimBoneIndices = ClothLodFacade.GetSimBoneIndices();
+				const TArrayView<TArray<float>> SimBoneWeights = ClothLodFacade.GetSimBoneWeights();
+				const TConstArrayView<FVector3f> SimPositions = ClothLodFacade.GetSimRestPosition();
+				const int32 NumSimVertices = ClothLodFacade.GetNumSimVertices();
 
 				checkSlow(SimPositions.Num() == SimBoneIndices.Num());
 
-				const int32 NumVert = ClothLodFacade.GetNumSimVertices();
 				constexpr bool bUseParallel = true;
 
 				// Iterate over each vertex and write the data from FBoneWeights into cloth collection managed arrays
-				ParallelFor(NumVert, [&](int32 VertexID)
-				{
-					const FVector3d PosD = FVector3d(SimPositions[VertexID]);
 
-					UE::AnimationCore::FBoneWeights BoneWeights;
-					TransferBoneWeights.Compute(PosD, FTransformSRT3d::Identity(), BoneWeights, &TargetBoneToIndex);
+				ParallelFor(NumSimVertices,
+					[&SimPositions, &TransferBoneWeights, &TargetBoneToIndex, &SimNumBoneInfluences, &SimBoneIndices, &SimBoneWeights](int32 VertexID)
+					{
+						const FVector3d PosD = FVector3d(SimPositions[VertexID]);
+
+						UE::AnimationCore::FBoneWeights BoneWeights;
+						TransferBoneWeights.Compute(PosD, FTransformSRT3d::Identity(), BoneWeights, &TargetBoneToIndex);
+
+						const int32 NumBones = BoneWeights.Num();
 
-					const int32 NumBones = BoneWeights.Num();
+						SimNumBoneInfluences[VertexID] = NumBones;
+						SimBoneIndices[VertexID].SetNum(NumBones);
+						SimBoneWeights[VertexID].SetNum(NumBones);
 
-					SimNumBoneInfluences[VertexID] = NumBones;
-					SimBoneIndices[VertexID].SetNum(NumBones);
-					SimBoneWeights[VertexID].SetNum(NumBones);
+						for (int BoneIdx = 0; BoneIdx < NumBones; ++BoneIdx)
+						{
+							SimBoneIndices[VertexID][BoneIdx] = BoneWeights[BoneIdx].GetBoneIndex();
+							SimBoneWeights[VertexID][BoneIdx] = BoneWeights[BoneIdx].GetWeight();
+						}
 
-					RenderNumBoneInfluences[VertexID] = NumBones;
-					RenderBoneIndices[VertexID].SetNum(NumBones);
-					RenderBoneWeights[VertexID].SetNum(NumBones);
+					}, bUseParallel ? EParallelForFlags::None : EParallelForFlags::ForceSingleThread);
 
-					for (int BoneIdx = 0; BoneIdx < NumBones; ++BoneIdx)
+				// Cloth collection data arrays we are writing to
+				const TArrayView<int32> RenderNumBoneInfluences = ClothLodFacade.GetRenderNumBoneInfluences();
+				const TArrayView<TArray<int32>> RenderBoneIndices = ClothLodFacade.GetRenderBoneIndices();
+				const TArrayView<TArray<float>> RenderBoneWeights = ClothLodFacade.GetRenderBoneWeights();
+				const TConstArrayView<FVector3f> RenderPositions = ClothLodFacade.GetRenderPosition();
+				const int32 NumRenderVertices = ClothLodFacade.GetNumRenderVertices();
+
+				// Iterate over each vertex and write the data from FBoneWeights into cloth collection managed arrays
+				ParallelFor(NumRenderVertices,
+					[&RenderPositions, &TransferBoneWeights, &TargetBoneToIndex, &RenderNumBoneInfluences, &RenderBoneIndices, &RenderBoneWeights](int32 VertexID)
 					{
-						SimBoneIndices[VertexID][BoneIdx] = BoneWeights[BoneIdx].GetBoneIndex();
-						SimBoneWeights[VertexID][BoneIdx] = BoneWeights[BoneIdx].GetWeight();
+						const FVector3d PosD = FVector3d(RenderPositions[VertexID]);
+
+						UE::AnimationCore::FBoneWeights BoneWeights;
+						TransferBoneWeights.Compute(PosD, FTransformSRT3d::Identity(), BoneWeights, &TargetBoneToIndex);
+
+						const int32 NumBones = BoneWeights.Num();
 
-						RenderBoneIndices[VertexID][BoneIdx] = BoneWeights[BoneIdx].GetBoneIndex();
-						RenderBoneWeights[VertexID][BoneIdx] = BoneWeights[BoneIdx].GetWeight();
-					}
+						RenderNumBoneInfluences[VertexID] = NumBones;
+						RenderBoneIndices[VertexID].SetNum(NumBones);
+						RenderBoneWeights[VertexID].SetNum(NumBones);
 
-				}, bUseParallel ? EParallelForFlags::None : EParallelForFlags::ForceSingleThread);
+						for (int BoneIdx = 0; BoneIdx < NumBones; ++BoneIdx)
+						{
+							RenderBoneIndices[VertexID][BoneIdx] = BoneWeights[BoneIdx].GetBoneIndex();
+							RenderBoneWeights[VertexID][BoneIdx] = BoneWeights[BoneIdx].GetWeight();
+						}
 
+					}, bUseParallel ? EParallelForFlags::None : EParallelForFlags::ForceSingleThread);
 
 				ClothLodFacade.SetSkeletonAssetPathName(InputSkeleleton->GetPathName());
 			}
