Move cloth asset editor classes into UE::Chaos::ClothAsset namespace.

Just moved everything that isn't a USTRUCT/UCLASS or Slate class to the namespace.

diff --git a/Engine/Plugins/Experimental/ChaosClothAssetEditor/Source/ChaosClothAssetEditor/Private/ChaosClothAsset/ChaosClothAssetEditorModule.cpp b/Engine/Plugins/Experimental/ChaosClothAssetEditor/Source/ChaosClothAssetEditor/Private/ChaosClothAsset/ChaosClothAssetEditorModule.cpp
--- a/Engine/Plugins/Experimental/ChaosClothAssetEditor/Source/ChaosClothAssetEditor/Private/ChaosClothAsset/ChaosClothAssetEditorModule.cpp
+++ b/Engine/Plugins/Experimental/ChaosClothAssetEditor/Source/ChaosClothAssetEditor/Private/ChaosClothAsset/ChaosClothAssetEditorModule.cpp
@@ -14,10 +14,10 @@
 
 #define LOCTEXT_NAMESPACE "FChaosClothAssetEditorModule"
 
+namespace UE::Chaos::ClothAsset
+{
 void FChaosClothAssetEditorModule::StartupModule()
 {
-	using namespace UE::Chaos::ClothAsset;
-
 	FChaosClothAssetEditorStyle::Get(); // Causes the constructor to be called
 
 	FChaosClothAssetEditorCommands::Register();
@@ -37,7 +37,8 @@ void FChaosClothAssetEditorModule::ShutdownModule()
 
 	// TODO: Unregister details view customizations
 }
+} // namespace UE::Chaos::ClothAsset
 
 #undef LOCTEXT_NAMESPACE
 	
-IMPLEMENT_MODULE(FChaosClothAssetEditorModule, ChaosClothAssetEditor)
\ No newline at end of file
+IMPLEMENT_MODULE(UE::Chaos::ClothAsset::FChaosClothAssetEditorModule, ChaosClothAssetEditor)
\ No newline at end of file
diff --git a/Engine/Plugins/Experimental/ChaosClothAssetEditor/Source/ChaosClothAssetEditor/Private/ChaosClothAsset/ChaosClothAssetEditorModule.h b/Engine/Plugins/Experimental/ChaosClothAssetEditor/Source/ChaosClothAssetEditor/Private/ChaosClothAsset/ChaosClothAssetEditorModule.h
--- a/Engine/Plugins/Experimental/ChaosClothAssetEditor/Source/ChaosClothAssetEditor/Private/ChaosClothAsset/ChaosClothAssetEditorModule.h
+++ b/Engine/Plugins/Experimental/ChaosClothAssetEditor/Source/ChaosClothAssetEditor/Private/ChaosClothAsset/ChaosClothAssetEditorModule.h
@@ -6,11 +6,11 @@
 #include "BaseCharacterFXEditorModule.h"
 
 class FLayoutExtender;
+
 namespace UE::Chaos::ClothAsset
 {
-	class FAssetTypeActions_ClothAsset;
-	class FAssetTypeActions_ClothPreset;
-}
+class FAssetTypeActions_ClothAsset;
+class FAssetTypeActions_ClothPreset;
 
 class CHAOSCLOTHASSETEDITOR_API FChaosClothAssetEditorModule : public FBaseCharacterFXEditorModule
 {
@@ -21,7 +21,8 @@ public:
 	virtual void ShutdownModule() override;
 
 private:
-	UE::Chaos::ClothAsset::FAssetTypeActions_ClothAsset* AssetTypeActions_ClothAsset;
-	UE::Chaos::ClothAsset::FAssetTypeActions_ClothPreset* AssetTypeActions_ClothPreset;
+	FAssetTypeActions_ClothAsset* AssetTypeActions_ClothAsset;
+	FAssetTypeActions_ClothPreset* AssetTypeActions_ClothPreset;
 
 };
+} // namespace UE::Chaos::ClothAsset
diff --git a/Engine/Plugins/Experimental/ChaosClothAssetEditor/Source/ChaosClothAssetEditor/Private/ChaosClothAsset/ClothEditor.cpp b/Engine/Plugins/Experimental/ChaosClothAssetEditor/Source/ChaosClothAssetEditor/Private/ChaosClothAsset/ClothEditor.cpp
--- a/Engine/Plugins/Experimental/ChaosClothAssetEditor/Source/ChaosClothAssetEditor/Private/ChaosClothAsset/ClothEditor.cpp
+++ b/Engine/Plugins/Experimental/ChaosClothAssetEditor/Source/ChaosClothAssetEditor/Private/ChaosClothAsset/ClothEditor.cpp
@@ -7,6 +7,7 @@
 
 TSharedPtr<FBaseAssetToolkit> UChaosClothAssetEditor::CreateToolkit()
 {
+	using namespace UE::Chaos::ClothAsset;
 	TSharedPtr<FChaosClothAssetEditorToolkit> Toolkit = MakeShared<FChaosClothAssetEditorToolkit>(this);
 	return Toolkit;
 }
diff --git a/Engine/Plugins/Experimental/ChaosClothAssetEditor/Source/ChaosClothAssetEditor/Private/ChaosClothAsset/ClothEditor3DViewportClient.cpp b/Engine/Plugins/Experimental/ChaosClothAssetEditor/Source/ChaosClothAssetEditor/Private/ChaosClothAsset/ClothEditor3DViewportClient.cpp
--- a/Engine/Plugins/Experimental/ChaosClothAssetEditor/Source/ChaosClothAssetEditor/Private/ChaosClothAsset/ClothEditor3DViewportClient.cpp
+++ b/Engine/Plugins/Experimental/ChaosClothAssetEditor/Source/ChaosClothAssetEditor/Private/ChaosClothAsset/ClothEditor3DViewportClient.cpp
@@ -20,6 +20,8 @@
 #include "AssetViewerSettings.h"
 #include "Editor/EditorPerProjectUserSettings.h"
 
+namespace UE::Chaos::ClothAsset
+{
 FChaosClothAssetEditor3DViewportClient::FChaosClothAssetEditor3DViewportClient(FEditorModeTools* InModeTools,
 	TSharedPtr<FChaosClothPreviewScene> InPreviewScene,
 	TSharedPtr<FClothEditorSimulationVisualization> InVisualization,
@@ -232,7 +234,7 @@ void FChaosClothAssetEditor3DViewportClient::Draw(const FSceneView* View, FPrimi
 	FEditorViewportClient::Draw(View, PDI);
 	TSharedPtr<FClothEditorSimulationVisualization> Visualization = ClothEditorSimulationVisualization.Pin();
 	UChaosClothComponent* const ClothComponent = GetPreviewClothComponent();
-	if(Visualization && ClothComponent)
+	if (Visualization && ClothComponent)
 	{
 		Visualization->DebugDrawSimulation(ClothComponent, PDI);
 	}
@@ -312,7 +314,7 @@ void FChaosClothAssetEditor3DViewportClient::ProcessClick(FSceneView& View, HHit
 		const HActor* const ActorProxy = static_cast<HActor*>(HitProxy);
 		if (ActorProxy && ActorProxy->Actor)
 		{
-			const AActor* const Actor = ActorProxy->Actor;		
+			const AActor* const Actor = ActorProxy->Actor;
 			USceneComponent* RootComponent = Actor->GetRootComponent();
 			if (RootComponent)
 			{
@@ -325,7 +327,7 @@ void FChaosClothAssetEditor3DViewportClient::ProcessClick(FSceneView& View, HHit
 					// Don't use USelection::ToggleSelect here because that checks membership in GObjectSelection, not the given USelection...
 					TArray<USceneComponent*> Components;
 					SelectedComponents->GetSelectedObjects(Components);
-					const bool bIsSelected = Components.Contains(RootComponent);	
+					const bool bIsSelected = Components.Contains(RootComponent);
 
 					if (bIsSelected)
 					{
@@ -430,4 +432,5 @@ void FChaosClothAssetEditor3DViewportClient::SetAdvancedShowFlagsForScene(const
 	{
 		EngineShowFlags.DisableAdvancedFeatures();
 	}
-}
\ No newline at end of file
+}
+} // namespace UE::Chaos::ClothAsset
diff --git a/Engine/Plugins/Experimental/ChaosClothAssetEditor/Source/ChaosClothAssetEditor/Private/ChaosClothAsset/ClothEditor3DViewportClient.h b/Engine/Plugins/Experimental/ChaosClothAssetEditor/Source/ChaosClothAssetEditor/Private/ChaosClothAsset/ClothEditor3DViewportClient.h
--- a/Engine/Plugins/Experimental/ChaosClothAssetEditor/Source/ChaosClothAssetEditor/Private/ChaosClothAsset/ClothEditor3DViewportClient.h
+++ b/Engine/Plugins/Experimental/ChaosClothAssetEditor/Source/ChaosClothAssetEditor/Private/ChaosClothAsset/ClothEditor3DViewportClient.h
@@ -8,9 +8,13 @@
 
 class UChaosClothComponent;
 class UChaosClothAssetEditorMode;
-class FChaosClothAssetEditorToolkit;
 class UTransformProxy;
 class UCombinedTransformGizmo;
+
+namespace UE::Chaos::ClothAsset
+{
+
+class FChaosClothAssetEditorToolkit;
 class FChaosClothPreviewScene;
 class FClothEditorSimulationVisualization;
 
@@ -36,7 +40,7 @@ public:
 
 	void ClearSelectedComponents();
 
-	void EnableSimMeshWireframe(bool bEnable ) { bSimMeshWireframe = bEnable; }
+	void EnableSimMeshWireframe(bool bEnable) { bSimMeshWireframe = bEnable; }
 	bool SimMeshWireframeEnabled() const { return bSimMeshWireframe; }
 
 	void EnableRenderMeshWireframe(bool bEnable);
@@ -102,3 +106,4 @@ private:
 	TObjectPtr<UCombinedTransformGizmo> Gizmo = nullptr;
 
 };
+} // namespace UE::Chaos::ClothAsset
diff --git a/Engine/Plugins/Experimental/ChaosClothAssetEditor/Source/ChaosClothAssetEditor/Private/ChaosClothAsset/ClothEditorCommands.cpp b/Engine/Plugins/Experimental/ChaosClothAssetEditor/Source/ChaosClothAssetEditor/Private/ChaosClothAsset/ClothEditorCommands.cpp
--- a/Engine/Plugins/Experimental/ChaosClothAssetEditor/Source/ChaosClothAssetEditor/Private/ChaosClothAsset/ClothEditorCommands.cpp
+++ b/Engine/Plugins/Experimental/ChaosClothAssetEditor/Source/ChaosClothAssetEditor/Private/ChaosClothAsset/ClothEditorCommands.cpp
@@ -6,6 +6,9 @@
 
 #define LOCTEXT_NAMESPACE "FChaosClothAssetEditorCommands"
 
+namespace UE::Chaos::ClothAsset
+{
+
 const FString FChaosClothAssetEditorCommands::BeginRemeshToolIdentifier = TEXT("BeginRemeshTool");
 const FString FChaosClothAssetEditorCommands::BeginAttributeEditorToolIdentifier = TEXT("BeginAttributeEditorTool");
 const FString FChaosClothAssetEditorCommands::BeginWeightMapPaintToolIdentifier = TEXT("BeginWeightMapPaintTool");
@@ -72,6 +75,6 @@ void FChaosClothAssetEditorCommands::UpdateToolCommandBinding(UInteractiveTool*
 		}
 	}
 }
-
+} // namespace UE::Chaos::ClothAsset
 
 #undef LOCTEXT_NAMESPACE
diff --git a/Engine/Plugins/Experimental/ChaosClothAssetEditor/Source/ChaosClothAssetEditor/Private/ChaosClothAsset/ClothEditorCommands.h b/Engine/Plugins/Experimental/ChaosClothAssetEditor/Source/ChaosClothAssetEditor/Private/ChaosClothAsset/ClothEditorCommands.h
--- a/Engine/Plugins/Experimental/ChaosClothAssetEditor/Source/ChaosClothAssetEditor/Private/ChaosClothAsset/ClothEditorCommands.h
+++ b/Engine/Plugins/Experimental/ChaosClothAssetEditor/Source/ChaosClothAssetEditor/Private/ChaosClothAsset/ClothEditorCommands.h
@@ -4,6 +4,8 @@
 
 #include "BaseCharacterFXEditorCommands.h"
 
+namespace UE::Chaos::ClothAsset
+{
 class CHAOSCLOTHASSETEDITOR_API FChaosClothAssetEditorCommands : public TBaseCharacterFXEditorCommands<FChaosClothAssetEditorCommands>
 {
 public:
@@ -64,3 +66,4 @@ public:
 	TSharedPtr<FUICommandInfo> ToggleSimulationSuspended;
 
 };
+} // namespace UE::Chaos::ClothAsset
diff --git a/Engine/Plugins/Experimental/ChaosClothAssetEditor/Source/ChaosClothAssetEditor/Private/ChaosClothAsset/ClothEditorMode.cpp b/Engine/Plugins/Experimental/ChaosClothAssetEditor/Source/ChaosClothAssetEditor/Private/ChaosClothAsset/ClothEditorMode.cpp
--- a/Engine/Plugins/Experimental/ChaosClothAssetEditor/Source/ChaosClothAssetEditor/Private/ChaosClothAsset/ClothEditorMode.cpp
+++ b/Engine/Plugins/Experimental/ChaosClothAssetEditor/Source/ChaosClothAssetEditor/Private/ChaosClothAsset/ClothEditorMode.cpp
@@ -196,6 +196,8 @@ void UChaosClothAssetEditorMode::RegisterClothTool(TSharedPtr<FUICommandInfo> UI
 
 void UChaosClothAssetEditorMode::RegisterPreviewTools()
 {
+	using namespace UE::Chaos::ClothAsset;
+
 	UEditorInteractiveToolsContext* const PreviewSceneToolsContext = PreviewScene->GetClothPreviewEditorModeManager()->GetInteractiveToolsContext();
 
 	const FChaosClothAssetEditorCommands& CommandInfos = FChaosClothAssetEditorCommands::Get();
@@ -204,6 +206,8 @@ void UChaosClothAssetEditorMode::RegisterPreviewTools()
 
 void UChaosClothAssetEditorMode::RegisterTools()
 {
+	using namespace UE::Chaos::ClothAsset;
+
 	const FChaosClothAssetEditorCommands& CommandInfos = FChaosClothAssetEditorCommands::Get();
 
 	// Note that the identifiers below need to match the command names so that the tool icons can 
@@ -241,11 +245,13 @@ bool UChaosClothAssetEditorMode::ShouldToolStartBeAllowed(const FString& ToolIde
 
 void UChaosClothAssetEditorMode::CreateToolkit()
 {
-	Toolkit = MakeShared<FChaosClothAssetEditorModeToolkit>();
+	Toolkit = MakeShared<UE::Chaos::ClothAsset::FChaosClothAssetEditorModeToolkit>();
 }
 
 void UChaosClothAssetEditorMode::OnToolStarted(UInteractiveToolManager* Manager, UInteractiveTool* Tool)
 {
+	using namespace UE::Chaos::ClothAsset;
+
 	FChaosClothAssetEditorCommands::UpdateToolCommandBinding(Tool, ToolCommandList, false);
 
 	bCanTogglePattern2DMode = false;
@@ -259,7 +265,7 @@ void UChaosClothAssetEditorMode::OnToolStarted(UInteractiveToolManager* Manager,
 
 void UChaosClothAssetEditorMode::OnToolEnded(UInteractiveToolManager* Manager, UInteractiveTool* Tool)
 {
-	FChaosClothAssetEditorCommands::UpdateToolCommandBinding(Tool, ToolCommandList, true);
+	UE::Chaos::ClothAsset::FChaosClothAssetEditorCommands::UpdateToolCommandBinding(Tool, ToolCommandList, true);
 
 	bCanTogglePattern2DMode = true;
 	ActiveToolsContext = nullptr;
@@ -279,6 +285,7 @@ void UChaosClothAssetEditorMode::PostUndo()
 
 void UChaosClothAssetEditorMode::BindCommands()
 {
+	using namespace UE::Chaos::ClothAsset;
 	const FChaosClothAssetEditorCommands& CommandInfos = FChaosClothAssetEditorCommands::Get();
 	const TSharedRef<FUICommandList>& CommandList = Toolkit->GetToolkitCommands();
 
@@ -335,8 +342,10 @@ void UChaosClothAssetEditorMode::Exit()
 	Super::Exit();
 }
 
-void UChaosClothAssetEditorMode::SetPreviewScene(FChaosClothPreviewScene* InPreviewScene)
+void UChaosClothAssetEditorMode::SetPreviewScene(UE::Chaos::ClothAsset::FChaosClothPreviewScene* InPreviewScene)
 {
+	using namespace UE::Chaos::ClothAsset;
+
 	PreviewScene = InPreviewScene;
 
 
@@ -887,7 +896,7 @@ void UChaosClothAssetEditorMode::TogglePatternMode()
 	bPattern2DMode = !bPattern2DMode;
 	ReinitializeDynamicMeshComponents();
 
-	TSharedPtr<FChaosClothEditorRestSpaceViewportClient> VC = RestSpaceViewportClient.Pin();
+	TSharedPtr<UE::Chaos::ClothAsset::FChaosClothEditorRestSpaceViewportClient> VC = RestSpaceViewportClient.Pin();
 	if (VC.IsValid())
 	{
 		VC->Set2DMode(bPattern2DMode);
@@ -902,11 +911,11 @@ bool UChaosClothAssetEditorMode::CanTogglePatternMode() const
 }
 
 
-void UChaosClothAssetEditorMode::SetRestSpaceViewportClient(TWeakPtr<FChaosClothEditorRestSpaceViewportClient, ESPMode::ThreadSafe> InViewportClient)
+void UChaosClothAssetEditorMode::SetRestSpaceViewportClient(TWeakPtr<UE::Chaos::ClothAsset::FChaosClothEditorRestSpaceViewportClient, ESPMode::ThreadSafe> InViewportClient)
 {
 	RestSpaceViewportClient = InViewportClient;
 
-	TSharedPtr<FChaosClothEditorRestSpaceViewportClient> VC = RestSpaceViewportClient.Pin();
+	TSharedPtr<UE::Chaos::ClothAsset::FChaosClothEditorRestSpaceViewportClient> VC = RestSpaceViewportClient.Pin();
 	if (VC.IsValid())
 	{
 		VC->Set2DMode(bPattern2DMode);
diff --git a/Engine/Plugins/Experimental/ChaosClothAssetEditor/Source/ChaosClothAssetEditor/Private/ChaosClothAsset/ClothEditorMode.h b/Engine/Plugins/Experimental/ChaosClothAssetEditor/Source/ChaosClothAssetEditor/Private/ChaosClothAsset/ClothEditorMode.h
--- a/Engine/Plugins/Experimental/ChaosClothAssetEditor/Source/ChaosClothAssetEditor/Private/ChaosClothAsset/ClothEditorMode.h
+++ b/Engine/Plugins/Experimental/ChaosClothAssetEditor/Source/ChaosClothAssetEditor/Private/ChaosClothAsset/ClothEditorMode.h
@@ -17,21 +17,25 @@ class UPreviewMesh;
 class UToolTarget;
 class FToolTargetTypeRequirements;
 class UWorld;
-class FChaosClothAssetEditorModeToolkit;
 class UInteractiveToolPropertySet; 
 class UMeshOpPreviewWithBackgroundCompute;
 class UClothToolViewportButtonsAPI;
 class UDynamicMeshComponent;
 class UChaosClothComponent;
 class FEditorViewportClient;
-class FChaosClothEditorRestSpaceViewportClient;
 class FViewport;
 class UDataflow;
 class UDataflowComponent;
-class FChaosClothPreviewScene;
-class FChaosClothAssetEditor3DViewportClient;
 class SDataflowGraphEditor;
 struct FManagedArrayCollection;
+namespace UE::Chaos::ClothAsset
+{
+class FChaosClothPreviewScene;
+class FChaosClothAssetEditor3DViewportClient;
+class FChaosClothAssetEditorModeToolkit;
+class FChaosClothAssetEditorToolkit;
+class FChaosClothEditorRestSpaceViewportClient;
+}
 
 /**
  * The cloth editor mode is the mode used in the cloth asset editor. It holds most of the inter-tool state.
@@ -77,8 +81,8 @@ public:
 
 private:
 
-	friend class FChaosClothAssetEditorToolkit;
-	friend class FChaosClothAssetEditorModeToolkit;
+	friend class UE::Chaos::ClothAsset::FChaosClothAssetEditorToolkit;
+	friend class UE::Chaos::ClothAsset::FChaosClothAssetEditorModeToolkit;
 
 	// UEdMode
 	virtual void Enter() final;
@@ -114,12 +118,12 @@ private:
 	// Register the set of tools that operate on objects in the 3D preview world 
 	void RegisterPreviewTools();
 
-	void SetPreviewScene(FChaosClothPreviewScene* PreviewScene);
+	void SetPreviewScene(UE::Chaos::ClothAsset::FChaosClothPreviewScene* PreviewScene);
 
 	// Bounding box for rest space meshes
 	virtual FBox SceneBoundingBox() const override;
 
-	void SetRestSpaceViewportClient(TWeakPtr<FChaosClothEditorRestSpaceViewportClient, ESPMode::ThreadSafe> ViewportClient);
+	void SetRestSpaceViewportClient(TWeakPtr<UE::Chaos::ClothAsset::FChaosClothEditorRestSpaceViewportClient, ESPMode::ThreadSafe> ViewportClient);
 	void RefocusRestSpaceViewportClient();
 
 	// intended to be called by the toolkit when selected node in the Dataflow graph changes
@@ -138,7 +142,7 @@ private:
 	TArray<TObjectPtr<UMeshElementsVisualizer>> WireframesToTick;
 
 	// Preview Scene, here largely for convenience to avoid having to pass it around functions. Owned by the ClothEditorToolkit.
-	FChaosClothPreviewScene* PreviewScene = nullptr;
+	UE::Chaos::ClothAsset::FChaosClothPreviewScene* PreviewScene = nullptr;
 
 	// Mode-level property objects (visible or not) that get ticked.
 	UPROPERTY()
@@ -160,7 +164,7 @@ private:
 	};
 	TArray<FDynamicMeshSourceInfo> DynamicMeshSourceInfos;
 
-	TWeakPtr<FChaosClothEditorRestSpaceViewportClient, ESPMode::ThreadSafe> RestSpaceViewportClient;
+	TWeakPtr<UE::Chaos::ClothAsset::FChaosClothEditorRestSpaceViewportClient, ESPMode::ThreadSafe> RestSpaceViewportClient;
 
 	// Handle to a callback triggered when the current selection changes
 	FDelegateHandle SelectionModifiedEventHandle;
diff --git a/Engine/Plugins/Experimental/ChaosClothAssetEditor/Source/ChaosClothAssetEditor/Private/ChaosClothAsset/ClothEditorModeToolkit.cpp b/Engine/Plugins/Experimental/ChaosClothAssetEditor/Source/ChaosClothAssetEditor/Private/ChaosClothAsset/ClothEditorModeToolkit.cpp
--- a/Engine/Plugins/Experimental/ChaosClothAssetEditor/Source/ChaosClothAssetEditor/Private/ChaosClothAsset/ClothEditorModeToolkit.cpp
+++ b/Engine/Plugins/Experimental/ChaosClothAssetEditor/Source/ChaosClothAssetEditor/Private/ChaosClothAsset/ClothEditorModeToolkit.cpp
@@ -23,6 +23,9 @@
 	
 #define LOCTEXT_NAMESPACE "FChaosClothAssetEditorModeToolkit"
 
+namespace UE::Chaos::ClothAsset
+{
+
 FName FChaosClothAssetEditorModeToolkit::GetToolkitFName() const
 {
 	return FName("ChaosClothAssetEditorMode");
@@ -274,5 +277,6 @@ void FChaosClothAssetEditorModeToolkit::SetPreviewViewportWidget(TWeakPtr<SChaos
 	PreviewViewportWidget = InPreviewViewportWidget;
 }
 
+} // namespace UE::Chaos::ClothAsset
 
 #undef LOCTEXT_NAMESPACE
diff --git a/Engine/Plugins/Experimental/ChaosClothAssetEditor/Source/ChaosClothAssetEditor/Private/ChaosClothAsset/ClothEditorModeToolkit.h b/Engine/Plugins/Experimental/ChaosClothAssetEditor/Source/ChaosClothAssetEditor/Private/ChaosClothAsset/ClothEditorModeToolkit.h
--- a/Engine/Plugins/Experimental/ChaosClothAssetEditor/Source/ChaosClothAssetEditor/Private/ChaosClothAsset/ClothEditorModeToolkit.h
+++ b/Engine/Plugins/Experimental/ChaosClothAssetEditor/Source/ChaosClothAssetEditor/Private/ChaosClothAsset/ClothEditorModeToolkit.h
@@ -15,6 +15,8 @@ class SBaseCharacterFXEditorViewport;
  * NOTE: the cloth editor has two separate viewports/worlds/modemanagers/toolscontexts, so we need to track which
  * one is currently active.
  */
+namespace UE::Chaos::ClothAsset
+{
 class CHAOSCLOTHASSETEDITOR_API FChaosClothAssetEditorModeToolkit : public FBaseCharacterFXEditorModeToolkit
 {
 public:
@@ -47,3 +49,4 @@ private:
 	TWeakPtr<SChaosClothAssetEditor3DViewport> PreviewViewportWidget;
 
 };
+} // namespace UE::Chaos::ClothAsset
diff --git a/Engine/Plugins/Experimental/ChaosClothAssetEditor/Source/ChaosClothAssetEditor/Private/ChaosClothAsset/ClothEditorModeUILayer.cpp b/Engine/Plugins/Experimental/ChaosClothAssetEditor/Source/ChaosClothAssetEditor/Private/ChaosClothAsset/ClothEditorModeUILayer.cpp
--- a/Engine/Plugins/Experimental/ChaosClothAssetEditor/Source/ChaosClothAssetEditor/Private/ChaosClothAsset/ClothEditorModeUILayer.cpp
+++ b/Engine/Plugins/Experimental/ChaosClothAssetEditor/Source/ChaosClothAssetEditor/Private/ChaosClothAsset/ClothEditorModeUILayer.cpp
@@ -14,12 +14,14 @@ const FName UChaosClothAssetEditorUISubsystem::EditorSidePanelAreaName = "ChaosC
 
 void UChaosClothAssetEditorUISubsystem::Initialize(FSubsystemCollectionBase& Collection)
 {
+	using namespace UE::Chaos::ClothAsset;
 	FChaosClothAssetEditorModule& ChaosClothAssetEditorModule = FModuleManager::GetModuleChecked<FChaosClothAssetEditorModule>("ChaosClothAssetEditor");
 	ChaosClothAssetEditorModule.OnRegisterLayoutExtensions().AddUObject(this, &UChaosClothAssetEditorUISubsystem::RegisterLayoutExtensions);
 }
 
 void UChaosClothAssetEditorUISubsystem::Deinitialize()
 {
+	using namespace UE::Chaos::ClothAsset;
 	FChaosClothAssetEditorModule& ChaosClothAssetEditorModule = FModuleManager::GetModuleChecked<FChaosClothAssetEditorModule>("ChaosClothAssetEditor");
 	ChaosClothAssetEditorModule.OnRegisterLayoutExtensions().RemoveAll(this);
 }
diff --git a/Engine/Plugins/Experimental/ChaosClothAssetEditor/Source/ChaosClothAssetEditor/Private/ChaosClothAsset/ClothEditorPreviewScene.cpp b/Engine/Plugins/Experimental/ChaosClothAssetEditor/Source/ChaosClothAssetEditor/Private/ChaosClothAsset/ClothEditorPreviewScene.cpp
--- a/Engine/Plugins/Experimental/ChaosClothAssetEditor/Source/ChaosClothAssetEditor/Private/ChaosClothAsset/ClothEditorPreviewScene.cpp
+++ b/Engine/Plugins/Experimental/ChaosClothAssetEditor/Source/ChaosClothAssetEditor/Private/ChaosClothAsset/ClothEditorPreviewScene.cpp
@@ -11,7 +11,7 @@
 
 #define LOCTEXT_NAMESPACE "UChaosClothEditorPreviewScene"
 
-void UChaosClothPreviewSceneDescription::SetPreviewScene(FChaosClothPreviewScene* InPreviewScene)
+void UChaosClothPreviewSceneDescription::SetPreviewScene(UE::Chaos::ClothAsset::FChaosClothPreviewScene* InPreviewScene)
 {
 	PreviewScene = InPreviewScene;
 }
@@ -24,7 +24,8 @@ void UChaosClothPreviewSceneDescription::PostEditChangeProperty(FPropertyChanged
 	}
 }
 
-
+namespace UE::Chaos::ClothAsset
+{
 
 FChaosClothPreviewScene::FChaosClothPreviewScene(FPreviewScene::ConstructionValues ConstructionValues) :
 	FAdvancedPreviewScene(ConstructionValues)
@@ -290,6 +291,7 @@ const UAnimSingleNodeInstance* const FChaosClothPreviewScene::GetPreviewAnimInst
 {
 	return PreviewAnimInstance;
 }
+} // namespace UE::Chaos::ClothAsset
 
 #undef LOCTEXT_NAMESPACE
 
diff --git a/Engine/Plugins/Experimental/ChaosClothAssetEditor/Source/ChaosClothAssetEditor/Private/ChaosClothAsset/ClothEditorPreviewScene.h b/Engine/Plugins/Experimental/ChaosClothAssetEditor/Source/ChaosClothAssetEditor/Private/ChaosClothAsset/ClothEditorPreviewScene.h
--- a/Engine/Plugins/Experimental/ChaosClothAssetEditor/Source/ChaosClothAssetEditor/Private/ChaosClothAsset/ClothEditorPreviewScene.h
+++ b/Engine/Plugins/Experimental/ChaosClothAssetEditor/Source/ChaosClothAssetEditor/Private/ChaosClothAsset/ClothEditorPreviewScene.h
@@ -9,11 +9,15 @@
 class UChaosClothAsset;
 class USkeletalMesh;
 class UChaosClothComponent;
-class FChaosClothPreviewScene;
 class FAssetEditorModeManager;
 class UAnimationAsset;
 class UAnimSingleNodeInstance;
 
+namespace UE::Chaos::ClothAsset
+{
+class FChaosClothPreviewScene;
+}
+
 ///
 /// The UChaosClothPreviewSceneDescription is a description of the Preview scene contents, intended to be editable in an FAdvancedPreviewSettingsWidget
 /// 
@@ -23,7 +27,7 @@ class CHAOSCLOTHASSETEDITOR_API UChaosClothPreviewSceneDescription : public UObj
 public:
 	GENERATED_BODY()
 
-	void SetPreviewScene(FChaosClothPreviewScene* PreviewScene);
+	void SetPreviewScene(UE::Chaos::ClothAsset::FChaosClothPreviewScene* PreviewScene);
 
 	// Skeletal Mesh source asset
 	UPROPERTY(EditAnywhere, Category="SkeletalMesh")
@@ -40,10 +44,12 @@ private:
 	// Listen for changes to the scene description members and notify the PreviewScene
 	virtual void PostEditChangeProperty(struct FPropertyChangedEvent& PropertyChangedEvent) override;
 
-	FChaosClothPreviewScene* PreviewScene;
+	UE::Chaos::ClothAsset::FChaosClothPreviewScene* PreviewScene;
 };
 
 
+namespace UE::Chaos::ClothAsset
+{
 ///
 /// FChaosClothPreviewScene is the actual Preview scene, with contents specified by the SceneDescription
 /// 
@@ -101,5 +107,5 @@ private:
 	TObjectPtr<USkeletalMeshComponent> SkeletalMeshComponent;
 
 };
-
+} // namespace UE::Chaos::ClothAsset
 
diff --git a/Engine/Plugins/Experimental/ChaosClothAssetEditor/Source/ChaosClothAssetEditor/Private/ChaosClothAsset/ClothEditorRestSpaceViewportClient.cpp b/Engine/Plugins/Experimental/ChaosClothAssetEditor/Source/ChaosClothAssetEditor/Private/ChaosClothAsset/ClothEditorRestSpaceViewportClient.cpp
--- a/Engine/Plugins/Experimental/ChaosClothAssetEditor/Source/ChaosClothAssetEditor/Private/ChaosClothAsset/ClothEditorRestSpaceViewportClient.cpp
+++ b/Engine/Plugins/Experimental/ChaosClothAssetEditor/Source/ChaosClothAssetEditor/Private/ChaosClothAsset/ClothEditorRestSpaceViewportClient.cpp
@@ -10,6 +10,8 @@
 #include "Components/DynamicMeshComponent.h"
 #include "Framework/Application/SlateApplication.h"
 
+namespace UE::Chaos::ClothAsset
+{
 FChaosClothEditorRestSpaceViewportClient::FChaosClothEditorRestSpaceViewportClient(FEditorModeTools* InModeTools,
 	FPreviewScene* InPreviewScene,
 	const TWeakPtr<SEditorViewport>& InEditorViewportWidget) :
@@ -86,7 +88,7 @@ void FChaosClothEditorRestSpaceViewportClient::ProcessClick(FSceneView& View, HH
 	if (HitProxy && HitProxy->IsA(HActor::StaticGetType()))
 	{
 		const HActor* ActorProxy = static_cast<HActor*>(HitProxy);
-		if (ActorProxy && ActorProxy->Actor )
+		if (ActorProxy && ActorProxy->Actor)
 		{
 			const AActor* Actor = ActorProxy->Actor;
 			const TSet<UActorComponent*>& OwnedComponents = Actor->GetComponents();
@@ -118,4 +120,5 @@ void FChaosClothEditorRestSpaceViewportClient::SetEditorViewportWidget(TWeakPtr<
 void FChaosClothEditorRestSpaceViewportClient::SetToolCommandList(TWeakPtr<FUICommandList> InToolCommandList)
 {
 	ToolCommandList = InToolCommandList;
-}
\ No newline at end of file
+}
+} // namespace UE::Chaos::ClothAsset
diff --git a/Engine/Plugins/Experimental/ChaosClothAssetEditor/Source/ChaosClothAssetEditor/Private/ChaosClothAsset/ClothEditorRestSpaceViewportClient.h b/Engine/Plugins/Experimental/ChaosClothAssetEditor/Source/ChaosClothAssetEditor/Private/ChaosClothAsset/ClothEditorRestSpaceViewportClient.h
--- a/Engine/Plugins/Experimental/ChaosClothAssetEditor/Source/ChaosClothAssetEditor/Private/ChaosClothAsset/ClothEditorRestSpaceViewportClient.h
+++ b/Engine/Plugins/Experimental/ChaosClothAssetEditor/Source/ChaosClothAssetEditor/Private/ChaosClothAsset/ClothEditorRestSpaceViewportClient.h
@@ -9,6 +9,8 @@
 
 class UInputBehaviorSet;
 
+namespace UE::Chaos::ClothAsset
+{
 class CHAOSCLOTHASSETEDITOR_API FChaosClothEditorRestSpaceViewportClient : public FEditorViewportClient
 {
 public:
@@ -34,4 +36,4 @@ private:
 
 	TWeakPtr<FUICommandList> ToolCommandList;
 };
-
+} // namespace UE::Chaos::ClothAsset
diff --git a/Engine/Plugins/Experimental/ChaosClothAssetEditor/Source/ChaosClothAssetEditor/Private/ChaosClothAsset/ClothEditorSimulationVisualization.cpp b/Engine/Plugins/Experimental/ChaosClothAssetEditor/Source/ChaosClothAssetEditor/Private/ChaosClothAsset/ClothEditorSimulationVisualization.cpp
--- a/Engine/Plugins/Experimental/ChaosClothAssetEditor/Source/ChaosClothAssetEditor/Private/ChaosClothAsset/ClothEditorSimulationVisualization.cpp
+++ b/Engine/Plugins/Experimental/ChaosClothAssetEditor/Source/ChaosClothAssetEditor/Private/ChaosClothAsset/ClothEditorSimulationVisualization.cpp
@@ -11,6 +11,8 @@
 
 #define LOCTEXT_NAMESPACE "ClothEditorSimulationVisualization"
 
+namespace UE::Chaos::ClothAsset
+{
 namespace Private
 {
 struct FVisualizationOption
@@ -207,4 +209,6 @@ void FClothEditorSimulationVisualization::ShowClothSections(UChaosClothComponent
 		}
 	}
 }
+} // namespace UE::Chaos::ClothAsset
+
 #undef LOCTEXT_NAMESPACE
diff --git a/Engine/Plugins/Experimental/ChaosClothAssetEditor/Source/ChaosClothAssetEditor/Private/ChaosClothAsset/ClothEditorSimulationVisualization.h b/Engine/Plugins/Experimental/ChaosClothAssetEditor/Source/ChaosClothAssetEditor/Private/ChaosClothAsset/ClothEditorSimulationVisualization.h
--- a/Engine/Plugins/Experimental/ChaosClothAssetEditor/Source/ChaosClothAssetEditor/Private/ChaosClothAsset/ClothEditorSimulationVisualization.h
+++ b/Engine/Plugins/Experimental/ChaosClothAssetEditor/Source/ChaosClothAssetEditor/Private/ChaosClothAsset/ClothEditorSimulationVisualization.h
@@ -4,13 +4,17 @@
 #include "Containers/BitArray.h"
 #include "Templates/SharedPointer.h"
 
-class FChaosClothAssetEditor3DViewportClient;
 class FMenuBuilder;
 class UChaosClothComponent;
 class FPrimitiveDrawInterface;
 class FCanvas;
 class FSceneView;
 
+namespace UE::Chaos::ClothAsset
+{
+
+class FChaosClothAssetEditor3DViewportClient;
+
 class FClothEditorSimulationVisualization
 {
 public:
@@ -29,4 +33,5 @@ private:
 private:
 	/** Flags used to store the checked status for the visualization options. */
 	TBitArray<> Flags;
-};
\ No newline at end of file
+};
+} // namespace UE::Chaos::ClothAsset
diff --git a/Engine/Plugins/Experimental/ChaosClothAssetEditor/Source/ChaosClothAssetEditor/Private/ChaosClothAsset/ClothEditorStyle.cpp b/Engine/Plugins/Experimental/ChaosClothAssetEditor/Source/ChaosClothAssetEditor/Private/ChaosClothAsset/ClothEditorStyle.cpp
--- a/Engine/Plugins/Experimental/ChaosClothAssetEditor/Source/ChaosClothAssetEditor/Private/ChaosClothAsset/ClothEditorStyle.cpp
+++ b/Engine/Plugins/Experimental/ChaosClothAssetEditor/Source/ChaosClothAssetEditor/Private/ChaosClothAsset/ClothEditorStyle.cpp
@@ -8,6 +8,8 @@
 #include "Interfaces/IPluginManager.h"
 #include "Misc/Paths.h"
 
+namespace UE::Chaos::ClothAsset
+{
 const FName FChaosClothAssetEditorStyle::StyleName("ClothStyle");
 
 FString FChaosClothAssetEditorStyle::InContent(const FString& RelativePath, const ANSICHAR* Extension)
@@ -27,7 +29,7 @@ FChaosClothAssetEditorStyle::FChaosClothAssetEditorStyle()
 	// TODO: Set this once we get our own UI content
 	FString ContentDir = IPluginManager::Get().FindPlugin(TEXT("ModelingToolsEditorMode"))->GetContentDir();
 	SetContentRoot(ContentDir);
-	
+
 	// Some standard icon sizes used elsewhere in the editor
 	const FVector2D Icon8x8(8.0f, 8.0f);
 	const FVector2D Icon16x16(16.0f, 16.0f);
@@ -41,8 +43,8 @@ FChaosClothAssetEditorStyle::FChaosClothAssetEditorStyle()
 	const FVector2D ToolbarIconSize = Icon20x20;
 
 	FString PropertyNameString = "ChaosClothAssetEditor." + FChaosClothAssetEditorCommands::BeginRemeshToolIdentifier;
-	Set(*PropertyNameString, new FSlateImageBrush(FChaosClothAssetEditorStyle::InContent("Icons/Remesh_40x", ".png"), ToolbarIconSize ));
-	
+	Set(*PropertyNameString, new FSlateImageBrush(FChaosClothAssetEditorStyle::InContent("Icons/Remesh_40x", ".png"), ToolbarIconSize));
+
 	PropertyNameString = "ChaosClothAssetEditor." + FChaosClothAssetEditorCommands::BeginAttributeEditorToolIdentifier;
 	Set(*PropertyNameString, new FSlateImageBrush(FChaosClothAssetEditorStyle::InContent("Icons/AttributeEditor_40x", ".png"), ToolbarIconSize));
 
@@ -81,7 +83,7 @@ FChaosClothAssetEditorStyle::FChaosClothAssetEditorStyle()
 	// TODO: Get a permanent icon to replace this placeholder
 	PropertyNameString = "ChaosClothAssetEditor." + FChaosClothAssetEditorCommands::HardResetSimulationIdentifier;
 	Set(*PropertyNameString, new CORE_IMAGE_BRUSH_SVG("Starship/Common/Apply", ToolbarIconSize));
-	
+
 	FSlateStyleRegistry::RegisterSlateStyle(*this);
 }
 
@@ -95,5 +97,4 @@ FChaosClothAssetEditorStyle& FChaosClothAssetEditorStyle::Get()
 	static FChaosClothAssetEditorStyle Inst;
 	return Inst;
 }
-
-
+} // namespace UE::Chaos::ClothAsset
diff --git a/Engine/Plugins/Experimental/ChaosClothAssetEditor/Source/ChaosClothAssetEditor/Private/ChaosClothAsset/ClothEditorStyle.h b/Engine/Plugins/Experimental/ChaosClothAssetEditor/Source/ChaosClothAssetEditor/Private/ChaosClothAsset/ClothEditorStyle.h
--- a/Engine/Plugins/Experimental/ChaosClothAssetEditor/Source/ChaosClothAssetEditor/Private/ChaosClothAsset/ClothEditorStyle.h
+++ b/Engine/Plugins/Experimental/ChaosClothAssetEditor/Source/ChaosClothAssetEditor/Private/ChaosClothAsset/ClothEditorStyle.h
@@ -5,6 +5,8 @@
 #include "Styling/SlateStyle.h"
 #include "Styling/SlateStyleRegistry.h"
 
+namespace UE::Chaos::ClothAsset
+{
 /**
  * Slate style set for Cloth Editor
  */
@@ -23,3 +25,4 @@ private:
 	FChaosClothAssetEditorStyle();
 	~FChaosClothAssetEditorStyle();
 };
+} // namespace UE::Chaos::ClothAsset
diff --git a/Engine/Plugins/Experimental/ChaosClothAssetEditor/Source/ChaosClothAssetEditor/Private/ChaosClothAsset/ClothEditorToolkit.cpp b/Engine/Plugins/Experimental/ChaosClothAssetEditor/Source/ChaosClothAssetEditor/Private/ChaosClothAsset/ClothEditorToolkit.cpp
--- a/Engine/Plugins/Experimental/ChaosClothAssetEditor/Source/ChaosClothAssetEditor/Private/ChaosClothAsset/ClothEditorToolkit.cpp
+++ b/Engine/Plugins/Experimental/ChaosClothAssetEditor/Source/ChaosClothAssetEditor/Private/ChaosClothAsset/ClothEditorToolkit.cpp
@@ -41,6 +41,8 @@
 
 #define LOCTEXT_NAMESPACE "ChaosClothAssetEditorToolkit"
 
+namespace UE::Chaos::ClothAsset
+{
 const FName FChaosClothAssetEditorToolkit::ClothPreviewTabID(TEXT("ChaosClothAssetEditor_ClothPreviewTab"));
 const FName FChaosClothAssetEditorToolkit::OutlinerTabID(TEXT("ChaosClothAssetEditor_OutlinerTab"));
 const FName FChaosClothAssetEditorToolkit::PreviewSceneDetailsTabID(TEXT("ChaosClothAssetEditor_PreviewSceneDetailsTab"));
@@ -48,7 +50,7 @@ const FName FChaosClothAssetEditorToolkit::GraphCanvasTabId(TEXT("ChaosClothAsse
 const FName FChaosClothAssetEditorToolkit::NodeDetailsTabId(TEXT("ChaosClothAssetEditor_NodeDetails"));
 
 
-namespace ClothEditorToolkitHelpers
+namespace Private
 {
 	UDataflow* GetDataflowFrom(const UObject* InObject)
 	{
@@ -236,7 +238,7 @@ void FChaosClothAssetEditorToolkit::Tick(float DeltaTime)
 			DataflowContext = TSharedPtr<Dataflow::FEngineContext>(new Dataflow::FClothAssetDataflowContext(ClothAsset, Dataflow, Dataflow::FTimestamp::Invalid));
 			LastDataflowNodeTimestamp = Dataflow::FTimestamp::Invalid;
 		}
-		DataflowTerminalPath = ClothEditorToolkitHelpers::GetDataflowTerminalFrom(ClothAsset);
+		DataflowTerminalPath = Private::GetDataflowTerminalFrom(ClothAsset);
 
 		Dataflow::FTimestamp OldTimestamp = LastDataflowNodeTimestamp;
 		FDataflowEditorCommands::EvaluateTerminalNode(*DataflowContext.Get(), LastDataflowNodeTimestamp, Dataflow, nullptr, nullptr, ClothAsset, DataflowTerminalPath);
@@ -308,12 +310,12 @@ void FChaosClothAssetEditorToolkit::CreateWidgets()
 
 	if (ClothAsset)
 	{
-		Dataflow = ClothEditorToolkitHelpers::GetDataflowFrom(ClothAsset);
+		Dataflow = Private::GetDataflowFrom(ClothAsset);
 
 		// TODO: Figure out how to create the GraphEditor widgets when the ClothAsset doesn't have a Dataflow property set
 		if (Dataflow)
 		{
-			DataflowTerminalPath = ClothEditorToolkitHelpers::GetDataflowTerminalFrom(ClothAsset);
+			DataflowTerminalPath = Private::GetDataflowTerminalFrom(ClothAsset);
 
 			Dataflow->Schema = UDataflowSchema::StaticClass();
 
@@ -757,7 +759,7 @@ void FChaosClothAssetEditorToolkit::OnFinishedChangingAssetProperties(const FPro
 		const UChaosClothAsset* const ClothAsset = GetAsset();
 		if (ClothAsset)
 		{
-			Dataflow = ClothEditorToolkitHelpers::GetDataflowFrom(ClothAsset);
+			Dataflow = Private::GetDataflowFrom(ClothAsset);
 
 			if (Dataflow)
 			{
@@ -846,7 +848,7 @@ void FChaosClothAssetEditorToolkit::ReinitializeGraphEditorWidget()
 
 	if (!GraphEditor)
 	{
-		DataflowTerminalPath = ClothEditorToolkitHelpers::GetDataflowTerminalFrom(ClothAsset);
+		DataflowTerminalPath = Private::GetDataflowTerminalFrom(ClothAsset);
 
 		NodeDetailsEditor = CreateNodeDetailsEditorWidget(ClothAsset);
 		if (NodeDetailsTab.IsValid())
@@ -1047,6 +1049,6 @@ void FChaosClothAssetEditorToolkit::InvalidateViews()
 	ViewportClient->Invalidate();
 	ClothPreviewViewportClient->Invalidate();
 }
-
+} // namespace UE::Chaos::ClothAsset
 
 #undef LOCTEXT_NAMESPACE
diff --git a/Engine/Plugins/Experimental/ChaosClothAssetEditor/Source/ChaosClothAssetEditor/Private/ChaosClothAsset/ClothEditorToolkit.h b/Engine/Plugins/Experimental/ChaosClothAssetEditor/Source/ChaosClothAssetEditor/Private/ChaosClothAsset/ClothEditorToolkit.h
--- a/Engine/Plugins/Experimental/ChaosClothAssetEditor/Source/ChaosClothAssetEditor/Private/ChaosClothAsset/ClothEditorToolkit.h
+++ b/Engine/Plugins/Experimental/ChaosClothAssetEditor/Source/ChaosClothAssetEditor/Private/ChaosClothAsset/ClothEditorToolkit.h
@@ -6,7 +6,6 @@
 #include "Dataflow/DataflowObjectInterface.h"
 #include "ChaosClothAsset/ClothEditorPreviewScene.h"
 
-class FChaosClothAssetEditor3DViewportClient;
 template<typename T> class SComboBox;
 class SClothCollectionOutliner;
 class UDataflow;
@@ -18,7 +17,12 @@ class UEdGraphNode;
 class UChaosClothComponent;
 class SChaosClothAssetEditorRestSpaceViewport;
 class SChaosClothAssetEditor3DViewport;
+
+namespace UE::Chaos::ClothAsset
+{
 class FClothEditorSimulationVisualization;
+class FChaosClothAssetEditor3DViewportClient;
+}
 
 namespace Dataflow
 {
@@ -33,8 +37,10 @@ namespace Dataflow
 	};
 }
 
+namespace UE::Chaos::ClothAsset
+{
 /**
- * The toolkit is supposed to act as the UI manager for the asset editor. It's responsible 
+ * The toolkit is supposed to act as the UI manager for the asset editor. It's responsible
  * for setting up viewports and most toolbars, except for the internals of the mode panel.
  * However, because the toolkit also sets up the mode manager, and much of the important
  * state is held in the UChaosClothAssetEditorMode managed by the mode manager, the toolkit also ends up
@@ -155,4 +161,4 @@ private:
 	TSharedPtr<SDockTab> NodeDetailsTab;
 	TSharedPtr<IStructureDetailsView> NodeDetailsEditor;
 };
-
+} // namespace UE::Chaos::ClothAsset
diff --git a/Engine/Plugins/Experimental/ChaosClothAssetEditor/Source/ChaosClothAssetEditor/Private/ChaosClothAsset/SClothAnimationScrubPanel.cpp b/Engine/Plugins/Experimental/ChaosClothAssetEditor/Source/ChaosClothAssetEditor/Private/ChaosClothAsset/SClothAnimationScrubPanel.cpp
--- a/Engine/Plugins/Experimental/ChaosClothAssetEditor/Source/ChaosClothAssetEditor/Private/ChaosClothAsset/SClothAnimationScrubPanel.cpp
+++ b/Engine/Plugins/Experimental/ChaosClothAssetEditor/Source/ChaosClothAssetEditor/Private/ChaosClothAsset/SClothAnimationScrubPanel.cpp
@@ -10,7 +10,7 @@
 
 #define LOCTEXT_NAMESPACE "ClothAnimationScrubPanel"
 
-void SClothAnimationScrubPanel::Construct( const SClothAnimationScrubPanel::FArguments& InArgs, const TWeakPtr<FChaosClothPreviewScene> InPreviewScene)
+void SClothAnimationScrubPanel::Construct( const SClothAnimationScrubPanel::FArguments& InArgs, const TWeakPtr<UE::Chaos::ClothAsset::FChaosClothPreviewScene> InPreviewScene)
 {
 	PreviewSceneWeakPtr = InPreviewScene;
 
@@ -208,7 +208,7 @@ void SClothAnimationScrubPanel::OnBeginSliderMovement()
 
 uint32 SClothAnimationScrubPanel::GetNumberOfKeys() const
 {
-	if (const TSharedPtr<FChaosClothPreviewScene> PreviewScene = PreviewSceneWeakPtr.Pin())
+	if (const TSharedPtr<UE::Chaos::ClothAsset::FChaosClothPreviewScene> PreviewScene = PreviewSceneWeakPtr.Pin())
 	{
 		if (UAnimSingleNodeInstance* const PreviewInstance = PreviewScene->GetPreviewAnimInstance())	// non-const because UAnimSingleNodeInstance::GetLength() is non-const
 		{
@@ -238,7 +238,7 @@ uint32 SClothAnimationScrubPanel::GetNumberOfKeys() const
 
 float SClothAnimationScrubPanel::GetSequenceLength() const
 {
-	if (const TSharedPtr<FChaosClothPreviewScene> PreviewScene = PreviewSceneWeakPtr.Pin())
+	if (const TSharedPtr<UE::Chaos::ClothAsset::FChaosClothPreviewScene> PreviewScene = PreviewSceneWeakPtr.Pin())
 	{
 		if (UAnimSingleNodeInstance* const PreviewInstance = PreviewScene->GetPreviewAnimInstance())	// non-const because UAnimSingleNodeInstance::GetLength() is non-const
 		{
@@ -272,7 +272,7 @@ bool SClothAnimationScrubPanel::GetDisplayDrag() const
 
 UAnimSingleNodeInstance* SClothAnimationScrubPanel::GetPreviewAnimationInstance()
 {
-	if (const TSharedPtr<FChaosClothPreviewScene> PreviewScene = PreviewSceneWeakPtr.Pin())
+	if (const TSharedPtr<UE::Chaos::ClothAsset::FChaosClothPreviewScene> PreviewScene = PreviewSceneWeakPtr.Pin())
 	{
 		return PreviewScene->GetPreviewAnimInstance();
 	}
@@ -282,7 +282,7 @@ UAnimSingleNodeInstance* SClothAnimationScrubPanel::GetPreviewAnimationInstance(
 
 const UAnimSingleNodeInstance* SClothAnimationScrubPanel::GetPreviewAnimationInstance() const
 {
-	if (const TSharedPtr<const FChaosClothPreviewScene> PreviewScene = PreviewSceneWeakPtr.Pin())
+	if (const TSharedPtr<const UE::Chaos::ClothAsset::FChaosClothPreviewScene> PreviewScene = PreviewSceneWeakPtr.Pin())
 	{
 		return PreviewScene->GetPreviewAnimInstance();
 	}
diff --git a/Engine/Plugins/Experimental/ChaosClothAssetEditor/Source/ChaosClothAssetEditor/Private/ChaosClothAsset/SClothAnimationScrubPanel.h b/Engine/Plugins/Experimental/ChaosClothAssetEditor/Source/ChaosClothAssetEditor/Private/ChaosClothAsset/SClothAnimationScrubPanel.h
--- a/Engine/Plugins/Experimental/ChaosClothAssetEditor/Source/ChaosClothAssetEditor/Private/ChaosClothAsset/SClothAnimationScrubPanel.h
+++ b/Engine/Plugins/Experimental/ChaosClothAssetEditor/Source/ChaosClothAssetEditor/Private/ChaosClothAsset/SClothAnimationScrubPanel.h
@@ -6,8 +6,11 @@
 #include "ITransportControl.h"  // EPlaybackMode::Type
 
 class SScrubControlPanel;
-class FChaosClothPreviewScene;
 class UAnimSingleNodeInstance;
+namespace UE::Chaos::ClothAsset
+{
+class FChaosClothPreviewScene;
+}
 
 ///
 /// A simple panel containing animation controls to operate on a UAnimSingleNodeInstance
@@ -21,7 +24,7 @@ private:
 		SLATE_ATTRIBUTE( float, ViewInputMax )
 	SLATE_END_ARGS()
 
-	void Construct( const FArguments& InArgs, const TWeakPtr<FChaosClothPreviewScene> InPreviewScene );
+	void Construct( const FArguments& InArgs, const TWeakPtr<UE::Chaos::ClothAsset::FChaosClothPreviewScene> InPreviewScene );
 
 	// notifiers 
 	FReply OnClick_Forward_Step();
@@ -47,7 +50,7 @@ private:
 
 	bool GetDisplayDrag() const;
 
-	TWeakPtr<FChaosClothPreviewScene> PreviewSceneWeakPtr;
+	TWeakPtr<UE::Chaos::ClothAsset::FChaosClothPreviewScene> PreviewSceneWeakPtr;
 
 	TSharedPtr<SScrubControlPanel> ScrubControlPanel;
 };
diff --git a/Engine/Plugins/Experimental/ChaosClothAssetEditor/Source/ChaosClothAssetEditor/Private/ChaosClothAsset/SClothEditor3DViewport.cpp b/Engine/Plugins/Experimental/ChaosClothAssetEditor/Source/ChaosClothAssetEditor/Private/ChaosClothAsset/SClothEditor3DViewport.cpp
--- a/Engine/Plugins/Experimental/ChaosClothAssetEditor/Source/ChaosClothAssetEditor/Private/ChaosClothAsset/SClothEditor3DViewport.cpp
+++ b/Engine/Plugins/Experimental/ChaosClothAssetEditor/Source/ChaosClothAssetEditor/Private/ChaosClothAsset/SClothEditor3DViewport.cpp
@@ -45,21 +45,23 @@ void SChaosClothAssetEditor3DViewport::Construct(const FArguments& InArgs, const
 	];
 }
 
-TWeakPtr<FChaosClothPreviewScene> SChaosClothAssetEditor3DViewport::GetPreviewScene()
+TWeakPtr<UE::Chaos::ClothAsset::FChaosClothPreviewScene> SChaosClothAssetEditor3DViewport::GetPreviewScene()
 {
-	const TSharedPtr<FChaosClothAssetEditor3DViewportClient> ClothViewportClient = StaticCastSharedPtr<FChaosClothAssetEditor3DViewportClient>(Client);
+	const TSharedPtr<UE::Chaos::ClothAsset::FChaosClothAssetEditor3DViewportClient> ClothViewportClient = StaticCastSharedPtr<UE::Chaos::ClothAsset::FChaosClothAssetEditor3DViewportClient>(Client);
 	return ClothViewportClient->GetClothPreviewScene();
 }
 
-TWeakPtr<const FChaosClothPreviewScene> SChaosClothAssetEditor3DViewport::GetPreviewScene() const
+TWeakPtr<const UE::Chaos::ClothAsset::FChaosClothPreviewScene> SChaosClothAssetEditor3DViewport::GetPreviewScene() const
 {
-	const TSharedPtr<const FChaosClothAssetEditor3DViewportClient> ClothViewportClient = StaticCastSharedPtr<FChaosClothAssetEditor3DViewportClient>(Client);
+	const TSharedPtr<const UE::Chaos::ClothAsset::FChaosClothAssetEditor3DViewportClient> ClothViewportClient = StaticCastSharedPtr<UE::Chaos::ClothAsset::FChaosClothAssetEditor3DViewportClient>(Client);
 	return ClothViewportClient->GetClothPreviewScene();
 }
 
 
 void SChaosClothAssetEditor3DViewport::BindCommands()
 {
+	using namespace UE::Chaos::ClothAsset;
+
 	SAssetEditorViewport::BindCommands();
 	const FChaosClothAssetEditorCommands& CommandInfos = FChaosClothAssetEditorCommands::Get();
 
@@ -133,7 +135,7 @@ TSharedPtr<SWidget> SChaosClothAssetEditor3DViewport::MakeViewportToolbar()
 
 void SChaosClothAssetEditor3DViewport::OnFocusViewportToSelection()
 {
-	const FBox PreviewBoundingBox = StaticCastSharedPtr<FChaosClothAssetEditor3DViewportClient>(Client)->PreviewBoundingBox();
+	const FBox PreviewBoundingBox = StaticCastSharedPtr<UE::Chaos::ClothAsset::FChaosClothAssetEditor3DViewportClient>(Client)->PreviewBoundingBox();
 	Client->FocusViewportOnBox(PreviewBoundingBox);
 }
 
@@ -155,6 +157,8 @@ float SChaosClothAssetEditor3DViewport::GetViewMinInput() const
 
 float SChaosClothAssetEditor3DViewport::GetViewMaxInput() const
 {
+	using namespace UE::Chaos::ClothAsset;
+
 	// (these are non-const because UAnimSingleNodeInstance::GetLength() is non-const)
 	const TSharedPtr<FChaosClothAssetEditor3DViewportClient> ClothViewportClient = StaticCastSharedPtr<FChaosClothAssetEditor3DViewportClient>(Client);
 	const TSharedPtr<FChaosClothPreviewScene> Scene = ClothViewportClient->GetClothPreviewScene().Pin();
@@ -171,6 +175,8 @@ float SChaosClothAssetEditor3DViewport::GetViewMaxInput() const
 
 EVisibility SChaosClothAssetEditor3DViewport::GetAnimControlVisibility() const
 {
+	using namespace UE::Chaos::ClothAsset;
+
 	const TSharedPtr<const FChaosClothPreviewScene> Scene = GetPreviewScene().Pin();
 	return (Scene && Scene->GetSkeletalMeshComponent() && Scene->GetPreviewAnimInstance()) ? EVisibility::Visible : EVisibility::Hidden;
 }
diff --git a/Engine/Plugins/Experimental/ChaosClothAssetEditor/Source/ChaosClothAssetEditor/Private/ChaosClothAsset/SClothEditor3DViewport.h b/Engine/Plugins/Experimental/ChaosClothAssetEditor/Source/ChaosClothAssetEditor/Private/ChaosClothAsset/SClothEditor3DViewport.h
--- a/Engine/Plugins/Experimental/ChaosClothAssetEditor/Source/ChaosClothAssetEditor/Private/ChaosClothAsset/SClothEditor3DViewport.h
+++ b/Engine/Plugins/Experimental/ChaosClothAssetEditor/Source/ChaosClothAssetEditor/Private/ChaosClothAsset/SClothEditor3DViewport.h
@@ -7,7 +7,10 @@
 #include "SCommonEditorViewportToolbarBase.h"
 #include "SBaseCharacterFXEditorViewport.h"
 
+namespace UE::Chaos::ClothAsset
+{
 class FChaosClothPreviewScene;
+}
 class SClothAnimationScrubPanel;
 
 /**
@@ -38,8 +41,8 @@ public:
 
 private:
 
-	TWeakPtr<FChaosClothPreviewScene> GetPreviewScene();
-	TWeakPtr<const FChaosClothPreviewScene> GetPreviewScene() const;
+	TWeakPtr<UE::Chaos::ClothAsset::FChaosClothPreviewScene> GetPreviewScene();
+	TWeakPtr<const UE::Chaos::ClothAsset::FChaosClothPreviewScene> GetPreviewScene() const;
 
 	float GetViewMinInput() const;
 	float GetViewMaxInput() const;
diff --git a/Engine/Plugins/Experimental/ChaosClothAssetEditor/Source/ChaosClothAssetEditor/Private/ChaosClothAsset/SClothEditor3DViewportToolBar.cpp b/Engine/Plugins/Experimental/ChaosClothAssetEditor/Source/ChaosClothAssetEditor/Private/ChaosClothAsset/SClothEditor3DViewportToolBar.cpp
--- a/Engine/Plugins/Experimental/ChaosClothAssetEditor/Source/ChaosClothAssetEditor/Private/ChaosClothAsset/SClothEditor3DViewportToolBar.cpp
+++ b/Engine/Plugins/Experimental/ChaosClothAssetEditor/Source/ChaosClothAssetEditor/Private/ChaosClothAsset/SClothEditor3DViewportToolBar.cpp
@@ -79,6 +79,8 @@ TSharedRef<SWidget> SChaosClothAssetEditor3DViewportToolBar::MakeDisplayToolBar(
 
 TSharedRef<SWidget> SChaosClothAssetEditor3DViewportToolBar::MakeToolBar(const TSharedPtr<FExtender> InExtenders)
 {
+	using namespace UE::Chaos::ClothAsset;
+
 	// The following is modeled after portions of STransformViewportToolBar, which gets 
 	// used in SCommonEditorViewportToolbarBase.
 
@@ -145,6 +147,8 @@ TSharedRef<SWidget> SChaosClothAssetEditor3DViewportToolBar::MakeToolBar(const T
 
 void SChaosClothAssetEditor3DViewportToolBar::ExtendOptionsMenu(FMenuBuilder& OptionsMenuBuilder) const
 {
+	using namespace UE::Chaos::ClothAsset;
+
 	constexpr bool bOpenSubMenuOnClick = false;
 	constexpr bool bShouldCloseWindowAfterMenuSelection = false;
 	OptionsMenuBuilder.AddSubMenu(
diff --git a/Engine/Plugins/Experimental/ChaosClothAssetEditor/Source/ChaosClothAssetEditor/Private/ChaosClothAsset/SClothEditorRestSpaceViewport.cpp b/Engine/Plugins/Experimental/ChaosClothAssetEditor/Source/ChaosClothAssetEditor/Private/ChaosClothAsset/SClothEditorRestSpaceViewport.cpp
--- a/Engine/Plugins/Experimental/ChaosClothAssetEditor/Source/ChaosClothAssetEditor/Private/ChaosClothAsset/SClothEditorRestSpaceViewport.cpp
+++ b/Engine/Plugins/Experimental/ChaosClothAssetEditor/Source/ChaosClothAssetEditor/Private/ChaosClothAsset/SClothEditorRestSpaceViewport.cpp
@@ -11,6 +11,8 @@
 
 void SChaosClothAssetEditorRestSpaceViewport::BindCommands()
 {
+	using namespace UE::Chaos::ClothAsset;
+
 	SAssetEditorViewport::BindCommands();
 
 	const FChaosClothAssetEditorCommands& CommandInfos = FChaosClothAssetEditorCommands::Get();
diff --git a/Engine/Plugins/Experimental/ChaosClothAssetEditor/Source/ChaosClothAssetEditor/Private/ChaosClothAsset/SClothEditorRestSpaceViewportToolBar.cpp b/Engine/Plugins/Experimental/ChaosClothAssetEditor/Source/ChaosClothAssetEditor/Private/ChaosClothAsset/SClothEditorRestSpaceViewportToolBar.cpp
--- a/Engine/Plugins/Experimental/ChaosClothAssetEditor/Source/ChaosClothAssetEditor/Private/ChaosClothAsset/SClothEditorRestSpaceViewportToolBar.cpp
+++ b/Engine/Plugins/Experimental/ChaosClothAssetEditor/Source/ChaosClothAssetEditor/Private/ChaosClothAsset/SClothEditorRestSpaceViewportToolBar.cpp
@@ -75,6 +75,8 @@ TSharedRef<SWidget> SChaosClothAssetEditorRestSpaceViewportToolBar::MakeDisplayT
 
 TSharedRef<SWidget> SChaosClothAssetEditorRestSpaceViewportToolBar::MakeToolBar(const TSharedPtr<FExtender> InExtenders)
 {
+	using namespace UE::Chaos::ClothAsset;
+
 	// The following is modeled after portions of STransformViewportToolBar, which gets 
 	// used in SCommonEditorViewportToolbarBase.
 
