Chaos Cloth Editor - Fixed crash/fatal log caused by failling CastChecked.

#rnx

diff --git a/Engine/Plugins/Experimental/ChaosClothAssetEditor/Source/ChaosClothAssetEditor/Private/ChaosClothAsset/ClothEditorToolkit.cpp b/Engine/Plugins/Experimental/ChaosClothAssetEditor/Source/ChaosClothAssetEditor/Private/ChaosClothAsset/ClothEditorToolkit.cpp
--- a/Engine/Plugins/Experimental/ChaosClothAssetEditor/Source/ChaosClothAssetEditor/Private/ChaosClothAsset/ClothEditorToolkit.cpp
+++ b/Engine/Plugins/Experimental/ChaosClothAssetEditor/Source/ChaosClothAssetEditor/Private/ChaosClothAsset/ClothEditorToolkit.cpp
@@ -976,7 +976,7 @@ void FChaosClothAssetEditorToolkit::OnNodeSelectionChanged(const TSet<UObject*>&
 		Dataflow->LastModifiedRenderTarget = Dataflow::FTimestamp::Current();
 	}
 
-	UChaosClothAssetEditorMode* const ClothMode = CastChecked<UChaosClothAssetEditorMode>(EditorModeManager->GetActiveScriptableMode(UChaosClothAssetEditorMode::EM_ChaosClothAssetEditorModeId));
+	UChaosClothAssetEditorMode* const ClothMode = Cast<UChaosClothAssetEditorMode>(EditorModeManager->GetActiveScriptableMode(UChaosClothAssetEditorMode::EM_ChaosClothAssetEditorModeId));
 	if (ClothMode)
 	{
 		ClothMode->SetSelectedClothCollection(Collection);
