Cloth Weight Map Paint Tool: change 0-1 color from white-green to black-white

diff --git a/Engine/Plugins/Experimental/ChaosClothAssetEditor/Source/ChaosClothAssetEditorTools/Private/ChaosClothAsset/ClothWeightMapPaintTool.cpp b/Engine/Plugins/Experimental/ChaosClothAssetEditor/Source/ChaosClothAssetEditorTools/Private/ChaosClothAsset/ClothWeightMapPaintTool.cpp
--- a/Engine/Plugins/Experimental/ChaosClothAssetEditor/Source/ChaosClothAssetEditorTools/Private/ChaosClothAsset/ClothWeightMapPaintTool.cpp
+++ b/Engine/Plugins/Experimental/ChaosClothAssetEditor/Source/ChaosClothAssetEditorTools/Private/ChaosClothAsset/ClothWeightMapPaintTool.cpp
@@ -1401,8 +1401,8 @@ bool UClothEditorWeightMapPaintTool::CanAccept() const
 
 FColor UClothEditorWeightMapPaintTool::GetColorForWeightValue(double WeightValue)
 {
-	FColor MaxColor = LinearColors::ForestGreen3b();
-	FColor MinColor = LinearColors::White3b();
+	FColor MaxColor = LinearColors::White3b();
+	FColor MinColor = LinearColors::Black3b();
 	FColor Color;
 	double ClampedValue = FMath::Clamp(WeightValue, 0.0, 1.0);
 	Color.R = FMath::LerpStable(MinColor.R, MaxColor.R, ClampedValue);
