ClothTrainingTool: get render positions from proxy

diff --git a/Engine/Plugins/Experimental/ChaosClothAsset/Source/ChaosClothAssetEngine/Private/ChaosClothAsset/ClothComponent.cpp b/Engine/Plugins/Experimental/ChaosClothAsset/Source/ChaosClothAssetEngine/Private/ChaosClothAsset/ClothComponent.cpp
--- a/Engine/Plugins/Experimental/ChaosClothAsset/Source/ChaosClothAssetEngine/Private/ChaosClothAsset/ClothComponent.cpp
+++ b/Engine/Plugins/Experimental/ChaosClothAsset/Source/ChaosClothAssetEngine/Private/ChaosClothAsset/ClothComponent.cpp
@@ -116,10 +116,10 @@ void UChaosClothComponent::OnRegister()
 	}
 }
 
-TUniquePtr<UE::Chaos::ClothAsset::FClothSimulationProxy> UChaosClothComponent::CreateClothSimulationProxy()
+TSharedPtr<UE::Chaos::ClothAsset::FClothSimulationProxy> UChaosClothComponent::CreateClothSimulationProxy()
 {
 	using namespace UE::Chaos::ClothAsset;
-	return MakeUnique<FClothSimulationProxy>(*this);
+	return MakeShared<FClothSimulationProxy>(*this);
 }
 
 bool UChaosClothComponent::IsSimulationSuspended() const
@@ -163,7 +163,7 @@ void UChaosClothComponent::OnUnregister()
 	Super::OnUnregister();
 
 	// Release cloth simulation
-	ClothSimulationProxy.Reset(nullptr);
+	ClothSimulationProxy.Reset();
 }
 
 bool UChaosClothComponent::IsComponentTickEnabled() const
@@ -387,20 +387,3 @@ bool UChaosClothComponent::ShouldWaitForParallelSimulationInTickComponent() cons
 
 	return bWaitForParallelTask || (CVarClothPhysicsWaitForParallelClothTask && CVarClothPhysicsWaitForParallelClothTask->GetBool());
 }
-
-void UChaosClothComponent::Pose(const TArray<FTransform>& InComponentSpaceTransforms)
-{
-	const bool bIsInputValid = InComponentSpaceTransforms.Num() == GetComponentSpaceTransforms().Num();
-	ensure(bIsInputValid);
-	if (!bIsInputValid)
-	{
-		return;
-	}
-	GetEditableComponentSpaceTransforms() = InComponentSpaceTransforms;
-	bNeedToFlipSpaceBaseBuffers = true;
-	FinalizeBoneTransform();
-
-	UpdateBounds();
-	MarkRenderTransformDirty();
-	MarkRenderDynamicDataDirty();
-}
\ No newline at end of file
diff --git a/Engine/Plugins/Experimental/ChaosClothAsset/Source/ChaosClothAssetEngine/Public/ChaosClothAsset/ClothComponent.h b/Engine/Plugins/Experimental/ChaosClothAsset/Source/ChaosClothAssetEngine/Public/ChaosClothAsset/ClothComponent.h
--- a/Engine/Plugins/Experimental/ChaosClothAsset/Source/ChaosClothAssetEngine/Public/ChaosClothAsset/ClothComponent.h
+++ b/Engine/Plugins/Experimental/ChaosClothAsset/Source/ChaosClothAssetEngine/Public/ChaosClothAsset/ClothComponent.h
@@ -95,9 +95,6 @@ public:
 
 	const UE::Chaos::ClothAsset::FClothSimulationProxy* GetClothSimulationProxy() const { return ClothSimulationProxy.Get(); }
 
-	/** Pose the cloth component using component space transforms. */
-	void Pose(const TArray<FTransform>& InComponentSpaceTransforms);
-
 	friend UE::Chaos::ClothAsset::FClothComponentCacheAdapter;
 
 protected:
@@ -128,7 +125,7 @@ protected:
 	//~ End USkinnedMeshComponent Interface
 
 	/** Override this function for setting up custom simulation proxies when the component is registered. */
-	virtual TUniquePtr<UE::Chaos::ClothAsset::FClothSimulationProxy> CreateClothSimulationProxy();
+	virtual TSharedPtr<UE::Chaos::ClothAsset::FClothSimulationProxy> CreateClothSimulationProxy();
 
 private:
 	void StartNewParallelSimulation(float DeltaTime);
@@ -177,5 +174,5 @@ private:
 	TSharedPtr<FManagedArrayCollection> PropertyCollection;
 	TUniquePtr<::Chaos::Softs::FCollectionPropertyFacade> CollectionPropertyFacade;
 
-	TUniquePtr<UE::Chaos::ClothAsset::FClothSimulationProxy> ClothSimulationProxy;
+	TSharedPtr<UE::Chaos::ClothAsset::FClothSimulationProxy> ClothSimulationProxy;
 };
diff --git a/Engine/Plugins/Experimental/ChaosClothAssetEditor/Source/ChaosClothAssetEditorTools/ChaosClothAssetEditorTools.Build.cs b/Engine/Plugins/Experimental/ChaosClothAssetEditor/Source/ChaosClothAssetEditorTools/ChaosClothAssetEditorTools.Build.cs
--- a/Engine/Plugins/Experimental/ChaosClothAssetEditor/Source/ChaosClothAssetEditorTools/ChaosClothAssetEditorTools.Build.cs
+++ b/Engine/Plugins/Experimental/ChaosClothAssetEditor/Source/ChaosClothAssetEditorTools/ChaosClothAssetEditorTools.Build.cs
@@ -23,6 +23,7 @@ public class ChaosClothAssetEditorTools : ModuleRules
 			{
 				"CoreUObject",
 				"Engine",
+				"RenderCore",
 				"Slate",
 				"SlateCore",
 				"Projects",
diff --git a/Engine/Plugins/Experimental/ChaosClothAssetEditor/Source/ChaosClothAssetEditorTools/Private/ChaosClothAsset/ClothDataGenerationComponent.cpp b/Engine/Plugins/Experimental/ChaosClothAssetEditor/Source/ChaosClothAssetEditorTools/Private/ChaosClothAsset/ClothDataGenerationComponent.cpp
new file mode 100644
--- /dev/null
+++ b/Engine/Plugins/Experimental/ChaosClothAssetEditor/Source/ChaosClothAssetEditorTools/Private/ChaosClothAsset/ClothDataGenerationComponent.cpp
@@ -0,0 +1,59 @@
+#include "ChaosClothAsset/ClothDataGenerationComponent.h"
+
+namespace UE::Chaos::ClothAsset
+{
+    FClothSimulationDataGenerationProxy::FClothSimulationDataGenerationProxy(const UChaosClothComponent& InClothComponent)
+	    : FClothSimulationProxy(InClothComponent)
+    {	
+    }
+
+    FClothSimulationDataGenerationProxy::~FClothSimulationDataGenerationProxy() = default;
+};
+
+UClothDataGenerationComponent::UClothDataGenerationComponent(const FObjectInitializer& ObjectInitializer)
+    : Super(ObjectInitializer)
+{
+}
+
+UClothDataGenerationComponent::UClothDataGenerationComponent(FVTableHelper& Helper)
+    : Super(Helper)
+{
+}
+
+UClothDataGenerationComponent::~UClothDataGenerationComponent() = default;
+
+TSharedPtr<UE::Chaos::ClothAsset::FClothSimulationProxy> UClothDataGenerationComponent::CreateClothSimulationProxy()
+{
+    TSharedPtr<FDataProxy> DataProxyShared = MakeShared<FDataProxy>(*this);
+    DataProxy = DataProxyShared;
+    return DataProxyShared;
+}
+
+void UClothDataGenerationComponent::TickComponent(float DeltaTime, enum ELevelTick TickType, FActorComponentTickFunction* ThisTickFunction)
+{
+    UActorComponent::TickComponent(DeltaTime, TickType, ThisTickFunction);
+}
+
+
+TWeakPtr<UE::Chaos::ClothAsset::FClothSimulationDataGenerationProxy> UClothDataGenerationComponent::GetProxy() const
+{
+    return DataProxy;
+}
+
+void UClothDataGenerationComponent::Pose(const TArray<FTransform>& InComponentSpaceTransforms)
+{
+	if (!ensure(InComponentSpaceTransforms.Num() == GetComponentSpaceTransforms().Num()))
+	{
+		return;
+	}
+	GetEditableComponentSpaceTransforms() = InComponentSpaceTransforms;
+	bNeedToFlipSpaceBaseBuffers = true;
+	FinalizeBoneTransform();
+
+	UpdateBounds();
+    if (IsInGameThread())
+    {
+        MarkRenderTransformDirty();
+        MarkRenderDynamicDataDirty();
+    }
+}
diff --git a/Engine/Plugins/Experimental/ChaosClothAssetEditor/Source/ChaosClothAssetEditorTools/Private/ChaosClothAsset/ClothDataGenerationComponent.h b/Engine/Plugins/Experimental/ChaosClothAssetEditor/Source/ChaosClothAssetEditorTools/Private/ChaosClothAsset/ClothDataGenerationComponent.h
new file mode 100644
--- /dev/null
+++ b/Engine/Plugins/Experimental/ChaosClothAssetEditor/Source/ChaosClothAssetEditorTools/Private/ChaosClothAsset/ClothDataGenerationComponent.h
@@ -0,0 +1,52 @@
+// Copyright Epic Games, Inc. All Rights Reserved.
+
+#pragma once
+
+#include "ChaosClothAsset/ClothComponent.h"
+#include "ChaosClothAsset/ClothSimulationProxy.h"
+
+#include "ClothDataGenerationComponent.generated.h"
+
+namespace UE::Chaos::ClothAsset
+{
+	class FClothSimulationDataGenerationProxy : public FClothSimulationProxy
+    {
+    public:
+        explicit FClothSimulationDataGenerationProxy(const UChaosClothComponent& InClothComponent);
+        ~FClothSimulationDataGenerationProxy();
+
+        using FClothSimulationProxy::Tick;
+        using FClothSimulationProxy::FillSimulationContext;
+        using FClothSimulationProxy::InitializeConfigs;
+        using FClothSimulationProxy::WriteSimulationData;
+    };
+};
+
+/**
+ * Cloth data generation component.
+ */
+UCLASS()
+class UClothDataGenerationComponent : public UChaosClothComponent
+{
+	GENERATED_BODY()	
+public:
+	UClothDataGenerationComponent(const FObjectInitializer& ObjectInitializer);
+	UClothDataGenerationComponent(FVTableHelper& Helper);
+	~UClothDataGenerationComponent();
+
+    using FDataProxy = UE::Chaos::ClothAsset::FClothSimulationDataGenerationProxy;
+    TWeakPtr<FDataProxy> GetProxy() const;
+    
+	/** Pose the cloth component using component space transforms. */
+    void Pose(const TArray<FTransform>& InComponentSpaceTransforms);
+protected:
+	//~ Begin UActorComponent Interface
+	virtual void TickComponent(float DeltaTime, enum ELevelTick TickType, FActorComponentTickFunction* ThisTickFunction) override;
+	//~ End UActorComponent Interface
+
+    /** Begin UChaosClothComponent Interface */
+	virtual TSharedPtr<UE::Chaos::ClothAsset::FClothSimulationProxy> CreateClothSimulationProxy() override;
+    /** End UChaosClothComponent Interface */
+private:
+    TWeakPtr<FDataProxy> DataProxy;
+};
diff --git a/Engine/Plugins/Experimental/ChaosClothAssetEditor/Source/ChaosClothAssetEditorTools/Private/ChaosClothAsset/ClothTrainingTool.cpp b/Engine/Plugins/Experimental/ChaosClothAssetEditor/Source/ChaosClothAssetEditorTools/Private/ChaosClothAsset/ClothTrainingTool.cpp
--- a/Engine/Plugins/Experimental/ChaosClothAssetEditor/Source/ChaosClothAssetEditorTools/Private/ChaosClothAsset/ClothTrainingTool.cpp
+++ b/Engine/Plugins/Experimental/ChaosClothAssetEditor/Source/ChaosClothAssetEditorTools/Private/ChaosClothAsset/ClothTrainingTool.cpp
@@ -9,6 +9,7 @@
 #include "Chaos/CacheCollection.h"
 #include "ChaosCloth/ChaosClothingSimulationConfig.h"
 #include "ChaosClothAsset/ClothAsset.h"
+#include "ChaosClothAsset/ClothDataGenerationComponent.h"
 #include "ChaosClothAsset/ClothComponent.h"
 #include "ChaosClothAsset/ClothComponentToolTarget.h"
 #include "ChaosClothAsset/ClothCollection.h"
@@ -22,34 +23,16 @@
 #include "ModelingOperators.h"
 #include "Misc/AsyncTaskNotification.h"
 #include "PhysicsEngine/PhysicsAsset.h"
+#include "SkeletalRenderPublic.h"
 #include "Tasks/Pipe.h"
 #include "ToolTargetManager.h"
 #include "UObject/SavePackage.h"
 
+#include <atomic>
+
 #define LOCTEXT_NAMESPACE "ClothTrainingTool"
 DEFINE_LOG_CATEGORY_STATIC(LogClothTrainingTool, Log, All);
 
-class UClothTrainingTool::FClothSimulationDataGenerationProxy : public UE::Chaos::ClothAsset::FClothSimulationProxy
-{
-public:
-	explicit FClothSimulationDataGenerationProxy(const UChaosClothComponent& InClothComponent);
-	~FClothSimulationDataGenerationProxy();
-
-	using UE::Chaos::ClothAsset::FClothSimulationProxy::Tick;
-	using UE::Chaos::ClothAsset::FClothSimulationProxy::FillSimulationContext;
-	using UE::Chaos::ClothAsset::FClothSimulationProxy::InitializeConfigs;
-	using UE::Chaos::ClothAsset::FClothSimulationProxy::WriteSimulationData;
-};
-
-UClothTrainingTool::FClothSimulationDataGenerationProxy::FClothSimulationDataGenerationProxy(const UChaosClothComponent& InClothComponent)
-	: FClothSimulationProxy(InClothComponent)
-{	
-}
-
-UClothTrainingTool::FClothSimulationDataGenerationProxy::~FClothSimulationDataGenerationProxy()
-{
-}
-
 namespace UE::ClothTrainingTool::Private
 {
 	UChaosCache& GetCache(UChaosCacheCollection& CacheCollection)
@@ -142,9 +125,11 @@ namespace UE::ClothTrainingTool::Private
 
 struct UClothTrainingTool::FSimResource
 {
-	UChaosClothComponent* ClothComponent = nullptr;
-	TUniquePtr<FClothSimulationDataGenerationProxy> Proxy = nullptr;
-	TUniquePtr<UE::Tasks::FPipe> Pipe = nullptr;
+	UClothDataGenerationComponent* ClothComponent = nullptr;
+	TSharedPtr<FProxy> Proxy;
+	TUniquePtr<UE::Tasks::FPipe> Pipe;
+	FEvent* SkinEvent = nullptr;
+	std::atomic<bool> bNeedsSkin = false;
 };
 
 class UClothTrainingTool::FLaunchSimsOp : public UE::Geometry::TGenericDataOperator<FSkinnedMeshVertices>
@@ -161,7 +146,6 @@ public:
 
 private:
 	using FPipe = UE::Tasks::FPipe;
-	using FProxy = UClothTrainingTool::FClothSimulationDataGenerationProxy;
 
 	enum class ESaveType
 	{
@@ -174,7 +158,8 @@ private:
 	void RestoreAnimationSequence();
 	TArray<FTransform> GetBoneTransforms(UChaosClothComponent& InClothComponent, int32 Frame) const;
 	bool GetSimPositions(FProxy& DataGenerationProxy, TArray<FVector3f>& OutPositions) const;
-	void AddToCache(FProxy& DataGenerationProxy, UChaosCache& Cache, int32 Frame) const;
+	void GetRenderPositions(FSimResource& SimResource, TArray<FVector3f>& OutPositions) const;
+	void AddToCache(FSimResource& SimResource, UChaosCache& Cache, int32 Frame) const;
 
 	TArray<FSimResource>& SimResources;
 	FCriticalSection& SimMutex;
@@ -200,6 +185,7 @@ struct UClothTrainingTool::FTaskResource
 
 	bool AllocateSimResources_GameThread(const UChaosClothComponent& InClothComponent, int32 Num);
 	void FreeSimResources_GameThread();
+	void FlushRendering();
 };
 
 void UClothTrainingTool::FLaunchSimsOp::PrepareAnimationSequence()
@@ -267,20 +253,17 @@ TArray<FTransform> UClothTrainingTool::FLaunchSimsOp::GetBoneTransforms(UChaosCl
 	return ComponentSpaceTransforms;
 }
 
-void UClothTrainingTool::FLaunchSimsOp::AddToCache(FProxy& DataGenerationProxy, UChaosCache& OutCache, int32 Frame) const
+void UClothTrainingTool::FLaunchSimsOp::AddToCache(FSimResource& SimResource, UChaosCache& OutCache, int32 Frame) const
 {
-	TArray<FVector3f> SimPositions;
-	if (!GetSimPositions(DataGenerationProxy, SimPositions))
-	{
-		return;
-	}
+	TArray<FVector3f> Positions;
+	GetRenderPositions(SimResource, Positions);
 
 	constexpr float CacheFPS = 30;
 	const float Time = Frame / CacheFPS;
 	FPendingFrameWrite NewFrame;
 	NewFrame.Time = Time;
 
-	const int32 NumParticles = SimPositions.Num();
+	const int32 NumParticles = Positions.Num();
 	TArray<int32>& PendingID = NewFrame.PendingChannelsIndices;
 	TArray<float> PendingPX, PendingPY, PendingPZ;
 	TArray<float> PendingVX, PendingVY, PendingVZ;
@@ -291,7 +274,7 @@ void UClothTrainingTool::FLaunchSimsOp::AddToCache(FProxy& DataGenerationProxy,
 
 	for (int32 ParticleIndex = 0; ParticleIndex < NumParticles; ++ParticleIndex)
 	{
-		const FVector3f& Position = SimPositions[ParticleIndex];
+		const FVector3f& Position = Positions[ParticleIndex];
 		PendingID[ParticleIndex] = ParticleIndex;
 		PendingPX[ParticleIndex] = Position.X;
 		PendingPY[ParticleIndex] = Position.Y;
@@ -362,13 +345,15 @@ void UClothTrainingTool::FLaunchSimsOp::CalculateResult(FProgressCancel* Progres
 
 void UClothTrainingTool::FLaunchSimsOp::Simulate(FSimResource &SimResource, int32 AnimFrame, int32 CacheFrame, UChaosCache& Cache, FProgressCancel* Progress, float ProgressStep)
 {
-	UChaosClothComponent& TaskComponent = *SimResource.ClothComponent;
-	FProxy& DataGenerationProxy = *SimResource.Proxy.Get();
+	UClothDataGenerationComponent& TaskComponent = *SimResource.ClothComponent;
+	FProxy& DataGenerationProxy = *SimResource.Proxy;
 
 	const float TimeStep = ToolProperties->TimeStep;
 	const int32 NumSteps = ToolProperties->NumSteps;
 	const ESaveType SaveType = ToolProperties->bDebug ? ESaveType::EveryStep : ESaveType::LastStep;
 
+	const TArray<FTransform> Transforms = GetBoneTransforms(TaskComponent, AnimFrame);
+	TaskComponent.Pose(Transforms);
 	TaskComponent.ForceNextUpdateTeleportAndReset();
 	DataGenerationProxy.FillSimulationContext(TimeStep);
 	DataGenerationProxy.InitializeConfigs();
@@ -395,7 +380,7 @@ void UClothTrainingTool::FLaunchSimsOp::Simulate(FSimResource &SimResource, int3
 			if (SaveType == ESaveType::EveryStep)
 			{
 				DataGenerationProxy.WriteSimulationData();
-				AddToCache(DataGenerationProxy, Cache, Step);
+				AddToCache(SimResource, Cache, Step);
 			}
 		}
 	}
@@ -403,7 +388,7 @@ void UClothTrainingTool::FLaunchSimsOp::Simulate(FSimResource &SimResource, int3
 	if (SaveType == ESaveType::LastStep && !bCancelled)
 	{
 		DataGenerationProxy.WriteSimulationData();
-		AddToCache(DataGenerationProxy, Cache, CacheFrame);
+		AddToCache(SimResource, Cache, CacheFrame);
 	}
 
 	Progress->AdvanceCurrentScopeProgressBy(ProgressStep);
@@ -433,6 +418,23 @@ bool UClothTrainingTool::FLaunchSimsOp::GetSimPositions(FProxy& DataGenerationPr
 	return true;
 }
 
+void UClothTrainingTool::FLaunchSimsOp::GetRenderPositions(FSimResource& SimResource, TArray<FVector3f> &OutPositions) const
+{
+	check(SimResource.ClothComponent);
+	TArray<FFinalSkinVertex> OutVertices;
+	// This could potentially be slow. 
+	SimResource.ClothComponent->RecreateRenderState_Concurrent();
+	SimResource.bNeedsSkin.store(true);
+	SimResource.SkinEvent->Wait();
+	
+	SimResource.ClothComponent->GetCPUSkinnedCachedFinalVertices(OutVertices);
+	OutPositions.SetNum(OutVertices.Num());
+	for (int32 Index = 0; Index < OutVertices.Num(); ++Index)
+	{
+		OutPositions[Index] = OutVertices[Index].Position;
+	}
+}
+
 bool UClothTrainingTool::IsClothComponentValid() const
 {
 	if (ClothComponent == nullptr)
@@ -582,6 +584,7 @@ void UClothTrainingTool::TickTraining()
 
 	if (!bFinished)
 	{
+		TaskResource->FlushRendering();
 		const FDateTime CurrentTime = FDateTime::UtcNow();
 		const double SinceLastUpdate = (CurrentTime - TaskResource->LastUpdateTime).GetTotalSeconds();
 		if (SinceLastUpdate < 0.2)
@@ -618,7 +621,6 @@ void UClothTrainingTool::OnTick(float DeltaTime)
 	}
 }
 
-
 void UClothTrainingTool::RequestAction(EClothTrainingToolActions ActionType)
 {
 	if (PendingAction != EClothTrainingToolActions::NoAction)
@@ -645,23 +647,37 @@ void UClothTrainingTool::Shutdown(EToolShutdownType ShutdownType)
 
 bool UClothTrainingTool::FTaskResource::AllocateSimResources_GameThread(const UChaosClothComponent& InClothComponent, int32 Num)
 {
-	SimResources.Reserve(Num);
+	SimResources.SetNumUninitialized(Num);
 	for(int32 Index = 0; Index < Num; ++Index)
 	{
-		UChaosClothComponent* const CopyComponent = DuplicateObject<UChaosClothComponent>(&InClothComponent, InClothComponent.GetOuter());
+		UClothDataGenerationComponent* const CopyComponent = NewObject<UClothDataGenerationComponent>(InClothComponent.GetOuter());
+		CopyComponent->SetClothAsset(InClothComponent.GetClothAsset());
 		CopyComponent->RegisterComponent();
 		CopyComponent->SetWorldTransform(InClothComponent.GetComponentTransform());
-		FSimResource SimResource;
+
+		USkinnedMeshComponent* const PoseComponent = CopyComponent->LeaderPoseComponent.Get() ? CopyComponent->LeaderPoseComponent.Get() : CopyComponent;
+		constexpr int32 LODIndex = 0;
+		PoseComponent->SetForcedLOD(LODIndex + 1);
+		PoseComponent->UpdateLODStatus();
+		PoseComponent->RefreshBoneTransforms(nullptr);
+		CopyComponent->bRenderStatic = false;
+		constexpr bool bRecreateRenderStateImmediately = true;
+		CopyComponent->SetCPUSkinningEnabled(true, bRecreateRenderStateImmediately);
+		CopyComponent->ResumeSimulation();
+
+		FSimResource& SimResource = SimResources[Index];
 		SimResource.ClothComponent = CopyComponent;
-		SimResource.Proxy = MakeUnique<FClothSimulationDataGenerationProxy>(*CopyComponent);
+		SimResource.Proxy = CopyComponent->GetProxy().Pin();
+		check(SimResource.Proxy != nullptr);
 		SimResource.Pipe = MakeUnique<UE::Tasks::FPipe>(*FString::Printf(TEXT("SimPipe:%d"), Index));
+		SimResource.SkinEvent = FPlatformProcess::GetSynchEventFromPool();
+		SimResource.bNeedsSkin.store(false);
 
 		if (CopyComponent == nullptr || SimResource.Proxy == nullptr || SimResource.Pipe == nullptr)
 		{
 			UE_LOG(LogClothTrainingTool, Error, TEXT("Failed to allocate simulation resources"));
 			return false;
 		}
-		SimResources.Add(MoveTemp(SimResource));
 	}
 	SimMutex = MakeUnique<FCriticalSection>();
 	return true;
@@ -669,18 +685,49 @@ bool UClothTrainingTool::FTaskResource::AllocateSimResources_GameThread(const UC
 
 void UClothTrainingTool::FTaskResource::FreeSimResources_GameThread()
 {
+	while (!SimMutex->TryLock())
+	{
+		FlushRendering();
+		FPlatformProcess::Sleep(0.1f);
+	}
+	for (FSimResource& SimResource : SimResources)
+	{
+		FPlatformProcess::ReturnSynchEventToPool(SimResource.SkinEvent);
+		SimResource.Pipe.Reset();
+		SimResource.ClothComponent->UnregisterComponent();
+		SimResource.ClothComponent->DestroyComponent();
+	}
+	SimResources.Empty();
+	SimMutex->Unlock();
+	SimMutex.Reset();
+}
+
+void UClothTrainingTool::FTaskResource::FlushRendering()
+{
+	// Copy bNeedsSkin
+	TArray<bool> NeedsSkin;
+	NeedsSkin.SetNum(SimResources.Num());
+	bool bAnyNeedsSkin = false;
+	for (int32 Index = 0; Index < SimResources.Num(); ++Index)
+	{
+		const bool bNeedsSkin = SimResources[Index].bNeedsSkin.load();
+		bAnyNeedsSkin |= bNeedsSkin;
+		NeedsSkin[Index] = bNeedsSkin;
+	}
+
+	if (bAnyNeedsSkin)
 	{
-		FScopeLock Lock(SimMutex.Get());
-		for (FSimResource& SimResource : SimResources)
+		FlushRenderingCommands();
+		for (int32 Index = 0; Index < SimResources.Num(); ++Index)
 		{
-			SimResource.Pipe.Reset();
-			SimResource.Proxy.Reset();
-			SimResource.ClothComponent->UnregisterComponent();
-			SimResource.ClothComponent->DestroyComponent();
+			if (NeedsSkin[Index])
+			{
+				SimResources[Index].bNeedsSkin.store(false);
+				SimResources[Index].SkinEvent->Trigger();
+			}
 		}
-		SimResources.Empty();
 	}
-	SimMutex.Reset();
+
 }
 
 UChaosCacheCollection* UClothTrainingTool::GetCacheCollection() const
diff --git a/Engine/Plugins/Experimental/ChaosClothAssetEditor/Source/ChaosClothAssetEditorTools/Public/ChaosClothAsset/ClothTrainingTool.h b/Engine/Plugins/Experimental/ChaosClothAssetEditor/Source/ChaosClothAssetEditorTools/Public/ChaosClothAsset/ClothTrainingTool.h
--- a/Engine/Plugins/Experimental/ChaosClothAssetEditor/Source/ChaosClothAssetEditorTools/Public/ChaosClothAsset/ClothTrainingTool.h
+++ b/Engine/Plugins/Experimental/ChaosClothAssetEditor/Source/ChaosClothAssetEditorTools/Public/ChaosClothAsset/ClothTrainingTool.h
@@ -19,6 +19,10 @@ class UChaosCacheCollection;
 class UChaosClothComponent;
 class UClothTrainingTool;
 
+namespace UE::Chaos::ClothAsset
+{
+	class FClothSimulationDataGenerationProxy;
+};
 
 // @@@@@@@@@ TODO: Change this to whatever makes sense for output
 struct FSkinnedMeshVertices
@@ -128,6 +132,7 @@ private:
 
 	using FTaskType = UE::Geometry::TModelingOpTask<FLaunchSimsOp>;
 	using FExecuterType = UE::Geometry::FAsyncTaskExecuterWithProgressCancel<FTaskType>;
+	using FProxy = UE::Chaos::ClothAsset::FClothSimulationDataGenerationProxy;
 
 	friend class UClothTrainingToolActionProperties;
 
