Cloth Editor: Handle rare case where you change dataflow node selection but don't have a valid context.

This seems to occur if you happen to change a value and select a new node in between dataflow graph ticks. (change value resets the context and ticking the graph will create a new one, so you have to sneak in a selection change between those two calls)

diff --git a/Engine/Plugins/Experimental/ChaosClothAssetEditor/Source/ChaosClothAssetEditor/Private/ChaosClothAsset/ClothEditorToolkit.cpp b/Engine/Plugins/Experimental/ChaosClothAssetEditor/Source/ChaosClothAssetEditor/Private/ChaosClothAsset/ClothEditorToolkit.cpp
--- a/Engine/Plugins/Experimental/ChaosClothAssetEditor/Source/ChaosClothAssetEditor/Private/ChaosClothAsset/ClothEditorToolkit.cpp
+++ b/Engine/Plugins/Experimental/ChaosClothAssetEditor/Source/ChaosClothAssetEditor/Private/ChaosClothAsset/ClothEditorToolkit.cpp
@@ -928,21 +928,24 @@ void FChaosClothAssetEditorToolkit::OnNodeSelectionChanged(const TSet<UObject*>&
 {
 	auto GetClothCollectionIfPossible = [](const TSharedPtr<const FDataflowNode> DataflowNode, const TSharedPtr<Dataflow::FEngineContext> Context) -> TSharedPtr<FManagedArrayCollection>
 	{
-		for (const FDataflowOutput* const Output : DataflowNode->GetOutputs())
+		if (Context.IsValid())
 		{
-			if (Output->GetType() == FName("FManagedArrayCollection"))
+			for (const FDataflowOutput* const Output : DataflowNode->GetOutputs())
 			{
-				const FManagedArrayCollection DefaultValue;
-				TSharedPtr<FManagedArrayCollection> Collection = MakeShared<FManagedArrayCollection>(Output->GetValue<FManagedArrayCollection>(*Context, DefaultValue));
-
-				// see if the output collection is a ClothCollection
-				const UE::Chaos::ClothAsset::FCollectionClothConstFacade ClothFacade(Collection);
-				if (ClothFacade.IsValid())
+				if (Output->GetType() == FName("FManagedArrayCollection"))
 				{
-					return Collection;
-				}
+					const FManagedArrayCollection DefaultValue;
+					TSharedPtr<FManagedArrayCollection> Collection = MakeShared<FManagedArrayCollection>(Output->GetValue<FManagedArrayCollection>(*Context, DefaultValue));
 
-				return MakeShared<FManagedArrayCollection>();
+					// see if the output collection is a ClothCollection
+					const UE::Chaos::ClothAsset::FCollectionClothConstFacade ClothFacade(Collection);
+					if (ClothFacade.IsValid())
+					{
+						return Collection;
+					}
+
+					return MakeShared<FManagedArrayCollection>();
+				}
 			}
 		}
 
@@ -976,7 +979,7 @@ void FChaosClothAssetEditorToolkit::OnNodeSelectionChanged(const TSet<UObject*>&
 		Dataflow->LastModifiedRenderTarget = Dataflow::FTimestamp::Current();
 	}
 
-	UChaosClothAssetEditorMode* const ClothMode = Cast<UChaosClothAssetEditorMode>(EditorModeManager->GetActiveScriptableMode(UChaosClothAssetEditorMode::EM_ChaosClothAssetEditorModeId));
+	UChaosClothAssetEditorMode* const ClothMode = CastChecked<UChaosClothAssetEditorMode>(EditorModeManager->GetActiveScriptableMode(UChaosClothAssetEditorMode::EM_ChaosClothAssetEditorModeId));
 	if (ClothMode)
 	{
 		ClothMode->SetSelectedClothCollection(Collection);
