Chaos Cloth:
- Fixed a crash in the static render data generation of multi-sections meshes that was using LOD based instead of section based indices to generate the MeshToMeshVertData for the cloth deformer.
- Fixed the static render data generation index remappping when using multi-sections meshes.
- Fixed the FClothCollection::GetElementsOffset() that was returning INDEX_NONE in some situations.

#rnx

diff --git a/Engine/Plugins/Experimental/ChaosClothAsset/Source/ChaosClothAsset/Private/ChaosClothAsset/ClothCollection.cpp b/Engine/Plugins/Experimental/ChaosClothAsset/Source/ChaosClothAsset/Private/ChaosClothAsset/ClothCollection.cpp
--- a/Engine/Plugins/Experimental/ChaosClothAsset/Source/ChaosClothAsset/Private/ChaosClothAsset/ClothCollection.cpp
+++ b/Engine/Plugins/Experimental/ChaosClothAsset/Source/ChaosClothAsset/Private/ChaosClothAsset/ClothCollection.cpp
@@ -512,6 +512,15 @@ namespace UE::Chaos::ClothAsset
 		return Start;
 	}
 
+	int32 FClothCollection::GetElementsOffset(const TManagedArray<int32>* StartArray, int32 BaseElementIndex, int32 ElementIndex) const
+	{
+		while ((*StartArray)[BaseElementIndex] == INDEX_NONE && BaseElementIndex < ElementIndex)
+		{
+			++BaseElementIndex;
+		}
+		return (*StartArray)[ElementIndex] - (*StartArray)[BaseElementIndex];
+	}
+
 	int32 FClothCollection::GetNumSubElements(
 		const TManagedArray<int32>* StartArray,
 		const TManagedArray<int32>* EndArray,
diff --git a/Engine/Plugins/Experimental/ChaosClothAsset/Source/ChaosClothAsset/Public/ChaosClothAsset/ClothCollection.h b/Engine/Plugins/Experimental/ChaosClothAsset/Source/ChaosClothAsset/Public/ChaosClothAsset/ClothCollection.h
--- a/Engine/Plugins/Experimental/ChaosClothAsset/Source/ChaosClothAsset/Public/ChaosClothAsset/ClothCollection.h
+++ b/Engine/Plugins/Experimental/ChaosClothAsset/Source/ChaosClothAsset/Public/ChaosClothAsset/ClothCollection.h
@@ -69,7 +69,7 @@ namespace UE::Chaos::ClothAsset
 		 * Return the difference between the start index of an element to the start index of the first sub-element in the group (base).
 		 * Usefull for getting back and forth between LOD/pattern indexation modes.
 		 */
-		int32 GetElementsOffset(const TManagedArray<int32>* StartArray, int32 BaseElementIndex, int32 ElementIndex) const { return (*StartArray)[ElementIndex] - (*StartArray)[BaseElementIndex]; }
+		int32 GetElementsOffset(const TManagedArray<int32>* StartArray, int32 BaseElementIndex, int32 ElementIndex) const;
 
 		int32 GetNumSubElements(
 			const TManagedArray<int32>* StartArray,
diff --git a/Engine/Plugins/Experimental/ChaosClothAssetEditor/Source/ChaosClothAssetTools/Private/ChaosClothAsset/ClothAssetBuilderEditor.cpp b/Engine/Plugins/Experimental/ChaosClothAssetEditor/Source/ChaosClothAssetTools/Private/ChaosClothAsset/ClothAssetBuilderEditor.cpp
--- a/Engine/Plugins/Experimental/ChaosClothAssetEditor/Source/ChaosClothAssetTools/Private/ChaosClothAsset/ClothAssetBuilderEditor.cpp
+++ b/Engine/Plugins/Experimental/ChaosClothAssetEditor/Source/ChaosClothAssetTools/Private/ChaosClothAsset/ClothAssetBuilderEditor.cpp
@@ -243,7 +243,8 @@ void UClothAssetBuilderEditor::BuildLod(FSkeletalMeshLODModel& LODModel, const U
 		}
 
 		// Remap the LOD indices with the new vertex indices
-		for (uint32& RenderIndex : LODModel.IndexBuffer)
+		const TArrayView<uint32> SectionIndexBuffer(LODModel.IndexBuffer.GetData() + Section.BaseIndex, NumIndices);
+		for (uint32& RenderIndex : SectionIndexBuffer)
 		{
 			RenderIndex = LodRenderIndexRemap[RenderIndex];
 		}
@@ -263,20 +264,30 @@ void UClothAssetBuilderEditor::BuildLod(FSkeletalMeshLODModel& LODModel, const U
 			Section.ClothingData.AssetGuid = ClothAsset.AssetGuid;  // There is only one cloth asset,
 			Section.CorrespondClothAssetIndex = 0;       // this one
 
-			TArray<FVector3f> RenderPositions;
-			TArray<FVector3f> RenderNormals;
-			TArray<FVector3f> RenderTangents;
+			TArray<FVector3f> SectionRenderPositions;
+			TArray<FVector3f> SectionRenderNormals;
+			TArray<FVector3f> SectionRenderTangents;
+			SectionRenderPositions.Reserve(NumVertices);
+			SectionRenderNormals.Reserve(NumVertices);
+			SectionRenderTangents.Reserve(NumVertices);
 			for (const FSoftSkinVertex& SoftVert : Section.SoftVertices)
 			{
-				RenderPositions.Add(SoftVert.Position);
-				RenderNormals.Add(SoftVert.TangentZ);
-				RenderTangents.Add(SoftVert.TangentX);
+				SectionRenderPositions.Add(SoftVert.Position);
+				SectionRenderNormals.Add(SoftVert.TangentZ);
+				SectionRenderTangents.Add(SoftVert.TangentX);
+			}
+			TArray<uint32> SectionRenderIndices;
+			SectionRenderIndices.Reserve(NumIndices);
+			for (const uint32 LodModelVertIndex : SectionIndexBuffer)
+			{
+				SectionRenderIndices.Add(LodModelVertIndex - Section.BaseVertexIndex);
 			}
 
-			const ClothingMeshUtils::ClothMeshDesc TargetMesh(RenderPositions, 
-				RenderNormals, 
-				RenderTangents,
-				LODModel.IndexBuffer);
+			const ClothingMeshUtils::ClothMeshDesc TargetMesh(
+				SectionRenderPositions, 
+				SectionRenderNormals, 
+				SectionRenderTangents,
+				SectionRenderIndices);
 
 			ClothingMeshUtils::GenerateMeshToMeshVertData(
 				Section.ClothMappingDataLODs[0],
