Chaos Cloth:
- Add new atomic simulation config nodes.
- Rename the anisotropic XPBD properties to include the word Aniso as to not conflict with the non anisotropic versions.
- Tidied up DataflowNodes files.
- Add name check for weight map names.

#rnx

diff --git a/Engine/Plugins/Experimental/ChaosClothAssetEditor/Source/ChaosClothAssetDataflowNodes/ChaosClothAssetDataflowNodes.Build.cs b/Engine/Plugins/Experimental/ChaosClothAssetEditor/Source/ChaosClothAssetDataflowNodes/ChaosClothAssetDataflowNodes.Build.cs
--- a/Engine/Plugins/Experimental/ChaosClothAssetEditor/Source/ChaosClothAssetDataflowNodes/ChaosClothAssetDataflowNodes.Build.cs
+++ b/Engine/Plugins/Experimental/ChaosClothAssetEditor/Source/ChaosClothAssetDataflowNodes/ChaosClothAssetDataflowNodes.Build.cs
@@ -21,19 +21,24 @@ public class ChaosClothAssetDataflowNodes : ModuleRules
 				"ChaosCloth",
 				"ChaosClothAsset",
 				"ChaosClothAssetEngine",
+				"ClothingSystemRuntimeCommon",
 				"CoreUObject",
 				"DataflowCore",
 				"DataflowEditor",
 				"DataflowEngine",
 				"DatasmithImporter",
+				"DetailCustomizations",
 				"DynamicMesh",
 				"Engine",
 				"ExternalSource",
 				"GeometryCore",
+				"InputCore",
 				"MeshConversion",
 				"MeshDescription",
 				"SkeletalMeshDescription",
 				"Slate",
+				"SlateCore",
+				"UnrealEd",
 				// ... add private dependencies that you statically link with here ...	
 			}
 			);
diff --git a/Engine/Plugins/Experimental/ChaosClothAssetEditor/Source/ChaosClothAssetDataflowNodes/Private/ChaosClothAsset/ChaosClothAssetDataflowNodesModule.cpp b/Engine/Plugins/Experimental/ChaosClothAssetEditor/Source/ChaosClothAssetDataflowNodes/Private/ChaosClothAsset/ChaosClothAssetDataflowNodesModule.cpp
--- a/Engine/Plugins/Experimental/ChaosClothAssetEditor/Source/ChaosClothAssetDataflowNodes/Private/ChaosClothAsset/ChaosClothAssetDataflowNodesModule.cpp
+++ b/Engine/Plugins/Experimental/ChaosClothAssetEditor/Source/ChaosClothAssetDataflowNodes/Private/ChaosClothAsset/ChaosClothAssetDataflowNodesModule.cpp
@@ -2,17 +2,31 @@
 
 #include "ChaosClothAsset/ChaosClothAssetDataflowNodesModule.h"
 #include "ChaosClothAsset/DataflowNodes.h"
+#include "ChaosClothAsset/WeightedValueCustomization.h"
 #include "Modules/ModuleManager.h"
 
 #define LOCTEXT_NAMESPACE "FChaosClothAssetDataflowNodesModule"
 
 void FChaosClothAssetDataflowNodesModule::StartupModule()
 {
-	UE::Chaos::ClothAsset::DataflowNodes::Register();
+	using namespace UE::Chaos::ClothAsset;
+
+	DataflowNodes::Register();
+
+	// Register type customizations
+	if (FPropertyEditorModule* const PropertyModule = FModuleManager::GetModulePtr<FPropertyEditorModule>("PropertyEditor"))
+	{
+		PropertyModule->RegisterCustomPropertyTypeLayout("ChaosClothAssetWeightedValue", FOnGetPropertyTypeCustomizationInstance::CreateStatic(&FWeightedValueCustomization::MakeInstance));
+	}
 }
 
 void FChaosClothAssetDataflowNodesModule::ShutdownModule()
 {
+	// Unregister type customizations
+	if (FPropertyEditorModule* const PropertyModule = FModuleManager::GetModulePtr<FPropertyEditorModule>("PropertyEditor"))
+	{
+		PropertyModule->UnregisterCustomPropertyTypeLayout("ChaosClothAssetWeightedValue");
+	}
 }
 
 #undef LOCTEXT_NAMESPACE
diff --git a/Engine/Plugins/Experimental/ChaosClothAssetEditor/Source/ChaosClothAssetDataflowNodes/Private/ChaosClothAsset/DataflowNodes.cpp b/Engine/Plugins/Experimental/ChaosClothAssetEditor/Source/ChaosClothAssetDataflowNodes/Private/ChaosClothAsset/DataflowNodes.cpp
--- a/Engine/Plugins/Experimental/ChaosClothAssetEditor/Source/ChaosClothAssetDataflowNodes/Private/ChaosClothAsset/DataflowNodes.cpp
+++ b/Engine/Plugins/Experimental/ChaosClothAssetEditor/Source/ChaosClothAssetDataflowNodes/Private/ChaosClothAsset/DataflowNodes.cpp
@@ -1,19 +1,40 @@
 // Copyright Epic Games, Inc. All Rights Reserved.
 
 #include "ChaosClothAsset/DataflowNodes.h"
-#include "ChaosClothAsset/TerminalNode.h"
+#include "ChaosClothAsset/AddWeightMapNode.h"
+#include "ChaosClothAsset/BindToRootBoneNode.h"
+#include "ChaosClothAsset/CopySimulationToRenderMeshNode.h"
 #include "ChaosClothAsset/DatasmithImportNode.h"
+#include "ChaosClothAsset/DeleteRenderMeshNode.h"
 #include "ChaosClothAsset/ImportNode.h"
-#include "ChaosClothAsset/SkeletalMeshImportNode.h"
-#include "ChaosClothAsset/StaticMeshImportNode.h"
 #include "ChaosClothAsset/PatternSelectionNode.h"
-#include "ChaosClothAsset/DeleteRenderMeshNode.h"
-#include "ChaosClothAsset/CopySimulationToRenderMeshNode.h"
 #include "ChaosClothAsset/ReverseNormalsNode.h"
+#include "ChaosClothAsset/SimulationAerodynamicsConfigNode.h"
+#include "ChaosClothAsset/SimulationAnimDriveConfigNode.h"
+#include "ChaosClothAsset/SimulationBackstopConfigNode.h"
+#include "ChaosClothAsset/SimulationCollisionConfigNode.h"
+#include "ChaosClothAsset/SimulationDampingConfigNode.h"
 #include "ChaosClothAsset/SimulationDefaultConfigNode.h"
-#include "ChaosClothAsset/SimulationXPBDConfigNode.h"
-#include "ChaosClothAsset/BindToRootBoneNode.h"
-#include "ChaosClothAsset/AddWeightMapNode.h"
+#include "ChaosClothAsset/SimulationGravityConfigNode.h"
+#include "ChaosClothAsset/SimulationMassConfigNode.h"
+#include "ChaosClothAsset/SimulationMaxDistanceConfigNode.h"
+#include "ChaosClothAsset/SimulationPBDAreaSpringConfigNode.h"
+#include "ChaosClothAsset/SimulationPBDBendingElementConfigNode.h"
+#include "ChaosClothAsset/SimulationPBDBendingSpringConfigNode.h"
+#include "ChaosClothAsset/SimulationPBDEdgeSpringConfigNode.h"
+#include "ChaosClothAsset/SimulationPressureConfigNode.h"
+#include "ChaosClothAsset/SimulationSelfCollisionConfigNode.h"
+#include "ChaosClothAsset/SimulationSolverConfigNode.h"
+#include "ChaosClothAsset/SimulationVelocityScaleConfigNode.h"
+#include "ChaosClothAsset/SimulationXPBDAreaSpringConfigNode.h"
+#include "ChaosClothAsset/SimulationXPBDAnisoBendingConfigNode.h"
+#include "ChaosClothAsset/SimulationXPBDAnisoStretchConfigNode.h"
+#include "ChaosClothAsset/SimulationXPBDBendingElementConfigNode.h"
+#include "ChaosClothAsset/SimulationXPBDBendingSpringConfigNode.h"
+#include "ChaosClothAsset/SimulationXPBDEdgeSpringConfigNode.h"
+#include "ChaosClothAsset/SkeletalMeshImportNode.h"
+#include "ChaosClothAsset/StaticMeshImportNode.h"
+#include "ChaosClothAsset/TerminalNode.h"
 #include "ChaosClothAsset/TransferSkinWeightsNode.h"
 #include "Dataflow/DataflowNodeFactory.h"
 #include "Framework/Notifications/NotificationManager.h"
@@ -25,28 +46,54 @@ namespace UE::Chaos::ClothAsset::DataflowNodes
 {
 	void Register()
 	{
-		DATAFLOW_NODE_REGISTER_CREATION_FACTORY(FChaosClothAssetTerminalNode);
+		DATAFLOW_NODE_REGISTER_CREATION_FACTORY(FChaosClothAssetAddWeightMapNode);
+		DATAFLOW_NODE_REGISTER_CREATION_FACTORY(FChaosClothAssetBindToRootBoneNode);
+		DATAFLOW_NODE_REGISTER_CREATION_FACTORY(FChaosClothAssetCopySimulationToRenderMeshNode);
 		DATAFLOW_NODE_REGISTER_CREATION_FACTORY(FChaosClothAssetDatasmithImportNode);
+		DATAFLOW_NODE_REGISTER_CREATION_FACTORY(FChaosClothAssetDeleteRenderMeshNode);
 		DATAFLOW_NODE_REGISTER_CREATION_FACTORY(FChaosClothAssetImportNode);
-		DATAFLOW_NODE_REGISTER_CREATION_FACTORY(FChaosClothAssetSkeletalMeshImportNode);
-		DATAFLOW_NODE_REGISTER_CREATION_FACTORY(FChaosClothAssetStaticMeshImportNode);
 		DATAFLOW_NODE_REGISTER_CREATION_FACTORY(FChaosClothAssetPatternSelectionNode);
-		DATAFLOW_NODE_REGISTER_CREATION_FACTORY(FChaosClothAssetDeleteRenderMeshNode);
-		DATAFLOW_NODE_REGISTER_CREATION_FACTORY(FChaosClothAssetCopySimulationToRenderMeshNode);
 		DATAFLOW_NODE_REGISTER_CREATION_FACTORY(FChaosClothAssetReverseNormalsNode);
+		DATAFLOW_NODE_REGISTER_CREATION_FACTORY(FChaosClothAssetSimulationAerodynamicsConfigNode);
+		DATAFLOW_NODE_REGISTER_CREATION_FACTORY(FChaosClothAssetSimulationAnimDriveConfigNode);
+		DATAFLOW_NODE_REGISTER_CREATION_FACTORY(FChaosClothAssetSimulationBackstopConfigNode);
+		DATAFLOW_NODE_REGISTER_CREATION_FACTORY(FChaosClothAssetSimulationCollisionConfigNode);
+		DATAFLOW_NODE_REGISTER_CREATION_FACTORY(FChaosClothAssetSimulationDampingConfigNode);
 		DATAFLOW_NODE_REGISTER_CREATION_FACTORY(FChaosClothAssetSimulationDefaultConfigNode);
-		DATAFLOW_NODE_REGISTER_CREATION_FACTORY(FChaosClothAssetSimulationXPBDConfigNode);
-		DATAFLOW_NODE_REGISTER_CREATION_FACTORY(FChaosClothAssetBindToRootBoneNode);
-		DATAFLOW_NODE_REGISTER_CREATION_FACTORY(FChaosClothAssetAddWeightMapNode);
+		DATAFLOW_NODE_REGISTER_CREATION_FACTORY(FChaosClothAssetSimulationGravityConfigNode);
+		DATAFLOW_NODE_REGISTER_CREATION_FACTORY(FChaosClothAssetSimulationMassConfigNode);
+		DATAFLOW_NODE_REGISTER_CREATION_FACTORY(FChaosClothAssetSimulationMaxDistanceConfigNode);
+		DATAFLOW_NODE_REGISTER_CREATION_FACTORY(FChaosClothAssetSimulationPBDAreaSpringConfigNode);
+		DATAFLOW_NODE_REGISTER_CREATION_FACTORY(FChaosClothAssetSimulationPBDBendingElementConfigNode);
+		DATAFLOW_NODE_REGISTER_CREATION_FACTORY(FChaosClothAssetSimulationPBDBendingSpringConfigNode);
+		DATAFLOW_NODE_REGISTER_CREATION_FACTORY(FChaosClothAssetSimulationPBDEdgeSpringConfigNode);
+		DATAFLOW_NODE_REGISTER_CREATION_FACTORY(FChaosClothAssetSimulationPressureConfigNode);
+		DATAFLOW_NODE_REGISTER_CREATION_FACTORY(FChaosClothAssetSimulationSelfCollisionConfigNode);
+		DATAFLOW_NODE_REGISTER_CREATION_FACTORY(FChaosClothAssetSimulationSolverConfigNode);
+		DATAFLOW_NODE_REGISTER_CREATION_FACTORY(FChaosClothAssetSimulationVelocityScaleConfigNode);
+		DATAFLOW_NODE_REGISTER_CREATION_FACTORY(FChaosClothAssetSimulationXPBDAnisoBendingConfigNode);
+		DATAFLOW_NODE_REGISTER_CREATION_FACTORY(FChaosClothAssetSimulationXPBDAnisoStretchConfigNode);
+		DATAFLOW_NODE_REGISTER_CREATION_FACTORY(FChaosClothAssetSimulationXPBDAreaSpringConfigNode);
+		DATAFLOW_NODE_REGISTER_CREATION_FACTORY(FChaosClothAssetSimulationXPBDBendingElementConfigNode);
+		DATAFLOW_NODE_REGISTER_CREATION_FACTORY(FChaosClothAssetSimulationXPBDBendingSpringConfigNode);
+		DATAFLOW_NODE_REGISTER_CREATION_FACTORY(FChaosClothAssetSimulationXPBDEdgeSpringConfigNode);
+		DATAFLOW_NODE_REGISTER_CREATION_FACTORY(FChaosClothAssetSkeletalMeshImportNode);
+		DATAFLOW_NODE_REGISTER_CREATION_FACTORY(FChaosClothAssetStaticMeshImportNode);
+		DATAFLOW_NODE_REGISTER_CREATION_FACTORY(FChaosClothAssetTerminalNode);
 		DATAFLOW_NODE_REGISTER_CREATION_FACTORY(FChaosClothAssetTransferSkinWeightsNode);
 	}
 
-	void LogAndToastWarning(const FText& Error)
+	void LogAndToastWarning(const FText& Message)
 	{
-		FNotificationInfo Info(Error);
-		Info.ExpireDuration = 5.0f;
-		FSlateNotificationManager::Get().AddNotification(Info);
+		FNotificationInfo NotificationInfo(Message);
+		NotificationInfo.ExpireDuration = 5.0f;
+		FSlateNotificationManager::Get().AddNotification(NotificationInfo);
 
-		UE_LOG(LogChaosClothAssetDataflowNodes, Warning, TEXT("%s"), *Error.ToString());
+		UE_LOG(LogChaosClothAssetDataflowNodes, Warning, TEXT("%s"), *Message.ToString());
+	}
+
+	FString MakeAttributeName(const FString& String)
+	{
+		return SlugStringForValidName(String, TEXT("_")).Replace(TEXT("\\"), TEXT("_"));
 	}
 }
diff --git a/Engine/Plugins/Experimental/ChaosClothAssetEditor/Source/ChaosClothAssetDataflowNodes/Private/ChaosClothAsset/SimulationAerodynamicsConfigNode.cpp b/Engine/Plugins/Experimental/ChaosClothAssetEditor/Source/ChaosClothAssetDataflowNodes/Private/ChaosClothAsset/SimulationAerodynamicsConfigNode.cpp
new file mode 100644
--- /dev/null
+++ b/Engine/Plugins/Experimental/ChaosClothAssetEditor/Source/ChaosClothAssetDataflowNodes/Private/ChaosClothAsset/SimulationAerodynamicsConfigNode.cpp
@@ -0,0 +1,21 @@
+// Copyright Epic Games, Inc. All Rights Reserved.
+
+#include "ChaosClothAsset/SimulationAerodynamicsConfigNode.h"
+#include "ChaosClothAsset/DataflowNodes.h"
+#include "ChaosClothAsset/SimulationBaseConfigNodePrivate.h"
+#include "Chaos/CollectionPropertyFacade.h"
+
+#include UE_INLINE_GENERATED_CPP_BY_NAME(SimulationAerodynamicsConfigNode)
+
+FChaosClothAssetSimulationAerodynamicsConfigNode::FChaosClothAssetSimulationAerodynamicsConfigNode(const Dataflow::FNodeParameters& InParam, FGuid InGuid)
+	: FChaosClothAssetSimulationBaseConfigNode(InParam, InGuid)
+{
+	RegisterCollectionConnections();
+}
+
+void FChaosClothAssetSimulationAerodynamicsConfigNode::AddProperties(::Chaos::Softs::FCollectionPropertyMutableFacade& Properties) const
+{
+	UE_CHAOS_CLOTHASSET_SIMULATIONCONFIG_SETPROPERTY(FluidDensity);
+	UE_CHAOS_CLOTHASSET_SIMULATIONCONFIG_SETPROPERTYWEIGHTED(Drag);
+	UE_CHAOS_CLOTHASSET_SIMULATIONCONFIG_SETPROPERTYWEIGHTED(Lift);
+}
diff --git a/Engine/Plugins/Experimental/ChaosClothAssetEditor/Source/ChaosClothAssetDataflowNodes/Private/ChaosClothAsset/SimulationAnimDriveConfigNode.cpp b/Engine/Plugins/Experimental/ChaosClothAssetEditor/Source/ChaosClothAssetDataflowNodes/Private/ChaosClothAsset/SimulationAnimDriveConfigNode.cpp
new file mode 100644
--- /dev/null
+++ b/Engine/Plugins/Experimental/ChaosClothAssetEditor/Source/ChaosClothAssetDataflowNodes/Private/ChaosClothAsset/SimulationAnimDriveConfigNode.cpp
@@ -0,0 +1,20 @@
+// Copyright Epic Games, Inc. All Rights Reserved.
+
+#include "ChaosClothAsset/SimulationAnimDriveConfigNode.h"
+#include "ChaosClothAsset/DataflowNodes.h"
+#include "ChaosClothAsset/SimulationBaseConfigNodePrivate.h"
+#include "Chaos/CollectionPropertyFacade.h"
+
+#include UE_INLINE_GENERATED_CPP_BY_NAME(SimulationAnimDriveConfigNode)
+
+FChaosClothAssetSimulationAnimDriveConfigNode::FChaosClothAssetSimulationAnimDriveConfigNode(const Dataflow::FNodeParameters& InParam, FGuid InGuid)
+	: FChaosClothAssetSimulationBaseConfigNode(InParam, InGuid)
+{
+	RegisterCollectionConnections();
+}
+
+void FChaosClothAssetSimulationAnimDriveConfigNode::AddProperties(::Chaos::Softs::FCollectionPropertyMutableFacade& Properties) const
+{
+	UE_CHAOS_CLOTHASSET_SIMULATIONCONFIG_SETPROPERTYWEIGHTED(AnimDriveStiffness);
+	UE_CHAOS_CLOTHASSET_SIMULATIONCONFIG_SETPROPERTYWEIGHTED(AnimDriveDamping);
+}
diff --git a/Engine/Plugins/Experimental/ChaosClothAssetEditor/Source/ChaosClothAssetDataflowNodes/Private/ChaosClothAsset/SimulationBackstopConfigNode.cpp b/Engine/Plugins/Experimental/ChaosClothAssetEditor/Source/ChaosClothAssetDataflowNodes/Private/ChaosClothAsset/SimulationBackstopConfigNode.cpp
new file mode 100644
--- /dev/null
+++ b/Engine/Plugins/Experimental/ChaosClothAssetEditor/Source/ChaosClothAssetDataflowNodes/Private/ChaosClothAsset/SimulationBackstopConfigNode.cpp
@@ -0,0 +1,20 @@
+// Copyright Epic Games, Inc. All Rights Reserved.
+
+#include "ChaosClothAsset/SimulationBackstopConfigNode.h"
+#include "ChaosClothAsset/DataflowNodes.h"
+#include "ChaosClothAsset/SimulationBaseConfigNodePrivate.h"
+#include "Chaos/CollectionPropertyFacade.h"
+
+#include UE_INLINE_GENERATED_CPP_BY_NAME(SimulationBackstopConfigNode)
+
+FChaosClothAssetSimulationBackstopConfigNode::FChaosClothAssetSimulationBackstopConfigNode(const Dataflow::FNodeParameters& InParam, FGuid InGuid)
+	: FChaosClothAssetSimulationBaseConfigNode(InParam, InGuid)
+{
+	RegisterCollectionConnections();
+}
+
+void FChaosClothAssetSimulationBackstopConfigNode::AddProperties(::Chaos::Softs::FCollectionPropertyMutableFacade& Properties) const
+{
+	UE_CHAOS_CLOTHASSET_SIMULATIONCONFIG_SETPROPERTYWEIGHTED(BackstopDistance);
+	UE_CHAOS_CLOTHASSET_SIMULATIONCONFIG_SETPROPERTYWEIGHTED(BackstopRadius);
+}
diff --git a/Engine/Plugins/Experimental/ChaosClothAssetEditor/Source/ChaosClothAssetDataflowNodes/Private/ChaosClothAsset/SimulationBaseConfigNode.cpp b/Engine/Plugins/Experimental/ChaosClothAssetEditor/Source/ChaosClothAssetDataflowNodes/Private/ChaosClothAsset/SimulationBaseConfigNode.cpp
new file mode 100644
--- /dev/null
+++ b/Engine/Plugins/Experimental/ChaosClothAssetEditor/Source/ChaosClothAssetDataflowNodes/Private/ChaosClothAsset/SimulationBaseConfigNode.cpp
@@ -0,0 +1,103 @@
+// Copyright Epic Games, Inc. All Rights Reserved.
+
+#include "ChaosClothAsset/SimulationBaseConfigNode.h"
+#include "ChaosClothAsset/DataflowNodes.h"
+#include "Chaos/CollectionPropertyFacade.h"
+
+#define LOCTEXT_NAMESPACE "ChaosClothAssetSimulationBaseConfigNode"
+
+namespace UE::Chaos::ClothAsset::Private
+{
+	static void LogAndToastDuplicateProperty(const FName& NodeName, const FName& PropertyName)
+	{
+		using namespace UE::Chaos::ClothAsset::DataflowNodes;
+
+		const FText Message = FText::Format(
+			LOCTEXT(
+				"DuplicateProperty",
+				"Cloth collection property '{1}' was already set in an upstream node, and its values are now overriden by node '{0}'."),
+			FText::FromName(NodeName),
+			FText::FromName(PropertyName));
+
+		LogAndToastWarning(Message);
+	}
+
+	static void LogAndToastSimilarProperty(const FName& NodeName, const FName& PropertyName, const FName& SimilarPropertyName)
+	{
+		using namespace UE::Chaos::ClothAsset::DataflowNodes;
+
+		const FText Message = FText::Format(
+			LOCTEXT(
+				"DuplicateProperty",
+				"Cloth collection property '{1}' set in node and '{0}' is similar to the property '{2}' already set in an upstream node, "
+				"which might result in an undefined simulation behavior."),
+			FText::FromName(NodeName),
+			FText::FromName(PropertyName),
+			FText::FromName(SimilarPropertyName));
+
+		LogAndToastWarning(Message);
+	}
+}
+
+FChaosClothAssetSimulationBaseConfigNode::FChaosClothAssetSimulationBaseConfigNode(const Dataflow::FNodeParameters& InParam, FGuid InGuid)
+	: FDataflowNode(InParam, InGuid)
+{}
+
+void FChaosClothAssetSimulationBaseConfigNode::RegisterCollectionConnections()
+{
+	RegisterInputConnection(&Collection);
+	RegisterOutputConnection(&Collection, &Collection);
+}
+
+void FChaosClothAssetSimulationBaseConfigNode::Evaluate(Dataflow::FContext& Context, const FDataflowOutput* Out) const
+{
+	using namespace Chaos::Softs;
+	using namespace UE::Chaos::ClothAsset::DataflowNodes;
+
+	if (Out->IsA<FManagedArrayCollection>(&Collection))
+	{
+		FManagedArrayCollection InCollection = GetValue<FManagedArrayCollection>(Context, &Collection);
+		const TSharedRef<FManagedArrayCollection> ClothCollection = MakeShared<FManagedArrayCollection>(MoveTemp(InCollection));
+
+		FCollectionPropertyMutableFacade Properties(ClothCollection);
+		Properties.DefineSchema();
+
+		AddProperties(Properties);
+
+		SetValue<FManagedArrayCollection>(Context, *ClothCollection, &Collection);
+	}
+}
+
+int32 FChaosClothAssetSimulationBaseConfigNode::AddPropertyHelper(
+	::Chaos::Softs::FCollectionPropertyMutableFacade& Properties,
+	const FName& PropertyName,
+	bool bIsAnimatable,
+	const TArray<FName>& SimilarPropertyNames) const
+{
+	using namespace UE::Chaos::ClothAsset::Private;
+
+	constexpr bool bIsEnabled = true;
+
+	int32 KeyIndex = Properties.GetKeyIndex(PropertyName.ToString());
+	if (KeyIndex == INDEX_NONE)
+	{
+		KeyIndex = Properties.AddProperty(PropertyName.ToString(), bIsEnabled, bIsAnimatable);
+	}
+	else
+	{
+		Properties.SetAnimatable(KeyIndex, bIsAnimatable);
+		LogAndToastDuplicateProperty(FDataflowNode::Name, PropertyName);
+	}
+
+	for (const FName& SimilarPropertyName : SimilarPropertyNames)
+	{
+		if (Properties.GetKeyIndex(SimilarPropertyName.ToString()) != INDEX_NONE)
+		{
+			LogAndToastSimilarProperty(FDataflowNode::Name, PropertyName, SimilarPropertyName);
+		}
+	}
+
+	return KeyIndex;
+}
+
+#undef LOCTEXT_NAMESPACE
diff --git a/Engine/Plugins/Experimental/ChaosClothAssetEditor/Source/ChaosClothAssetDataflowNodes/Private/ChaosClothAsset/SimulationBaseConfigNodePrivate.h b/Engine/Plugins/Experimental/ChaosClothAssetEditor/Source/ChaosClothAssetDataflowNodes/Private/ChaosClothAsset/SimulationBaseConfigNodePrivate.h
new file mode 100644
--- /dev/null
+++ b/Engine/Plugins/Experimental/ChaosClothAssetEditor/Source/ChaosClothAssetDataflowNodes/Private/ChaosClothAsset/SimulationBaseConfigNodePrivate.h
@@ -0,0 +1,64 @@
+// Copyright Epic Games, Inc. All Rights Reserved.
+
+#pragma once 
+
+/** Macro for adding and setting value properties from inside AddProperties() */
+#define UE_CHAOS_CLOTHASSET_SIMULATIONCONFIG_SETPROPERTY(Property) \
+	const int32 Property##KeyIndex = AddPropertyHelper(Properties, FName(TEXT(#Property))); \
+	Properties.SetValue(Property##KeyIndex, Property);
+
+/** Macro for adding and setting boolean value properties from inside AddProperties() */
+#define UE_CHAOS_CLOTHASSET_SIMULATIONCONFIG_SETPROPERTYBOOL(Property) \
+	const int32 Property##KeyIndex = AddPropertyHelper(Properties, FName(TEXT(#Property))); \
+	Properties.SetValue(Property##KeyIndex, b##Property);
+
+/** Macro for adding and setting enum value as int32 properties from inside AddProperties() */
+#define UE_CHAOS_CLOTHASSET_SIMULATIONCONFIG_SETPROPERTYENUM(Property) \
+	const int32 Property##KeyIndex = AddPropertyHelper(Properties, FName(TEXT(#Property))); \
+	Properties.SetValue(Property##KeyIndex, (int32)Property);
+
+/** Macro for adding and setting weighted value properties from inside AddProperties() */
+#define UE_CHAOS_CLOTHASSET_SIMULATIONCONFIG_SETPROPERTYWEIGHTED(Property) \
+	const int32 Property##KeyIndex = AddPropertyHelper(Properties, FName(TEXT(#Property)), Property.bIsAnimatable); \
+	Properties.SetWeightedValue(Property##KeyIndex, Property.Low, Property.High); \
+	Properties.SetStringValue(Property##KeyIndex, Property.WeightMap);
+
+/** Macro for adding and setting weighted value properties from inside AddProperties() and check for a similar property being overriden. */
+#define UE_CHAOS_CLOTHASSET_SIMULATIONCONFIG_SETPROPERTYCHECKED1(Property, SimilarProperty1) \
+	const int32 Property##KeyIndex = AddPropertyHelper(Properties, FName(TEXT(#Property)), true, \
+		{ FName(TEXT(#SimilarProperty1)) }); \
+	Properties.SetValue(Property##KeyIndex, Property);
+
+/** Macro for adding and setting weighted value properties from inside AddProperties() and check for two similar properties being overriden. */
+#define UE_CHAOS_CLOTHASSET_SIMULATIONCONFIG_SETPROPERTYCHECKED2(Property, SimilarProperty1, SimilarProperty2) \
+	const int32 Property##KeyIndex = AddPropertyHelper(Properties, FName(TEXT(#Property)), true, \
+		{ FName(TEXT(#SimilarProperty1)), FName(TEXT(#SimilarProperty2)) }); \
+	Properties.SetValue(Property##KeyIndex, Property);
+
+/** Macro for adding and setting weighted value properties from inside AddProperties() and check for a similar property being overriden. */
+#define UE_CHAOS_CLOTHASSET_SIMULATIONCONFIG_SETPROPERTYWEIGHTEDCHECKED1(Property, SimilarProperty1) \
+	const int32 Property##KeyIndex = AddPropertyHelper(Properties, FName(TEXT(#Property)), Property.bIsAnimatable, \
+		{ FName(TEXT(#SimilarProperty1)) }); \
+	Properties.SetWeightedValue(Property##KeyIndex, Property.Low, Property.High); \
+	Properties.SetStringValue(Property##KeyIndex, Property.WeightMap);
+
+/** Macro for adding and setting weighted value properties from inside AddProperties() and check for two similar properties being overriden. */
+#define UE_CHAOS_CLOTHASSET_SIMULATIONCONFIG_SETPROPERTYWEIGHTEDCHECKED2(Property, SimilarProperty1, SimilarProperty2) \
+	const int32 Property##KeyIndex = AddPropertyHelper(Properties, FName(TEXT(#Property)), Property.bIsAnimatable, \
+		{ FName(TEXT(#SimilarProperty1)), FName(TEXT(#SimilarProperty2)) }); \
+	Properties.SetWeightedValue(Property##KeyIndex, Property.Low, Property.High); \
+	Properties.SetStringValue(Property##KeyIndex, Property.WeightMap);
+
+/** Macro for adding and setting weighted value properties from inside AddProperties() and check for three similar properties being overriden. */
+#define UE_CHAOS_CLOTHASSET_SIMULATIONCONFIG_SETPROPERTYWEIGHTEDCHECKED3(Property, SimilarProperty1, SimilarProperty2, SimilarProperty3) \
+	const int32 Property##KeyIndex = AddPropertyHelper(Properties, FName(TEXT(#Property)), Property.bIsAnimatable, \
+		{ FName(TEXT(#SimilarProperty1)), FName(TEXT(#SimilarProperty2)), FName(TEXT(#SimilarProperty3)) }); \
+	Properties.SetWeightedValue(Property##KeyIndex, Property.Low, Property.High); \
+	Properties.SetStringValue(Property##KeyIndex, Property.WeightMap);
+
+/** Macro for adding and setting weighted value properties from inside AddProperties() and check for four similar properties being overriden. */
+#define UE_CHAOS_CLOTHASSET_SIMULATIONCONFIG_SETPROPERTYWEIGHTEDCHECKED4(Property, SimilarProperty1, SimilarProperty2, SimilarProperty3, SimilarProperty4) \
+	const int32 Property##KeyIndex = AddPropertyHelper(Properties, FName(TEXT(#Property)), Property.bIsAnimatable, \
+		{ FName(TEXT(#SimilarProperty1)), FName(TEXT(#SimilarProperty2)), FName(TEXT(#SimilarProperty3)), FName(TEXT(#SimilarProperty4)) }); \
+	Properties.SetWeightedValue(Property##KeyIndex, Property.Low, Property.High); \
+	Properties.SetStringValue(Property##KeyIndex, Property.WeightMap);
diff --git a/Engine/Plugins/Experimental/ChaosClothAssetEditor/Source/ChaosClothAssetDataflowNodes/Private/ChaosClothAsset/SimulationCollisionConfigNode.cpp b/Engine/Plugins/Experimental/ChaosClothAssetEditor/Source/ChaosClothAssetDataflowNodes/Private/ChaosClothAsset/SimulationCollisionConfigNode.cpp
new file mode 100644
--- /dev/null
+++ b/Engine/Plugins/Experimental/ChaosClothAssetEditor/Source/ChaosClothAssetDataflowNodes/Private/ChaosClothAsset/SimulationCollisionConfigNode.cpp
@@ -0,0 +1,21 @@
+// Copyright Epic Games, Inc. All Rights Reserved.
+
+#include "ChaosClothAsset/SimulationCollisionConfigNode.h"
+#include "ChaosClothAsset/DataflowNodes.h"
+#include "ChaosClothAsset/SimulationBaseConfigNodePrivate.h"
+#include "Chaos/CollectionPropertyFacade.h"
+
+#include UE_INLINE_GENERATED_CPP_BY_NAME(SimulationCollisionConfigNode)
+
+FChaosClothAssetSimulationCollisionConfigNode::FChaosClothAssetSimulationCollisionConfigNode(const Dataflow::FNodeParameters& InParam, FGuid InGuid)
+	: FChaosClothAssetSimulationBaseConfigNode(InParam, InGuid)
+{
+	RegisterCollectionConnections();
+}
+
+void FChaosClothAssetSimulationCollisionConfigNode::AddProperties(::Chaos::Softs::FCollectionPropertyMutableFacade& Properties) const
+{
+	UE_CHAOS_CLOTHASSET_SIMULATIONCONFIG_SETPROPERTY(CollisionThickness);
+	UE_CHAOS_CLOTHASSET_SIMULATIONCONFIG_SETPROPERTY(FrictionCoefficient);
+	UE_CHAOS_CLOTHASSET_SIMULATIONCONFIG_SETPROPERTYBOOL(UseCCD);
+}
diff --git a/Engine/Plugins/Experimental/ChaosClothAssetEditor/Source/ChaosClothAssetDataflowNodes/Private/ChaosClothAsset/SimulationDampingConfigNode.cpp b/Engine/Plugins/Experimental/ChaosClothAssetEditor/Source/ChaosClothAssetDataflowNodes/Private/ChaosClothAsset/SimulationDampingConfigNode.cpp
new file mode 100644
--- /dev/null
+++ b/Engine/Plugins/Experimental/ChaosClothAssetEditor/Source/ChaosClothAssetDataflowNodes/Private/ChaosClothAsset/SimulationDampingConfigNode.cpp
@@ -0,0 +1,20 @@
+// Copyright Epic Games, Inc. All Rights Reserved.
+
+#include "ChaosClothAsset/SimulationDampingConfigNode.h"
+#include "ChaosClothAsset/DataflowNodes.h"
+#include "ChaosClothAsset/SimulationBaseConfigNodePrivate.h"
+#include "Chaos/CollectionPropertyFacade.h"
+
+#include UE_INLINE_GENERATED_CPP_BY_NAME(SimulationDampingConfigNode)
+
+FChaosClothAssetSimulationDampingConfigNode::FChaosClothAssetSimulationDampingConfigNode(const Dataflow::FNodeParameters& InParam, FGuid InGuid)
+	: FChaosClothAssetSimulationBaseConfigNode(InParam, InGuid)
+{
+	RegisterCollectionConnections();
+}
+
+void FChaosClothAssetSimulationDampingConfigNode::AddProperties(::Chaos::Softs::FCollectionPropertyMutableFacade& Properties) const
+{
+	UE_CHAOS_CLOTHASSET_SIMULATIONCONFIG_SETPROPERTY(DampingCoefficient);
+	UE_CHAOS_CLOTHASSET_SIMULATIONCONFIG_SETPROPERTY(LocalDampingCoefficient);
+}
diff --git a/Engine/Plugins/Experimental/ChaosClothAssetEditor/Source/ChaosClothAssetDataflowNodes/Private/ChaosClothAsset/SimulationGravityConfigNode.cpp b/Engine/Plugins/Experimental/ChaosClothAssetEditor/Source/ChaosClothAssetDataflowNodes/Private/ChaosClothAsset/SimulationGravityConfigNode.cpp
new file mode 100644
--- /dev/null
+++ b/Engine/Plugins/Experimental/ChaosClothAssetEditor/Source/ChaosClothAssetDataflowNodes/Private/ChaosClothAsset/SimulationGravityConfigNode.cpp
@@ -0,0 +1,21 @@
+// Copyright Epic Games, Inc. All Rights Reserved.
+
+#include "ChaosClothAsset/SimulationGravityConfigNode.h"
+#include "ChaosClothAsset/DataflowNodes.h"
+#include "ChaosClothAsset/SimulationBaseConfigNodePrivate.h"
+#include "Chaos/CollectionPropertyFacade.h"
+
+#include UE_INLINE_GENERATED_CPP_BY_NAME(SimulationGravityConfigNode)
+
+FChaosClothAssetSimulationGravityConfigNode::FChaosClothAssetSimulationGravityConfigNode(const Dataflow::FNodeParameters& InParam, FGuid InGuid)
+	: FChaosClothAssetSimulationBaseConfigNode(InParam, InGuid)
+{
+	RegisterCollectionConnections();
+}
+
+void FChaosClothAssetSimulationGravityConfigNode::AddProperties(::Chaos::Softs::FCollectionPropertyMutableFacade& Properties) const
+{
+	UE_CHAOS_CLOTHASSET_SIMULATIONCONFIG_SETPROPERTYBOOL(UseGravityOverride);
+	UE_CHAOS_CLOTHASSET_SIMULATIONCONFIG_SETPROPERTY(GravityScale);
+	UE_CHAOS_CLOTHASSET_SIMULATIONCONFIG_SETPROPERTY(GravityOverride);
+}
diff --git a/Engine/Plugins/Experimental/ChaosClothAssetEditor/Source/ChaosClothAssetDataflowNodes/Private/ChaosClothAsset/SimulationMassConfigNode.cpp b/Engine/Plugins/Experimental/ChaosClothAssetEditor/Source/ChaosClothAssetDataflowNodes/Private/ChaosClothAsset/SimulationMassConfigNode.cpp
new file mode 100644
--- /dev/null
+++ b/Engine/Plugins/Experimental/ChaosClothAssetEditor/Source/ChaosClothAssetDataflowNodes/Private/ChaosClothAsset/SimulationMassConfigNode.cpp
@@ -0,0 +1,29 @@
+// Copyright Epic Games, Inc. All Rights Reserved.
+
+#include "ChaosClothAsset/SimulationMassConfigNode.h"
+#include "ChaosClothAsset/DataflowNodes.h"
+#include "ChaosClothAsset/SimulationBaseConfigNodePrivate.h"
+#include "Chaos/CollectionPropertyFacade.h"
+
+#include UE_INLINE_GENERATED_CPP_BY_NAME(SimulationMassConfigNode)
+
+FChaosClothAssetSimulationMassConfigNode::FChaosClothAssetSimulationMassConfigNode(const Dataflow::FNodeParameters& InParam, FGuid InGuid)
+	: FChaosClothAssetSimulationBaseConfigNode(InParam, InGuid)
+{
+	RegisterCollectionConnections();
+}
+
+void FChaosClothAssetSimulationMassConfigNode::AddProperties(::Chaos::Softs::FCollectionPropertyMutableFacade& Properties) const
+{
+	float MassValue;
+	switch (MassMode)
+	{
+	default:
+	case EClothMassMode::UniformMass: MassValue = UniformMass; break;
+	case EClothMassMode::TotalMass: MassValue = TotalMass; break;
+	case EClothMassMode::Density: MassValue = Density; break;
+	}
+
+	UE_CHAOS_CLOTHASSET_SIMULATIONCONFIG_SETPROPERTYENUM(MassMode);
+	UE_CHAOS_CLOTHASSET_SIMULATIONCONFIG_SETPROPERTY(MassValue);
+}
diff --git a/Engine/Plugins/Experimental/ChaosClothAssetEditor/Source/ChaosClothAssetDataflowNodes/Private/ChaosClothAsset/SimulationMaxDistanceConfigNode.cpp b/Engine/Plugins/Experimental/ChaosClothAssetEditor/Source/ChaosClothAssetDataflowNodes/Private/ChaosClothAsset/SimulationMaxDistanceConfigNode.cpp
new file mode 100644
--- /dev/null
+++ b/Engine/Plugins/Experimental/ChaosClothAssetEditor/Source/ChaosClothAssetDataflowNodes/Private/ChaosClothAsset/SimulationMaxDistanceConfigNode.cpp
@@ -0,0 +1,19 @@
+// Copyright Epic Games, Inc. All Rights Reserved.
+
+#include "ChaosClothAsset/SimulationMaxDistanceConfigNode.h"
+#include "ChaosClothAsset/DataflowNodes.h"
+#include "ChaosClothAsset/SimulationBaseConfigNodePrivate.h"
+#include "Chaos/CollectionPropertyFacade.h"
+
+#include UE_INLINE_GENERATED_CPP_BY_NAME(SimulationMaxDistanceConfigNode)
+
+FChaosClothAssetSimulationMaxDistanceConfigNode::FChaosClothAssetSimulationMaxDistanceConfigNode(const Dataflow::FNodeParameters& InParam, FGuid InGuid)
+	: FChaosClothAssetSimulationBaseConfigNode(InParam, InGuid)
+{
+	RegisterCollectionConnections();
+}
+
+void FChaosClothAssetSimulationMaxDistanceConfigNode::AddProperties(::Chaos::Softs::FCollectionPropertyMutableFacade& Properties) const
+{
+	UE_CHAOS_CLOTHASSET_SIMULATIONCONFIG_SETPROPERTYWEIGHTED(MaxDistance);
+}
diff --git a/Engine/Plugins/Experimental/ChaosClothAssetEditor/Source/ChaosClothAssetDataflowNodes/Private/ChaosClothAsset/SimulationPBDAreaSpringConfigNode.cpp b/Engine/Plugins/Experimental/ChaosClothAssetEditor/Source/ChaosClothAssetDataflowNodes/Private/ChaosClothAsset/SimulationPBDAreaSpringConfigNode.cpp
new file mode 100644
--- /dev/null
+++ b/Engine/Plugins/Experimental/ChaosClothAssetEditor/Source/ChaosClothAssetDataflowNodes/Private/ChaosClothAsset/SimulationPBDAreaSpringConfigNode.cpp
@@ -0,0 +1,21 @@
+// Copyright Epic Games, Inc. All Rights Reserved.
+
+#include "ChaosClothAsset/SimulationPBDAreaSpringConfigNode.h"
+#include "ChaosClothAsset/DataflowNodes.h"
+#include "ChaosClothAsset/SimulationBaseConfigNodePrivate.h"
+#include "Chaos/CollectionPropertyFacade.h"
+
+#include UE_INLINE_GENERATED_CPP_BY_NAME(SimulationPBDAreaSpringConfigNode)
+
+FChaosClothAssetSimulationPBDAreaSpringConfigNode::FChaosClothAssetSimulationPBDAreaSpringConfigNode(const Dataflow::FNodeParameters& InParam, FGuid InGuid)
+	: FChaosClothAssetSimulationBaseConfigNode(InParam, InGuid)
+{
+	RegisterCollectionConnections();
+}
+
+void FChaosClothAssetSimulationPBDAreaSpringConfigNode::AddProperties(::Chaos::Softs::FCollectionPropertyMutableFacade& Properties) const
+{
+	UE_CHAOS_CLOTHASSET_SIMULATIONCONFIG_SETPROPERTYWEIGHTEDCHECKED1(
+		AreaSpringStiffness,
+		XPBDAreaSpringStiffness);  // Existing properties to warn against
+}
diff --git a/Engine/Plugins/Experimental/ChaosClothAssetEditor/Source/ChaosClothAssetDataflowNodes/Private/ChaosClothAsset/SimulationPBDBendingElementConfigNode.cpp b/Engine/Plugins/Experimental/ChaosClothAssetEditor/Source/ChaosClothAssetDataflowNodes/Private/ChaosClothAsset/SimulationPBDBendingElementConfigNode.cpp
new file mode 100644
--- /dev/null
+++ b/Engine/Plugins/Experimental/ChaosClothAssetEditor/Source/ChaosClothAssetDataflowNodes/Private/ChaosClothAsset/SimulationPBDBendingElementConfigNode.cpp
@@ -0,0 +1,32 @@
+// Copyright Epic Games, Inc. All Rights Reserved.
+
+#include "ChaosClothAsset/SimulationPBDBendingElementConfigNode.h"
+#include "ChaosClothAsset/DataflowNodes.h"
+#include "ChaosClothAsset/SimulationBaseConfigNodePrivate.h"
+#include "Chaos/CollectionPropertyFacade.h"
+
+#include UE_INLINE_GENERATED_CPP_BY_NAME(SimulationPBDBendingElementConfigNode)
+
+FChaosClothAssetSimulationPBDBendingElementConfigNode::FChaosClothAssetSimulationPBDBendingElementConfigNode(const Dataflow::FNodeParameters& InParam, FGuid InGuid)
+	: FChaosClothAssetSimulationBaseConfigNode(InParam, InGuid)
+{
+	RegisterCollectionConnections();
+}
+
+void FChaosClothAssetSimulationPBDBendingElementConfigNode::AddProperties(::Chaos::Softs::FCollectionPropertyMutableFacade& Properties) const
+{
+	UE_CHAOS_CLOTHASSET_SIMULATIONCONFIG_SETPROPERTYWEIGHTEDCHECKED4(
+		BendingElementStiffness,
+		BendingSpringStiffness,          // Existing properties to warn against
+		XPBDBendingSpringStiffness,      //
+		XPBDBendingElementStiffness,     //
+		XPBDAnisoBendingStiffnessWarp);  //
+	UE_CHAOS_CLOTHASSET_SIMULATIONCONFIG_SETPROPERTYCHECKED2(
+		BucklingRatio,
+		XPBDBucklingRatio,        // Existing properties to warn against
+		XPBDAnisoBucklingRatio);  //
+	UE_CHAOS_CLOTHASSET_SIMULATIONCONFIG_SETPROPERTYWEIGHTEDCHECKED2(
+		BucklingStiffness,
+		XPBDBucklingStiffness,            // Existing properties to warn against
+		XPBDAnisoBucklingStiffnessWarp);  //
+}
diff --git a/Engine/Plugins/Experimental/ChaosClothAssetEditor/Source/ChaosClothAssetDataflowNodes/Private/ChaosClothAsset/SimulationPBDBendingSpringConfigNode.cpp b/Engine/Plugins/Experimental/ChaosClothAssetEditor/Source/ChaosClothAssetDataflowNodes/Private/ChaosClothAsset/SimulationPBDBendingSpringConfigNode.cpp
new file mode 100644
--- /dev/null
+++ b/Engine/Plugins/Experimental/ChaosClothAssetEditor/Source/ChaosClothAssetDataflowNodes/Private/ChaosClothAsset/SimulationPBDBendingSpringConfigNode.cpp
@@ -0,0 +1,24 @@
+// Copyright Epic Games, Inc. All Rights Reserved.
+
+#include "ChaosClothAsset/SimulationPBDBendingSpringConfigNode.h"
+#include "ChaosClothAsset/DataflowNodes.h"
+#include "ChaosClothAsset/SimulationBaseConfigNodePrivate.h"
+#include "Chaos/CollectionPropertyFacade.h"
+
+#include UE_INLINE_GENERATED_CPP_BY_NAME(SimulationPBDBendingSpringConfigNode)
+
+FChaosClothAssetSimulationPBDBendingSpringConfigNode::FChaosClothAssetSimulationPBDBendingSpringConfigNode(const Dataflow::FNodeParameters& InParam, FGuid InGuid)
+	: FChaosClothAssetSimulationBaseConfigNode(InParam, InGuid)
+{
+	RegisterCollectionConnections();
+}
+
+void FChaosClothAssetSimulationPBDBendingSpringConfigNode::AddProperties(::Chaos::Softs::FCollectionPropertyMutableFacade& Properties) const
+{
+	UE_CHAOS_CLOTHASSET_SIMULATIONCONFIG_SETPROPERTYWEIGHTEDCHECKED4(
+		BendingSpringStiffness,
+		BendingElementStiffness,         // Existing properties to warn against
+		XPBDBendingSpringStiffness,      //
+		XPBDBendingElementStiffness,     //
+		XPBDAnisoBendingStiffnessWarp);  //
+}
diff --git a/Engine/Plugins/Experimental/ChaosClothAssetEditor/Source/ChaosClothAssetDataflowNodes/Private/ChaosClothAsset/SimulationPBDEdgeSpringConfigNode.cpp b/Engine/Plugins/Experimental/ChaosClothAssetEditor/Source/ChaosClothAssetDataflowNodes/Private/ChaosClothAsset/SimulationPBDEdgeSpringConfigNode.cpp
new file mode 100644
--- /dev/null
+++ b/Engine/Plugins/Experimental/ChaosClothAssetEditor/Source/ChaosClothAssetDataflowNodes/Private/ChaosClothAsset/SimulationPBDEdgeSpringConfigNode.cpp
@@ -0,0 +1,22 @@
+// Copyright Epic Games, Inc. All Rights Reserved.
+
+#include "ChaosClothAsset/SimulationPBDEdgeSpringConfigNode.h"
+#include "ChaosClothAsset/DataflowNodes.h"
+#include "ChaosClothAsset/SimulationBaseConfigNodePrivate.h"
+#include "Chaos/CollectionPropertyFacade.h"
+
+#include UE_INLINE_GENERATED_CPP_BY_NAME(SimulationPBDEdgeSpringConfigNode)
+
+FChaosClothAssetSimulationPBDEdgeSpringConfigNode::FChaosClothAssetSimulationPBDEdgeSpringConfigNode(const Dataflow::FNodeParameters& InParam, FGuid InGuid)
+	: FChaosClothAssetSimulationBaseConfigNode(InParam, InGuid)
+{
+	RegisterCollectionConnections();
+}
+
+void FChaosClothAssetSimulationPBDEdgeSpringConfigNode::AddProperties(::Chaos::Softs::FCollectionPropertyMutableFacade& Properties) const
+{
+	UE_CHAOS_CLOTHASSET_SIMULATIONCONFIG_SETPROPERTYWEIGHTEDCHECKED2(
+		EdgeSpringStiffness,
+		XPBDEdgeSpringStiffness,         // Existing properties to warn against
+		XPBDAnisoStretchStiffnessWarp);  //
+}
diff --git a/Engine/Plugins/Experimental/ChaosClothAssetEditor/Source/ChaosClothAssetDataflowNodes/Private/ChaosClothAsset/SimulationPressureConfigNode.cpp b/Engine/Plugins/Experimental/ChaosClothAssetEditor/Source/ChaosClothAssetDataflowNodes/Private/ChaosClothAsset/SimulationPressureConfigNode.cpp
new file mode 100644
--- /dev/null
+++ b/Engine/Plugins/Experimental/ChaosClothAssetEditor/Source/ChaosClothAssetDataflowNodes/Private/ChaosClothAsset/SimulationPressureConfigNode.cpp
@@ -0,0 +1,19 @@
+// Copyright Epic Games, Inc. All Rights Reserved.
+
+#include "ChaosClothAsset/SimulationPressureConfigNode.h"
+#include "ChaosClothAsset/DataflowNodes.h"
+#include "ChaosClothAsset/SimulationBaseConfigNodePrivate.h"
+#include "Chaos/CollectionPropertyFacade.h"
+
+#include UE_INLINE_GENERATED_CPP_BY_NAME(SimulationPressureConfigNode)
+
+FChaosClothAssetSimulationPressureConfigNode::FChaosClothAssetSimulationPressureConfigNode(const Dataflow::FNodeParameters& InParam, FGuid InGuid)
+	: FChaosClothAssetSimulationBaseConfigNode(InParam, InGuid)
+{
+	RegisterCollectionConnections();
+}
+
+void FChaosClothAssetSimulationPressureConfigNode::AddProperties(::Chaos::Softs::FCollectionPropertyMutableFacade& Properties) const
+{
+	UE_CHAOS_CLOTHASSET_SIMULATIONCONFIG_SETPROPERTYWEIGHTED(Pressure);
+}
diff --git a/Engine/Plugins/Experimental/ChaosClothAssetEditor/Source/ChaosClothAssetDataflowNodes/Private/ChaosClothAsset/SimulationSelfCollisionConfigNode.cpp b/Engine/Plugins/Experimental/ChaosClothAssetEditor/Source/ChaosClothAssetDataflowNodes/Private/ChaosClothAsset/SimulationSelfCollisionConfigNode.cpp
new file mode 100644
--- /dev/null
+++ b/Engine/Plugins/Experimental/ChaosClothAssetEditor/Source/ChaosClothAssetDataflowNodes/Private/ChaosClothAsset/SimulationSelfCollisionConfigNode.cpp
@@ -0,0 +1,22 @@
+// Copyright Epic Games, Inc. All Rights Reserved.
+
+#include "ChaosClothAsset/SimulationSelfCollisionConfigNode.h"
+#include "ChaosClothAsset/DataflowNodes.h"
+#include "ChaosClothAsset/SimulationBaseConfigNodePrivate.h"
+#include "Chaos/CollectionPropertyFacade.h"
+
+#include UE_INLINE_GENERATED_CPP_BY_NAME(SimulationSelfCollisionConfigNode)
+
+FChaosClothAssetSimulationSelfCollisionConfigNode::FChaosClothAssetSimulationSelfCollisionConfigNode(const Dataflow::FNodeParameters& InParam, FGuid InGuid)
+	: FChaosClothAssetSimulationBaseConfigNode(InParam, InGuid)
+{
+	RegisterCollectionConnections();
+}
+
+void FChaosClothAssetSimulationSelfCollisionConfigNode::AddProperties(::Chaos::Softs::FCollectionPropertyMutableFacade& Properties) const
+{
+	UE_CHAOS_CLOTHASSET_SIMULATIONCONFIG_SETPROPERTY(SelfCollisionThickness);
+	UE_CHAOS_CLOTHASSET_SIMULATIONCONFIG_SETPROPERTY(SelfCollisionStiffness);
+	UE_CHAOS_CLOTHASSET_SIMULATIONCONFIG_SETPROPERTY(SelfCollisionFriction);
+	UE_CHAOS_CLOTHASSET_SIMULATIONCONFIG_SETPROPERTYBOOL(UseSelfIntersections);
+}
diff --git a/Engine/Plugins/Experimental/ChaosClothAssetEditor/Source/ChaosClothAssetDataflowNodes/Private/ChaosClothAsset/SimulationSolverConfigNode.cpp b/Engine/Plugins/Experimental/ChaosClothAssetEditor/Source/ChaosClothAssetDataflowNodes/Private/ChaosClothAsset/SimulationSolverConfigNode.cpp
new file mode 100644
--- /dev/null
+++ b/Engine/Plugins/Experimental/ChaosClothAssetEditor/Source/ChaosClothAssetDataflowNodes/Private/ChaosClothAsset/SimulationSolverConfigNode.cpp
@@ -0,0 +1,21 @@
+// Copyright Epic Games, Inc. All Rights Reserved.
+
+#include "ChaosClothAsset/SimulationSolverConfigNode.h"
+#include "ChaosClothAsset/DataflowNodes.h"
+#include "ChaosClothAsset/SimulationBaseConfigNodePrivate.h"
+#include "Chaos/CollectionPropertyFacade.h"
+
+#include UE_INLINE_GENERATED_CPP_BY_NAME(SimulationSolverConfigNode)
+
+FChaosClothAssetSimulationSolverConfigNode::FChaosClothAssetSimulationSolverConfigNode(const Dataflow::FNodeParameters& InParam, FGuid InGuid)
+	: FChaosClothAssetSimulationBaseConfigNode(InParam, InGuid)
+{
+	RegisterCollectionConnections();
+}
+
+void FChaosClothAssetSimulationSolverConfigNode::AddProperties(::Chaos::Softs::FCollectionPropertyMutableFacade& Properties) const
+{
+	UE_CHAOS_CLOTHASSET_SIMULATIONCONFIG_SETPROPERTY(NumIterations);
+	UE_CHAOS_CLOTHASSET_SIMULATIONCONFIG_SETPROPERTY(MaxNumIterations);
+	UE_CHAOS_CLOTHASSET_SIMULATIONCONFIG_SETPROPERTY(NumSubsteps);
+}
diff --git a/Engine/Plugins/Experimental/ChaosClothAssetEditor/Source/ChaosClothAssetDataflowNodes/Private/ChaosClothAsset/SimulationVelocityScaleConfigNode.cpp b/Engine/Plugins/Experimental/ChaosClothAssetEditor/Source/ChaosClothAssetDataflowNodes/Private/ChaosClothAsset/SimulationVelocityScaleConfigNode.cpp
new file mode 100644
--- /dev/null
+++ b/Engine/Plugins/Experimental/ChaosClothAssetEditor/Source/ChaosClothAssetDataflowNodes/Private/ChaosClothAsset/SimulationVelocityScaleConfigNode.cpp
@@ -0,0 +1,21 @@
+// Copyright Epic Games, Inc. All Rights Reserved.
+
+#include "ChaosClothAsset/SimulationVelocityScaleConfigNode.h"
+#include "ChaosClothAsset/DataflowNodes.h"
+#include "ChaosClothAsset/SimulationBaseConfigNodePrivate.h"
+#include "Chaos/CollectionPropertyFacade.h"
+
+#include UE_INLINE_GENERATED_CPP_BY_NAME(SimulationVelocityScaleConfigNode)
+
+FChaosClothAssetSimulationVelocityScaleConfigNode::FChaosClothAssetSimulationVelocityScaleConfigNode(const Dataflow::FNodeParameters& InParam, FGuid InGuid)
+	: FChaosClothAssetSimulationBaseConfigNode(InParam, InGuid)
+{
+	RegisterCollectionConnections();
+}
+
+void FChaosClothAssetSimulationVelocityScaleConfigNode::AddProperties(::Chaos::Softs::FCollectionPropertyMutableFacade& Properties) const
+{
+	UE_CHAOS_CLOTHASSET_SIMULATIONCONFIG_SETPROPERTY(LinearVelocityScale);
+	UE_CHAOS_CLOTHASSET_SIMULATIONCONFIG_SETPROPERTY(AngularVelocityScale);
+	UE_CHAOS_CLOTHASSET_SIMULATIONCONFIG_SETPROPERTY(FictitiousAngularScale);
+}
diff --git a/Engine/Plugins/Experimental/ChaosClothAssetEditor/Source/ChaosClothAssetDataflowNodes/Private/ChaosClothAsset/SimulationXPBDAnisoBendingConfigNode.cpp b/Engine/Plugins/Experimental/ChaosClothAssetEditor/Source/ChaosClothAssetDataflowNodes/Private/ChaosClothAsset/SimulationXPBDAnisoBendingConfigNode.cpp
new file mode 100644
--- /dev/null
+++ b/Engine/Plugins/Experimental/ChaosClothAssetEditor/Source/ChaosClothAssetDataflowNodes/Private/ChaosClothAsset/SimulationXPBDAnisoBendingConfigNode.cpp
@@ -0,0 +1,38 @@
+// Copyright Epic Games, Inc. All Rights Reserved.
+
+#include "ChaosClothAsset/SimulationXPBDAnisoBendingConfigNode.h"
+#include "ChaosClothAsset/DataflowNodes.h"
+#include "ChaosClothAsset/SimulationBaseConfigNodePrivate.h"
+#include "Chaos/CollectionPropertyFacade.h"
+
+#include UE_INLINE_GENERATED_CPP_BY_NAME(SimulationXPBDAnisoBendingConfigNode)
+
+FChaosClothAssetSimulationXPBDAnisoBendingConfigNode::FChaosClothAssetSimulationXPBDAnisoBendingConfigNode(const Dataflow::FNodeParameters& InParam, FGuid InGuid)
+	: FChaosClothAssetSimulationBaseConfigNode(InParam, InGuid)
+{
+	RegisterCollectionConnections();
+}
+
+void FChaosClothAssetSimulationXPBDAnisoBendingConfigNode::AddProperties(::Chaos::Softs::FCollectionPropertyMutableFacade& Properties) const
+{
+	UE_CHAOS_CLOTHASSET_SIMULATIONCONFIG_SETPROPERTYWEIGHTEDCHECKED4(
+		XPBDAnisoBendingStiffnessWarp,
+		BendingSpringStiffness,        // Existing properties to warn against
+		BendingElementStiffness,       //
+		XPBDBendingSpringStiffness,    //
+		XPBDBendingElementStiffness);  //
+	UE_CHAOS_CLOTHASSET_SIMULATIONCONFIG_SETPROPERTYWEIGHTED(XPBDAnisoBendingStiffnessWeft);
+	UE_CHAOS_CLOTHASSET_SIMULATIONCONFIG_SETPROPERTYWEIGHTED(XPBDAnisoBendingStiffnessBias);
+	UE_CHAOS_CLOTHASSET_SIMULATIONCONFIG_SETPROPERTYWEIGHTEDCHECKED2(
+		XPBDAnisoBendingDamping,
+		XPBDBendingSpringDamping,    // Existing properties to warn against
+		XPBDBendingElementDamping);  //
+	UE_CHAOS_CLOTHASSET_SIMULATIONCONFIG_SETPROPERTYCHECKED2(
+		XPBDAnisoBucklingRatio,
+		BucklingRatio,       // Existing properties to warn against
+		XPBDBucklingRatio);  //
+	UE_CHAOS_CLOTHASSET_SIMULATIONCONFIG_SETPROPERTYWEIGHTEDCHECKED2(
+		XPBDAnisoBucklingStiffnessWarp,
+		BucklingStiffness,       // Existing properties to warn against
+		XPBDBucklingStiffness);  //
+}
diff --git a/Engine/Plugins/Experimental/ChaosClothAssetEditor/Source/ChaosClothAssetDataflowNodes/Private/ChaosClothAsset/SimulationXPBDAnisoStretchConfigNode.cpp b/Engine/Plugins/Experimental/ChaosClothAssetEditor/Source/ChaosClothAssetDataflowNodes/Private/ChaosClothAsset/SimulationXPBDAnisoStretchConfigNode.cpp
new file mode 100644
--- /dev/null
+++ b/Engine/Plugins/Experimental/ChaosClothAssetEditor/Source/ChaosClothAssetDataflowNodes/Private/ChaosClothAsset/SimulationXPBDAnisoStretchConfigNode.cpp
@@ -0,0 +1,30 @@
+// Copyright Epic Games, Inc. All Rights Reserved.
+
+#include "ChaosClothAsset/SimulationXPBDAnisoStretchConfigNode.h"
+#include "ChaosClothAsset/DataflowNodes.h"
+#include "ChaosClothAsset/SimulationBaseConfigNodePrivate.h"
+#include "Chaos/CollectionPropertyFacade.h"
+
+#include UE_INLINE_GENERATED_CPP_BY_NAME(SimulationXPBDAnisoStretchConfigNode)
+
+FChaosClothAssetSimulationXPBDAnisoStretchConfigNode::FChaosClothAssetSimulationXPBDAnisoStretchConfigNode(const Dataflow::FNodeParameters& InParam, FGuid InGuid)
+	: FChaosClothAssetSimulationBaseConfigNode(InParam, InGuid)
+{
+	RegisterCollectionConnections();
+}
+
+void FChaosClothAssetSimulationXPBDAnisoStretchConfigNode::AddProperties(::Chaos::Softs::FCollectionPropertyMutableFacade& Properties) const
+{
+	UE_CHAOS_CLOTHASSET_SIMULATIONCONFIG_SETPROPERTYBOOL(XPBDAnisoStretchUse3dRestLengths);
+	UE_CHAOS_CLOTHASSET_SIMULATIONCONFIG_SETPROPERTYWEIGHTEDCHECKED2(
+		XPBDAnisoStretchStiffnessWarp,
+		EdgeSpringStiffness,           // Existing properties to warn against
+		XPBDEdgeSpringStiffness);      //
+	UE_CHAOS_CLOTHASSET_SIMULATIONCONFIG_SETPROPERTYWEIGHTED(XPBDAnisoStretchStiffnessWeft);
+	UE_CHAOS_CLOTHASSET_SIMULATIONCONFIG_SETPROPERTYWEIGHTED(XPBDAnisoStretchStiffnessBias);
+	UE_CHAOS_CLOTHASSET_SIMULATIONCONFIG_SETPROPERTYWEIGHTEDCHECKED1(
+		XPBDAnisoStretchDamping,
+		XPBDEdgeSpringDamping);  // Existing properties to warn against
+	UE_CHAOS_CLOTHASSET_SIMULATIONCONFIG_SETPROPERTYWEIGHTED(XPBDAnisoStretchWarpScale);
+	UE_CHAOS_CLOTHASSET_SIMULATIONCONFIG_SETPROPERTYWEIGHTED(XPBDAnisoStretchWeftScale);
+}
diff --git a/Engine/Plugins/Experimental/ChaosClothAssetEditor/Source/ChaosClothAssetDataflowNodes/Private/ChaosClothAsset/SimulationXPBDAreaSpringConfigNode.cpp b/Engine/Plugins/Experimental/ChaosClothAssetEditor/Source/ChaosClothAssetDataflowNodes/Private/ChaosClothAsset/SimulationXPBDAreaSpringConfigNode.cpp
new file mode 100644
--- /dev/null
+++ b/Engine/Plugins/Experimental/ChaosClothAssetEditor/Source/ChaosClothAssetDataflowNodes/Private/ChaosClothAsset/SimulationXPBDAreaSpringConfigNode.cpp
@@ -0,0 +1,21 @@
+// Copyright Epic Games, Inc. All Rights Reserved.
+
+#include "ChaosClothAsset/SimulationXPBDAreaSpringConfigNode.h"
+#include "ChaosClothAsset/DataflowNodes.h"
+#include "ChaosClothAsset/SimulationBaseConfigNodePrivate.h"
+#include "Chaos/CollectionPropertyFacade.h"
+
+#include UE_INLINE_GENERATED_CPP_BY_NAME(SimulationXPBDAreaSpringConfigNode)
+
+FChaosClothAssetSimulationXPBDAreaSpringConfigNode::FChaosClothAssetSimulationXPBDAreaSpringConfigNode(const Dataflow::FNodeParameters& InParam, FGuid InGuid)
+	: FChaosClothAssetSimulationBaseConfigNode(InParam, InGuid)
+{
+	RegisterCollectionConnections();
+}
+
+void FChaosClothAssetSimulationXPBDAreaSpringConfigNode::AddProperties(::Chaos::Softs::FCollectionPropertyMutableFacade& Properties) const
+{
+	UE_CHAOS_CLOTHASSET_SIMULATIONCONFIG_SETPROPERTYWEIGHTEDCHECKED1(
+		XPBDAreaSpringStiffness,
+		AreaSpringStiffness);  // Existing properties to warn against
+}
diff --git a/Engine/Plugins/Experimental/ChaosClothAssetEditor/Source/ChaosClothAssetDataflowNodes/Private/ChaosClothAsset/SimulationXPBDBendingElementConfigNode.cpp b/Engine/Plugins/Experimental/ChaosClothAssetEditor/Source/ChaosClothAssetDataflowNodes/Private/ChaosClothAsset/SimulationXPBDBendingElementConfigNode.cpp
new file mode 100644
--- /dev/null
+++ b/Engine/Plugins/Experimental/ChaosClothAssetEditor/Source/ChaosClothAssetDataflowNodes/Private/ChaosClothAsset/SimulationXPBDBendingElementConfigNode.cpp
@@ -0,0 +1,36 @@
+// Copyright Epic Games, Inc. All Rights Reserved.
+
+#include "ChaosClothAsset/SimulationXPBDBendingElementConfigNode.h"
+#include "ChaosClothAsset/DataflowNodes.h"
+#include "ChaosClothAsset/SimulationBaseConfigNodePrivate.h"
+#include "Chaos/CollectionPropertyFacade.h"
+
+#include UE_INLINE_GENERATED_CPP_BY_NAME(SimulationXPBDBendingElementConfigNode)
+
+FChaosClothAssetSimulationXPBDBendingElementConfigNode::FChaosClothAssetSimulationXPBDBendingElementConfigNode(const Dataflow::FNodeParameters& InParam, FGuid InGuid)
+	: FChaosClothAssetSimulationBaseConfigNode(InParam, InGuid)
+{
+	RegisterCollectionConnections();
+}
+
+void FChaosClothAssetSimulationXPBDBendingElementConfigNode::AddProperties(::Chaos::Softs::FCollectionPropertyMutableFacade& Properties) const
+{
+	UE_CHAOS_CLOTHASSET_SIMULATIONCONFIG_SETPROPERTYWEIGHTEDCHECKED4(
+		XPBDBendingElementStiffness,
+		BendingSpringStiffness,          // Existing properties to warn against
+		BendingElementStiffness,         //
+		XPBDBendingSpringStiffness,      //
+		XPBDAnisoBendingStiffnessWarp);  //
+	UE_CHAOS_CLOTHASSET_SIMULATIONCONFIG_SETPROPERTYWEIGHTEDCHECKED2(
+		XPBDBendingElementDamping,
+		XPBDBendingSpringDamping,  // Existing properties to warn against
+		XPBDAnisoBendingDamping);  //
+	UE_CHAOS_CLOTHASSET_SIMULATIONCONFIG_SETPROPERTYCHECKED2(
+		XPBDBucklingRatio,
+		BucklingRatio,            // Existing properties to warn against
+		XPBDAnisoBucklingRatio);  //
+	UE_CHAOS_CLOTHASSET_SIMULATIONCONFIG_SETPROPERTYWEIGHTEDCHECKED2(
+		XPBDBucklingStiffness,
+		BucklingStiffness,                // Existing properties to warn against
+		XPBDAnisoBucklingStiffnessWarp);  //
+}
diff --git a/Engine/Plugins/Experimental/ChaosClothAssetEditor/Source/ChaosClothAssetDataflowNodes/Private/ChaosClothAsset/SimulationXPBDBendingSpringConfigNode.cpp b/Engine/Plugins/Experimental/ChaosClothAssetEditor/Source/ChaosClothAssetDataflowNodes/Private/ChaosClothAsset/SimulationXPBDBendingSpringConfigNode.cpp
new file mode 100644
--- /dev/null
+++ b/Engine/Plugins/Experimental/ChaosClothAssetEditor/Source/ChaosClothAssetDataflowNodes/Private/ChaosClothAsset/SimulationXPBDBendingSpringConfigNode.cpp
@@ -0,0 +1,28 @@
+// Copyright Epic Games, Inc. All Rights Reserved.
+
+#include "ChaosClothAsset/SimulationXPBDBendingSpringConfigNode.h"
+#include "ChaosClothAsset/DataflowNodes.h"
+#include "ChaosClothAsset/SimulationBaseConfigNodePrivate.h"
+#include "Chaos/CollectionPropertyFacade.h"
+
+#include UE_INLINE_GENERATED_CPP_BY_NAME(SimulationXPBDBendingSpringConfigNode)
+
+FChaosClothAssetSimulationXPBDBendingSpringConfigNode::FChaosClothAssetSimulationXPBDBendingSpringConfigNode(const Dataflow::FNodeParameters& InParam, FGuid InGuid)
+	: FChaosClothAssetSimulationBaseConfigNode(InParam, InGuid)
+{
+	RegisterCollectionConnections();
+}
+
+void FChaosClothAssetSimulationXPBDBendingSpringConfigNode::AddProperties(::Chaos::Softs::FCollectionPropertyMutableFacade& Properties) const
+{
+	UE_CHAOS_CLOTHASSET_SIMULATIONCONFIG_SETPROPERTYWEIGHTEDCHECKED4(
+		XPBDBendingSpringStiffness,
+		BendingSpringStiffness,          // Existing properties to warn against
+		BendingElementStiffness,         //
+		XPBDBendingElementStiffness,     //
+		XPBDAnisoBendingStiffnessWarp);  //
+	UE_CHAOS_CLOTHASSET_SIMULATIONCONFIG_SETPROPERTYWEIGHTEDCHECKED2(
+		XPBDBendingSpringDamping,
+		XPBDBendingElementDamping,  // Existing properties to warn against
+		XPBDAnisoBendingDamping);   //
+}
diff --git a/Engine/Plugins/Experimental/ChaosClothAssetEditor/Source/ChaosClothAssetDataflowNodes/Private/ChaosClothAsset/SimulationXPBDConfigNode.cpp b/Engine/Plugins/Experimental/ChaosClothAssetEditor/Source/ChaosClothAssetDataflowNodes/Private/ChaosClothAsset/SimulationXPBDConfigNode.cpp
deleted file mode 100644
diff --git a/Engine/Plugins/Experimental/ChaosClothAssetEditor/Source/ChaosClothAssetDataflowNodes/Private/ChaosClothAsset/SimulationXPBDEdgeSpringConfigNode.cpp b/Engine/Plugins/Experimental/ChaosClothAssetEditor/Source/ChaosClothAssetDataflowNodes/Private/ChaosClothAsset/SimulationXPBDEdgeSpringConfigNode.cpp
new file mode 100644
--- /dev/null
+++ b/Engine/Plugins/Experimental/ChaosClothAssetEditor/Source/ChaosClothAssetDataflowNodes/Private/ChaosClothAsset/SimulationXPBDEdgeSpringConfigNode.cpp
@@ -0,0 +1,25 @@
+// Copyright Epic Games, Inc. All Rights Reserved.
+
+#include "ChaosClothAsset/SimulationXPBDEdgeSpringConfigNode.h"
+#include "ChaosClothAsset/DataflowNodes.h"
+#include "ChaosClothAsset/SimulationBaseConfigNodePrivate.h"
+#include "Chaos/CollectionPropertyFacade.h"
+
+#include UE_INLINE_GENERATED_CPP_BY_NAME(SimulationXPBDEdgeSpringConfigNode)
+
+FChaosClothAssetSimulationXPBDEdgeSpringConfigNode::FChaosClothAssetSimulationXPBDEdgeSpringConfigNode(const Dataflow::FNodeParameters& InParam, FGuid InGuid)
+	: FChaosClothAssetSimulationBaseConfigNode(InParam, InGuid)
+{
+	RegisterCollectionConnections();
+}
+
+void FChaosClothAssetSimulationXPBDEdgeSpringConfigNode::AddProperties(::Chaos::Softs::FCollectionPropertyMutableFacade& Properties) const
+{
+	UE_CHAOS_CLOTHASSET_SIMULATIONCONFIG_SETPROPERTYWEIGHTEDCHECKED2(
+		XPBDEdgeSpringStiffness,
+		EdgeSpringStiffness,             // Existing properties to warn against
+		XPBDAnisoStretchStiffnessWarp);  //
+	UE_CHAOS_CLOTHASSET_SIMULATIONCONFIG_SETPROPERTYWEIGHTEDCHECKED1(
+		XPBDEdgeSpringDamping,
+		XPBDAnisoStretchDamping);  // Existing properties to warn against
+}
diff --git a/Engine/Plugins/Experimental/ChaosClothAssetEditor/Source/ChaosClothAssetDataflowNodes/Private/ChaosClothAsset/WeightedValueCustomization.cpp b/Engine/Plugins/Experimental/ChaosClothAssetEditor/Source/ChaosClothAssetDataflowNodes/Private/ChaosClothAsset/WeightedValueCustomization.cpp
new file mode 100644
--- /dev/null
+++ b/Engine/Plugins/Experimental/ChaosClothAssetEditor/Source/ChaosClothAssetDataflowNodes/Private/ChaosClothAsset/WeightedValueCustomization.cpp
@@ -0,0 +1,352 @@
+// Copyright Epic Games, Inc. All Rights Reserved.
+
+#include "ChaosClothAsset/WeightedValueCustomization.h"
+#include "Widgets/Input/SCheckBox.h"
+#include "Widgets/Input/SNumericEntryBox.h"
+#include "Widgets/Input/SEditableTextBox.h"
+#include "DetailLayoutBuilder.h"
+#include "DetailWidgetRow.h"
+#include "Editor.h"
+
+#define LOCTEXT_NAMESPACE "ChaosClothAssetWeightedValueCustomization"
+
+namespace UE::Chaos::ClothAsset
+{
+	TSharedRef<IPropertyTypeCustomization> FWeightedValueCustomization::MakeInstance()
+	{
+		return MakeShareable(new FWeightedValueCustomization);
+	}
+
+	FWeightedValueCustomization::FWeightedValueCustomization() = default;
+	
+	FWeightedValueCustomization::~FWeightedValueCustomization() = default;
+
+	void FWeightedValueCustomization::MakeHeaderRow(TSharedRef<class IPropertyHandle>& StructPropertyHandle, FDetailWidgetRow& Row)
+	{
+		TWeakPtr<IPropertyHandle> StructWeakHandlePtr = StructPropertyHandle;
+
+		TSharedPtr<SHorizontalBox> HorizontalBox;
+
+		Row.NameContent()
+			[
+				StructPropertyHandle->CreatePropertyNameWidget()
+			]
+		.ValueContent()
+			// Make enough space for each child handle
+			.MinDesiredWidth(125.f * SortedChildHandles.Num())
+			.MaxDesiredWidth(125.f * SortedChildHandles.Num())
+			[
+				SAssignNew(HorizontalBox, SHorizontalBox)
+				.IsEnabled(this, &FMathStructCustomization::IsValueEnabled, StructWeakHandlePtr)
+			];
+
+		for (int32 ChildIndex = 0; ChildIndex < SortedChildHandles.Num(); ++ChildIndex)
+		{
+			TSharedRef<IPropertyHandle> ChildHandle = SortedChildHandles[ChildIndex];
+
+			// Propagate metadata to child properties so that it's reflected in the nested, individual spin boxes
+			ChildHandle->SetInstanceMetaData(TEXT("UIMin"), StructPropertyHandle->GetMetaData(TEXT("UIMin")));
+			ChildHandle->SetInstanceMetaData(TEXT("UIMax"), StructPropertyHandle->GetMetaData(TEXT("UIMax")));
+			ChildHandle->SetInstanceMetaData(TEXT("SliderExponent"), StructPropertyHandle->GetMetaData(TEXT("SliderExponent")));
+			ChildHandle->SetInstanceMetaData(TEXT("Delta"), StructPropertyHandle->GetMetaData(TEXT("Delta")));
+			ChildHandle->SetInstanceMetaData(TEXT("LinearDeltaSensitivity"), StructPropertyHandle->GetMetaData(TEXT("LinearDeltaSensitivity")));
+			ChildHandle->SetInstanceMetaData(TEXT("ShiftMouseMovePixelPerDelta"), StructPropertyHandle->GetMetaData(TEXT("ShiftMouseMovePixelPerDelta")));
+			ChildHandle->SetInstanceMetaData(TEXT("SupportDynamicSliderMaxValue"), StructPropertyHandle->GetMetaData(TEXT("SupportDynamicSliderMaxValue")));
+			ChildHandle->SetInstanceMetaData(TEXT("SupportDynamicSliderMinValue"), StructPropertyHandle->GetMetaData(TEXT("SupportDynamicSliderMinValue")));
+			ChildHandle->SetInstanceMetaData(TEXT("ClampMin"), StructPropertyHandle->GetMetaData(TEXT("ClampMin")));
+			ChildHandle->SetInstanceMetaData(TEXT("ClampMax"), StructPropertyHandle->GetMetaData(TEXT("ClampMax")));
+
+			const bool bLastChild = SortedChildHandles.Num() - 1 == ChildIndex;
+
+			TSharedRef<SWidget> ChildWidget = MakeChildWidget(StructPropertyHandle, ChildHandle);
+			
+			if (ChildHandle->GetPropertyClass() == FBoolProperty::StaticClass())
+			{
+				HorizontalBox->AddSlot()
+					.Padding(FMargin(0.f, 2.f, bLastChild ? 0.f : 3.f, 2.f))
+					.AutoWidth()  // keep the check box slots small
+					[
+						ChildWidget
+					];
+			}
+			else
+			{
+				if (ChildHandle->GetPropertyClass() == FFloatProperty::StaticClass())
+				{
+					NumericEntryBoxWidgetList.Add(ChildWidget);
+				}
+
+				HorizontalBox->AddSlot()
+					.Padding(FMargin(0.f, 2.f, bLastChild ? 0.f : 3.f, 2.f))
+					[
+						ChildWidget
+					];
+			}
+		}
+	}
+
+	TSharedRef<SWidget> FWeightedValueCustomization::MakeChildWidget(
+		TSharedRef<IPropertyHandle>& StructurePropertyHandle,
+		TSharedRef<IPropertyHandle>& PropertyHandle)
+	{
+		const FFieldClass* PropertyClass = PropertyHandle->GetPropertyClass();
+
+		if (PropertyClass == FFloatProperty::StaticClass())
+		{
+			return MakeFloatWidget(StructurePropertyHandle, PropertyHandle);
+		}
+		if (PropertyClass == FBoolProperty::StaticClass())
+		{
+			TWeakPtr<IPropertyHandle> WeakHandlePtr = PropertyHandle;
+			return
+				SNew(SCheckBox)
+				.ToolTipText(
+					LOCTEXT(
+						"IsAnimatable",
+						"Whether the property can ever be updated/animated in real time.\n"
+						"This could make the simulation takes more CPU time, even more so if the weight maps needs updating."))
+				.Type(ESlateCheckBoxType::CheckBox)
+				.IsChecked_Lambda([WeakHandlePtr]()->ECheckBoxState
+					{
+						bool bValue = false;
+						WeakHandlePtr.Pin()->GetValue(bValue);
+						return bValue ? ECheckBoxState::Checked : ECheckBoxState::Unchecked;
+					})
+				.OnCheckStateChanged_Lambda([WeakHandlePtr](ECheckBoxState CheckBoxState)
+					{
+						WeakHandlePtr.Pin()->SetValue(CheckBoxState == ECheckBoxState::Checked, EPropertyValueSetFlags::DefaultFlags);
+					});
+		}
+		if (PropertyClass == FStrProperty::StaticClass())
+		{
+			TWeakPtr<IPropertyHandle> WeakHandlePtr = PropertyHandle;
+			return
+				SNew(SEditableTextBox)
+				.ToolTipText(LOCTEXT("WeightMap", "The name of the weight map for this property."))
+				.Text_Lambda([WeakHandlePtr]() -> FText
+					{
+						FString Text;
+						WeakHandlePtr.Pin()->GetValue(Text);
+						return FText::FromString(Text);
+
+					})
+				.OnTextCommitted_Lambda([WeakHandlePtr](const FText& Text, ETextCommit::Type)
+					{
+						WeakHandlePtr.Pin()->SetValue(Text.ToString(), EPropertyValueSetFlags::DefaultFlags);
+					})
+				.Font(IPropertyTypeCustomizationUtils::GetRegularFont());
+		}
+
+		checkf(false, TEXT("Unsupported property class for the Weighted Values customization."));
+		return SNullWidget::NullWidget;
+	}
+
+	TSharedRef<SWidget> FWeightedValueCustomization::MakeFloatWidget(
+		TSharedRef<IPropertyHandle>& StructurePropertyHandle,
+		TSharedRef<IPropertyHandle>& PropertyHandle)
+	{
+		TOptional<float> MinValue, MaxValue, SliderMinValue, SliderMaxValue;
+		float SliderExponent, Delta;
+		int32 ShiftMouseMovePixelPerDelta = 1;
+		bool SupportDynamicSliderMaxValue = false;
+		bool SupportDynamicSliderMinValue = false;
+
+		ExtractFloatMetadata(StructurePropertyHandle, MinValue, MaxValue, SliderMinValue, SliderMaxValue, SliderExponent, Delta, ShiftMouseMovePixelPerDelta, SupportDynamicSliderMaxValue, SupportDynamicSliderMinValue);
+
+		TWeakPtr<IPropertyHandle> WeakHandlePtr = PropertyHandle;
+
+		return SNew(SNumericEntryBox<float>)
+			.IsEnabled(this, &FMathStructCustomization::IsValueEnabled, WeakHandlePtr)
+			.EditableTextBoxStyle(&FCoreStyle::Get().GetWidgetStyle<FEditableTextBoxStyle>("NormalEditableTextBox"))
+			.Value_Lambda([WeakHandlePtr]()
+				{
+					float Value = 0.;
+					return (WeakHandlePtr.Pin()->GetValue(Value) == FPropertyAccess::Success) ?
+						TOptional<float>(Value) :
+						TOptional<float>();  // Value couldn't be accessed, return an unset value
+				})
+			.Font(IDetailLayoutBuilder::GetDetailFont())
+			.UndeterminedString(NSLOCTEXT("PropertyEditor", "MultipleValues", "Multiple Values"))
+			.OnValueCommitted_Lambda([WeakHandlePtr](float Value, ETextCommit::Type)
+				{
+					WeakHandlePtr.Pin()->SetValue(Value, EPropertyValueSetFlags::DefaultFlags);
+				})
+			.OnValueChanged_Lambda([this, WeakHandlePtr](float Value)
+				{
+					if (bIsUsingSlider)
+					{
+						WeakHandlePtr.Pin()->SetValue(Value, EPropertyValueSetFlags::InteractiveChange);
+					}
+				})
+			.OnBeginSliderMovement_Lambda([this]()
+				{
+					bIsUsingSlider = true;
+					GEditor->BeginTransaction(LOCTEXT("SetVectorProperty", "Set Vector Property"));
+				})
+			.OnEndSliderMovement_Lambda([this](float Value)
+				{
+					bIsUsingSlider = false;
+					GEditor->EndTransaction();
+				})
+			.LabelVAlign(VAlign_Center)
+			// Only allow spin on handles with one object.  Otherwise it is not clear what value to spin
+			.AllowSpin(PropertyHandle->GetNumOuterObjects() < 2)
+			.ShiftMouseMovePixelPerDelta(ShiftMouseMovePixelPerDelta)
+			.SupportDynamicSliderMaxValue(SupportDynamicSliderMaxValue)
+			.SupportDynamicSliderMinValue(SupportDynamicSliderMinValue)
+			.OnDynamicSliderMaxValueChanged(this, &FWeightedValueCustomization::OnDynamicSliderMaxValueChanged)
+			.OnDynamicSliderMinValueChanged(this, &FWeightedValueCustomization::OnDynamicSliderMinValueChanged)
+			.MinValue(MinValue)
+			.MaxValue(MaxValue)
+			.MinSliderValue(SliderMinValue)
+			.MaxSliderValue(SliderMaxValue)
+			.SliderExponent(SliderExponent)
+			.Delta(Delta)
+			.Label()
+			[
+				SNew(STextBlock)
+					.Font(IDetailLayoutBuilder::GetDetailFont())
+					.Text(FText::FromString(PropertyHandle->GetMetaData(TEXT("ChaosClothAssetShortName"))))  // Case specific metadata, uses ChaosCloth prefix to avoid conflict with future engine extensions
+			];
+	}
+
+	// The following code is just a plain copy of FMathStructCustomization which
+	// would need changes to be able to serve as a base class for this customization.
+	void FWeightedValueCustomization::ExtractFloatMetadata(TSharedRef<IPropertyHandle>& PropertyHandle, TOptional<float>& MinValue, TOptional<float>& MaxValue, TOptional<float>& SliderMinValue, TOptional<float>& SliderMaxValue, float& SliderExponent, float& Delta, int32& ShiftMouseMovePixelPerDelta, bool& SupportDynamicSliderMaxValue, bool& SupportDynamicSliderMinValue)
+	{
+		FProperty* Property = PropertyHandle->GetProperty();
+
+		const FString& MetaUIMinString = Property->GetMetaData(TEXT("UIMin"));
+		const FString& MetaUIMaxString = Property->GetMetaData(TEXT("UIMax"));
+		const FString& SliderExponentString = Property->GetMetaData(TEXT("SliderExponent"));
+		const FString& DeltaString = Property->GetMetaData(TEXT("Delta"));
+		const FString& ShiftMouseMovePixelPerDeltaString = Property->GetMetaData(TEXT("ShiftMouseMovePixelPerDelta"));
+		const FString& SupportDynamicSliderMaxValueString = Property->GetMetaData(TEXT("SupportDynamicSliderMaxValue"));
+		const FString& SupportDynamicSliderMinValueString = Property->GetMetaData(TEXT("SupportDynamicSliderMinValue"));
+		const FString& ClampMinString = Property->GetMetaData(TEXT("ClampMin"));
+		const FString& ClampMaxString = Property->GetMetaData(TEXT("ClampMax"));
+
+		// If no UIMin/Max was specified then use the clamp string
+		const FString& UIMinString = MetaUIMinString.Len() ? MetaUIMinString : ClampMinString;
+		const FString& UIMaxString = MetaUIMaxString.Len() ? MetaUIMaxString : ClampMaxString;
+
+		float ClampMin = TNumericLimits<float>::Lowest();
+		float ClampMax = TNumericLimits<float>::Max();
+
+		if (!ClampMinString.IsEmpty())
+		{
+			TTypeFromString<float>::FromString(ClampMin, *ClampMinString);
+		}
+
+		if (!ClampMaxString.IsEmpty())
+		{
+			TTypeFromString<float>::FromString(ClampMax, *ClampMaxString);
+		}
+
+		float UIMin = TNumericLimits<float>::Lowest();
+		float UIMax = TNumericLimits<float>::Max();
+		TTypeFromString<float>::FromString(UIMin, *UIMinString);
+		TTypeFromString<float>::FromString(UIMax, *UIMaxString);
+
+		SliderExponent = float(1);
+
+		if (SliderExponentString.Len())
+		{
+			TTypeFromString<float>::FromString(SliderExponent, *SliderExponentString);
+		}
+
+		Delta = float(0);
+
+		if (DeltaString.Len())
+		{
+			TTypeFromString<float>::FromString(Delta, *DeltaString);
+		}
+
+		ShiftMouseMovePixelPerDelta = 1;
+		if (ShiftMouseMovePixelPerDeltaString.Len())
+		{
+			TTypeFromString<int32>::FromString(ShiftMouseMovePixelPerDelta, *ShiftMouseMovePixelPerDeltaString);
+			//The value should be greater or equal to 1
+			// 1 is neutral since it is a multiplier of the mouse drag pixel
+			if (ShiftMouseMovePixelPerDelta < 1)
+			{
+				ShiftMouseMovePixelPerDelta = 1;
+			}
+		}
+
+		const float ActualUIMin = FMath::Max(UIMin, ClampMin);
+		const float ActualUIMax = FMath::Min(UIMax, ClampMax);
+
+		MinValue = ClampMinString.Len() ? ClampMin : TOptional<float>();
+		MaxValue = ClampMaxString.Len() ? ClampMax : TOptional<float>();
+		SliderMinValue = (UIMinString.Len()) ? ActualUIMin : TOptional<float>();
+		SliderMaxValue = (UIMaxString.Len()) ? ActualUIMax : TOptional<float>();
+
+		SupportDynamicSliderMaxValue = SupportDynamicSliderMaxValueString.Len() > 0 && SupportDynamicSliderMaxValueString.ToBool();
+		SupportDynamicSliderMinValue = SupportDynamicSliderMinValueString.Len() > 0 && SupportDynamicSliderMinValueString.ToBool();
+	}
+
+	void FWeightedValueCustomization::OnDynamicSliderMaxValueChanged(float NewMaxSliderValue, TWeakPtr<SWidget> InValueChangedSourceWidget, bool IsOriginator, bool UpdateOnlyIfHigher)
+	{
+		for (TWeakPtr<SWidget>& Widget : NumericEntryBoxWidgetList)
+		{
+			TSharedPtr<SNumericEntryBox<float>> NumericBox = StaticCastSharedPtr<SNumericEntryBox<float>>(Widget.Pin());
+
+			if (NumericBox.IsValid())
+			{
+				TSharedPtr<SSpinBox<float>> SpinBox = StaticCastSharedPtr<SSpinBox<float>>(NumericBox->GetSpinBox());
+
+				if (SpinBox.IsValid())
+				{
+					if (SpinBox != InValueChangedSourceWidget)
+					{
+						if ((NewMaxSliderValue > SpinBox->GetMaxSliderValue() && UpdateOnlyIfHigher) || !UpdateOnlyIfHigher)
+						{
+							// Make sure the max slider value is not a getter otherwise we will break the link!
+							verifySlow(!SpinBox->IsMaxSliderValueBound());
+							SpinBox->SetMaxSliderValue(NewMaxSliderValue);
+						}
+					}
+				}
+			}
+		}
+
+		if (IsOriginator)
+		{
+			OnNumericEntryBoxDynamicSliderMaxValueChanged.Broadcast((float)NewMaxSliderValue, InValueChangedSourceWidget, false, UpdateOnlyIfHigher);
+		}
+	}
+
+	void FWeightedValueCustomization::OnDynamicSliderMinValueChanged(float NewMinSliderValue, TWeakPtr<SWidget> InValueChangedSourceWidget, bool IsOriginator, bool UpdateOnlyIfLower)
+	{
+		for (TWeakPtr<SWidget>& Widget : NumericEntryBoxWidgetList)
+		{
+			TSharedPtr<SNumericEntryBox<float>> NumericBox = StaticCastSharedPtr<SNumericEntryBox<float>>(Widget.Pin());
+
+			if (NumericBox.IsValid())
+			{
+				TSharedPtr<SSpinBox<float>> SpinBox = StaticCastSharedPtr<SSpinBox<float>>(NumericBox->GetSpinBox());
+
+				if (SpinBox.IsValid())
+				{
+					if (SpinBox != InValueChangedSourceWidget)
+					{
+						if ((NewMinSliderValue < SpinBox->GetMinSliderValue() && UpdateOnlyIfLower) || !UpdateOnlyIfLower)
+						{
+							// Make sure the min slider value is not a getter otherwise we will break the link!
+							verifySlow(!SpinBox->IsMinSliderValueBound());
+							SpinBox->SetMinSliderValue(NewMinSliderValue);
+						}
+					}
+				}
+			}
+		}
+
+		if (IsOriginator)
+		{
+			OnNumericEntryBoxDynamicSliderMinValueChanged.Broadcast((float)NewMinSliderValue, InValueChangedSourceWidget, false, UpdateOnlyIfLower);
+		}
+	}
+}  // End namespace UE::Chaos::ClothAsset
+
+#undef LOCTEXT_NAMESPACE
diff --git a/Engine/Plugins/Experimental/ChaosClothAssetEditor/Source/ChaosClothAssetDataflowNodes/Private/ChaosClothAsset/WeightedValueCustomization.h b/Engine/Plugins/Experimental/ChaosClothAssetEditor/Source/ChaosClothAssetDataflowNodes/Private/ChaosClothAsset/WeightedValueCustomization.h
new file mode 100644
--- /dev/null
+++ b/Engine/Plugins/Experimental/ChaosClothAssetEditor/Source/ChaosClothAssetDataflowNodes/Private/ChaosClothAsset/WeightedValueCustomization.h
@@ -0,0 +1,45 @@
+// Copyright Epic Games, Inc. All Rights Reserved.
+
+#pragma once
+
+#include "Customizations/MathStructCustomizations.h"
+
+template <typename OptionalType> struct TOptional;
+
+namespace UE::Chaos::ClothAsset
+{
+	/**
+	 * Customization for all property that also uses weight maps.
+	 * Works like a FMathStructCustomization with the difference that
+	 * the structure member short names are displayed in the header's row.
+	 */
+	class FWeightedValueCustomization
+		: public FMathStructCustomization
+	{
+	public:
+		static TSharedRef<IPropertyTypeCustomization> MakeInstance();
+
+		FWeightedValueCustomization();
+		virtual ~FWeightedValueCustomization() override;
+
+	protected:
+		virtual void MakeHeaderRow(TSharedRef<class IPropertyHandle>& StructPropertyHandle, FDetailWidgetRow& Row) override;
+		virtual TSharedRef<SWidget> MakeChildWidget(TSharedRef<IPropertyHandle>& StructurePropertyHandle, TSharedRef<IPropertyHandle>& PropertyHandle) override;
+
+	private:
+		static void ExtractFloatMetadata(
+			TSharedRef<IPropertyHandle>& PropertyHandle,
+			TOptional<float>& MinValue,
+			TOptional<float>& MaxValue,
+			TOptional<float>& SliderMinValue,
+			TOptional<float>& SliderMaxValue,
+			float& SliderExponent,
+			float& Delta,
+			int32& ShiftMouseMovePixelPerDelta,
+			bool& SupportDynamicSliderMaxValue,
+			bool& SupportDynamicSliderMinValue);
+		TSharedRef<SWidget> MakeFloatWidget(TSharedRef<IPropertyHandle>& StructurePropertyHandle, TSharedRef<IPropertyHandle>& PropertyHandle);
+		void OnDynamicSliderMaxValueChanged(float NewMaxSliderValue, TWeakPtr<SWidget> InValueChangedSourceWidget, bool IsOriginator, bool UpdateOnlyIfHigher);
+		void OnDynamicSliderMinValueChanged(float NewMinSliderValue, TWeakPtr<SWidget> InValueChangedSourceWidget, bool IsOriginator, bool UpdateOnlyIfLower);
+	};
+}
diff --git a/Engine/Plugins/Experimental/ChaosClothAssetEditor/Source/ChaosClothAssetDataflowNodes/Public/ChaosClothAsset/DataflowNodes.h b/Engine/Plugins/Experimental/ChaosClothAssetEditor/Source/ChaosClothAssetDataflowNodes/Public/ChaosClothAsset/DataflowNodes.h
--- a/Engine/Plugins/Experimental/ChaosClothAssetEditor/Source/ChaosClothAssetDataflowNodes/Public/ChaosClothAsset/DataflowNodes.h
+++ b/Engine/Plugins/Experimental/ChaosClothAssetEditor/Source/ChaosClothAssetDataflowNodes/Public/ChaosClothAsset/DataflowNodes.h
@@ -10,7 +10,9 @@
 namespace UE::Chaos::ClothAsset::DataflowNodes
 {
 	void Register();
-	void LogAndToastWarning(const FText& Error);
+	void LogAndToastWarning(const FText& Message);
+
+	FString MakeAttributeName(const FString& String);
 }
 
 DECLARE_LOG_CATEGORY_EXTERN(LogChaosClothAssetDataflowNodes, Log, All);
diff --git a/Engine/Plugins/Experimental/ChaosClothAssetEditor/Source/ChaosClothAssetDataflowNodes/Public/ChaosClothAsset/SimulationAerodynamicsConfigNode.h b/Engine/Plugins/Experimental/ChaosClothAssetEditor/Source/ChaosClothAssetDataflowNodes/Public/ChaosClothAsset/SimulationAerodynamicsConfigNode.h
new file mode 100644
--- /dev/null
+++ b/Engine/Plugins/Experimental/ChaosClothAssetEditor/Source/ChaosClothAssetDataflowNodes/Public/ChaosClothAsset/SimulationAerodynamicsConfigNode.h
@@ -0,0 +1,47 @@
+// Copyright Epic Games, Inc. All Rights Reserved.
+
+#pragma once 
+
+#include "ChaosClothAsset/SimulationBaseConfigNode.h"
+#include "ChaosClothAsset/WeightedValue.h"
+#include "SimulationAerodynamicsConfigNode.generated.h"
+
+/** Aerodynamics properties configuration node. */
+USTRUCT(Meta = (DataflowCloth))
+struct FChaosClothAssetSimulationAerodynamicsConfigNode : public FChaosClothAssetSimulationBaseConfigNode
+{
+	GENERATED_USTRUCT_BODY()
+	DATAFLOW_NODE_DEFINE_INTERNAL(FChaosClothAssetSimulationAerodynamicsConfigNode, "SimulationAerodynamicsConfig", "Cloth", "Cloth Simulation Aerodynamics Config")
+
+public:
+	/**
+	 * The density of the medium in which the aerodynamic forces take place, usually air.
+	 * The fluid density is given in kg/m^3.
+	 * Air density is considered to be around 1.225 kg/m^3 in average atmospheric conditions.
+	 */
+	UPROPERTY(EditAnywhere, Category = "Aerodynamics", Meta = (UIMin = "0", ClampMin = "0", ClampMax = "10"))
+	float FluidDensity = 1.225f;
+
+	/**
+	 * The aerodynamic coefficient of drag applying on each particle.
+	 * If a valid weight map is found with the given Weight Map name, then both Low and High values
+	 * are interpolated with the per particle weight to make the final value used for the simulation.
+	 * Otherwise all particles are considered to have a zero weight, and only the Low value is meaningful.
+	 */
+	UPROPERTY(EditAnywhere, Category = "Aerodynamics", Meta = (UIMin = "0", UIMax = "1", ClampMin = "0", ClampMax = "10"))
+	FChaosClothAssetWeightedValue Drag = { true, 0.035f, 1.f, TEXT("Drag") };
+
+	/**
+	 * The aerodynamic coefficient of lift applying on each particle.
+	 * If a valid weight map is found with the given Weight Map name, then both Low and High values
+	 * are interpolated with the per particle weight to make the final value used for the simulation.
+	 * Otherwise all particles are considered to have a zero weight, and only the Low value is meaningful.
+	 */
+	UPROPERTY(EditAnywhere, Category = "Aerodynamics", Meta = (UIMin = "0", UIMax = "1", ClampMin = "0", ClampMax = "10"))
+	FChaosClothAssetWeightedValue Lift = { true, 0.035f, 1.f, TEXT("Lift") };
+
+	FChaosClothAssetSimulationAerodynamicsConfigNode(const Dataflow::FNodeParameters& InParam, FGuid InGuid = FGuid::NewGuid());
+
+private:
+	virtual void AddProperties(::Chaos::Softs::FCollectionPropertyMutableFacade& Properties) const override;
+};
diff --git a/Engine/Plugins/Experimental/ChaosClothAssetEditor/Source/ChaosClothAssetDataflowNodes/Public/ChaosClothAsset/SimulationAnimDriveConfigNode.h b/Engine/Plugins/Experimental/ChaosClothAssetEditor/Source/ChaosClothAssetDataflowNodes/Public/ChaosClothAsset/SimulationAnimDriveConfigNode.h
new file mode 100644
--- /dev/null
+++ b/Engine/Plugins/Experimental/ChaosClothAssetEditor/Source/ChaosClothAssetDataflowNodes/Public/ChaosClothAsset/SimulationAnimDriveConfigNode.h
@@ -0,0 +1,39 @@
+// Copyright Epic Games, Inc. All Rights Reserved.
+
+#pragma once 
+
+#include "ChaosClothAsset/SimulationBaseConfigNode.h"
+#include "ChaosClothAsset/WeightedValue.h"
+#include "SimulationAnimDriveConfigNode.generated.h"
+
+/** Anim drive properties configuration node. */
+USTRUCT(Meta = (DataflowCloth))
+struct FChaosClothAssetSimulationAnimDriveConfigNode : public FChaosClothAssetSimulationBaseConfigNode
+{
+	GENERATED_USTRUCT_BODY()
+	DATAFLOW_NODE_DEFINE_INTERNAL(FChaosClothAssetSimulationAnimDriveConfigNode, "SimulationAnimDriveConfig", "Cloth", "Cloth Simulation AnimDrive Config")
+
+public:
+	/**
+	 * The strength of the constraint driving the cloth towards the animated skinned/goal mesh.
+	 * If a valid weight map is found with the given Weight Map name, then both Low and High values
+	 * are interpolated with the per particle weight to make the final value used for the simulation.
+	 * Otherwise all particles are considered to have a zero weight, and only the Low value is meaningful.
+	 */
+	UPROPERTY(EditAnywhere, Category = "AnimDrive Properties", Meta = (UIMin = "0", UIMax = "1", ClampMin = "0", ClampMax = "1"))
+	FChaosClothAssetWeightedValue AnimDriveStiffness = { true, 0.f, 1.f, TEXT("AnimDriveStiffness") };
+
+	/**
+	 * The damping amount of the anim drive.
+	 * If a valid weight map is found with the given Weight Map name, then both Low and High values
+	 * are interpolated with the per particle weight to make the final value used for the simulation.
+	 * Otherwise all particles are considered to have a zero weight, and only the Low value is meaningful.
+	 */
+	UPROPERTY(EditAnywhere, Category = "AnimDrive Properties", Meta = (UIMin = "0", UIMax = "1", ClampMin = "0", ClampMax = "1"))
+	FChaosClothAssetWeightedValue AnimDriveDamping = { true, 0.f, 1.f, TEXT("AnimDriveDamping") };
+	
+	FChaosClothAssetSimulationAnimDriveConfigNode(const Dataflow::FNodeParameters& InParam, FGuid InGuid = FGuid::NewGuid());
+
+private:
+	virtual void AddProperties(::Chaos::Softs::FCollectionPropertyMutableFacade& Properties) const override;
+};
diff --git a/Engine/Plugins/Experimental/ChaosClothAssetEditor/Source/ChaosClothAssetDataflowNodes/Public/ChaosClothAsset/SimulationBackstopConfigNode.h b/Engine/Plugins/Experimental/ChaosClothAssetEditor/Source/ChaosClothAssetDataflowNodes/Public/ChaosClothAsset/SimulationBackstopConfigNode.h
new file mode 100644
--- /dev/null
+++ b/Engine/Plugins/Experimental/ChaosClothAssetEditor/Source/ChaosClothAssetDataflowNodes/Public/ChaosClothAsset/SimulationBackstopConfigNode.h
@@ -0,0 +1,40 @@
+// Copyright Epic Games, Inc. All Rights Reserved.
+
+#pragma once 
+
+#include "ChaosClothAsset/SimulationBaseConfigNode.h"
+#include "ChaosClothAsset/WeightedValue.h"
+#include "SimulationBackstopConfigNode.generated.h"
+
+/** Backstop properties configuration node. */
+USTRUCT(Meta = (DataflowCloth))
+struct FChaosClothAssetSimulationBackstopConfigNode : public FChaosClothAssetSimulationBaseConfigNode
+{
+	GENERATED_USTRUCT_BODY()
+	DATAFLOW_NODE_DEFINE_INTERNAL(FChaosClothAssetSimulationBackstopConfigNode, "SimulationBackstopConfig", "Cloth", "Cloth Simulation Backstop Config")
+
+public:
+	/**
+	 * The distance from the surface of a backstop collision sphere to its associated particle's skinned position along the mesh normal.
+	 * Can be positive or negative depending on the desired effect.
+	 * If a valid weight map is found with the given Weight Map name, then both Low and High values
+	 * are interpolated with the per particle weight to make the final value used for the simulation.
+	 * Otherwise all particles are considered to have a zero weight, and only the Low value is meaningful.
+	 */
+	UPROPERTY(EditAnywhere, Category = "Backstop Properties", Meta = (UIMin = "0", UIMax = "100"))
+	FChaosClothAssetWeightedValue BackstopDistance = { true, 0.f, 100.f, TEXT("BackstopDistance") };
+
+	/**
+	 * The radius of the backstop sphere that each particle collides with.
+	 * If a valid weight map is found with the given Weight Map name, then both Low and High values
+	 * are interpolated with the per particle weight to make the final value used for the simulation.
+	 * Otherwise all particles are considered to have a zero weight, and only the Low value is meaningful.
+	 */
+	UPROPERTY(EditAnywhere, Category = "Backstop Properties", Meta = (UIMin = "0", UIMax = "100", ClampMin = "0"))
+	FChaosClothAssetWeightedValue BackstopRadius = { true, 0.f, 100.f, TEXT("BackstopRadius") };
+
+	FChaosClothAssetSimulationBackstopConfigNode(const Dataflow::FNodeParameters& InParam, FGuid InGuid = FGuid::NewGuid());
+
+private:
+	virtual void AddProperties(::Chaos::Softs::FCollectionPropertyMutableFacade& Properties) const override;
+};
diff --git a/Engine/Plugins/Experimental/ChaosClothAssetEditor/Source/ChaosClothAssetDataflowNodes/Public/ChaosClothAsset/SimulationBaseConfigNode.h b/Engine/Plugins/Experimental/ChaosClothAssetEditor/Source/ChaosClothAssetDataflowNodes/Public/ChaosClothAsset/SimulationBaseConfigNode.h
new file mode 100644
--- /dev/null
+++ b/Engine/Plugins/Experimental/ChaosClothAssetEditor/Source/ChaosClothAssetDataflowNodes/Public/ChaosClothAsset/SimulationBaseConfigNode.h
@@ -0,0 +1,44 @@
+// Copyright Epic Games, Inc. All Rights Reserved.
+
+#pragma once 
+
+#include "Dataflow/DataflowNode.h"
+#include "GeometryCollection/ManagedArrayCollection.h"
+#include "SimulationBaseConfigNode.generated.h"
+
+namespace Chaos::Softs
+{
+	class FCollectionPropertyMutableFacade;
+}
+
+/**
+ * Base abstract class for all cloth asset config nodes.
+ * Inherited class must call RegisterCollectionConnections() in constructor to use this base class Collection.
+ */
+USTRUCT(meta = (Abstract))
+struct FChaosClothAssetSimulationBaseConfigNode : public FDataflowNode
+{
+	GENERATED_USTRUCT_BODY()
+
+public:
+	UPROPERTY(Meta = (Dataflowinput, DataflowOutput, DisplayName = "Collection", DataflowPassthrough = "Collection"))
+	FManagedArrayCollection Collection;
+
+	FChaosClothAssetSimulationBaseConfigNode() = default;
+
+	FChaosClothAssetSimulationBaseConfigNode(const Dataflow::FNodeParameters& InParam, FGuid InGuid = FGuid::NewGuid());
+
+	virtual void Evaluate(Dataflow::FContext& Context, const FDataflowOutput* Out) const override;
+
+protected:
+	virtual void AddProperties(::Chaos::Softs::FCollectionPropertyMutableFacade& Properties) const
+	PURE_VIRTUAL(FChaosClothAssetSimulationBaseConfigNode::AddProperties, );
+
+	void RegisterCollectionConnections();
+
+	int32 AddPropertyHelper(
+		::Chaos::Softs::FCollectionPropertyMutableFacade& Properties,
+		const FName& PropertyName,
+		bool bIsAnimatable = true,
+		const TArray<FName>& SimilarPropertyNames = TArray<FName>()) const;
+};
diff --git a/Engine/Plugins/Experimental/ChaosClothAssetEditor/Source/ChaosClothAssetDataflowNodes/Public/ChaosClothAsset/SimulationCollisionConfigNode.h b/Engine/Plugins/Experimental/ChaosClothAssetEditor/Source/ChaosClothAssetDataflowNodes/Public/ChaosClothAsset/SimulationCollisionConfigNode.h
new file mode 100644
--- /dev/null
+++ b/Engine/Plugins/Experimental/ChaosClothAssetEditor/Source/ChaosClothAssetDataflowNodes/Public/ChaosClothAsset/SimulationCollisionConfigNode.h
@@ -0,0 +1,35 @@
+// Copyright Epic Games, Inc. All Rights Reserved.
+
+#pragma once 
+
+#include "ChaosClothAsset/SimulationBaseConfigNode.h"
+#include "SimulationCollisionConfigNode.generated.h"
+
+/** Physics mesh collision properties configuration node. */
+USTRUCT(Meta = (DataflowCloth))
+struct FChaosClothAssetSimulationCollisionConfigNode : public FChaosClothAssetSimulationBaseConfigNode
+{
+	GENERATED_USTRUCT_BODY()
+	DATAFLOW_NODE_DEFINE_INTERNAL(FChaosClothAssetSimulationCollisionConfigNode, "SimulationCollisionConfig", "Cloth", "Cloth Simulation Collision Config")
+
+public:
+	/** The added thickness of collision shapes. */
+	UPROPERTY(EditAnywhere, Category = "Collision Properties", Meta = (UIMin = "0", UIMax = "100", ClampMin = "0", ClampMax = "1000"))
+	float CollisionThickness = 1.0f;
+
+	/** Friction coefficient for cloth - collider interaction. */
+	UPROPERTY(EditAnywhere, Category = "Collision Properties", Meta = (UIMin = "0", UIMax = "1", ClampMin = "0", ClampMax = "10"))
+	float FrictionCoefficient = 0.8f;
+
+	/**
+	 * Use continuous collision detection (CCD) to prevent any missed collisions between fast moving particles and colliders.
+	 * This has a negative effect on performance compared to when resolving collision without using CCD.
+	 */
+	UPROPERTY(EditAnywhere, Category = "Collision Properties")
+	bool bUseCCD = false;
+
+	FChaosClothAssetSimulationCollisionConfigNode(const Dataflow::FNodeParameters& InParam, FGuid InGuid = FGuid::NewGuid());
+
+private:
+	virtual void AddProperties(::Chaos::Softs::FCollectionPropertyMutableFacade& Properties) const override;
+};
diff --git a/Engine/Plugins/Experimental/ChaosClothAssetEditor/Source/ChaosClothAssetDataflowNodes/Public/ChaosClothAsset/SimulationDampingConfigNode.h b/Engine/Plugins/Experimental/ChaosClothAssetEditor/Source/ChaosClothAssetDataflowNodes/Public/ChaosClothAsset/SimulationDampingConfigNode.h
new file mode 100644
--- /dev/null
+++ b/Engine/Plugins/Experimental/ChaosClothAssetEditor/Source/ChaosClothAssetDataflowNodes/Public/ChaosClothAsset/SimulationDampingConfigNode.h
@@ -0,0 +1,36 @@
+// Copyright Epic Games, Inc. All Rights Reserved.
+
+#pragma once 
+
+#include "ChaosClothAsset/SimulationBaseConfigNode.h"
+#include "SimulationDampingConfigNode.generated.h"
+
+/** Damping properties configuration node. */
+USTRUCT(Meta = (DataflowCloth))
+struct FChaosClothAssetSimulationDampingConfigNode : public FChaosClothAssetSimulationBaseConfigNode
+{
+	GENERATED_USTRUCT_BODY()
+	DATAFLOW_NODE_DEFINE_INTERNAL(FChaosClothAssetSimulationDampingConfigNode, "SimulationDampingConfig", "Cloth", "Cloth Simulation Damping Config")
+
+public:
+	/**
+	 * The amount of global damping applied to the cloth velocities, also known as point damping.
+	 * Point damping improves simulation stability, but can also cause an overall slow-down effect and therefore is best left to very small percentage amounts.
+	 */
+	UPROPERTY(EditAnywhere, Category = "Damping Properties", Meta = (UIMin = "0", UIMax = "1", ClampMin = "0", ClampMax = "1"))
+	float DampingCoefficient = 0.01f;
+
+	/**
+	 * The amount of local damping applied to the cloth velocities.
+	 * This type of damping only damps individual deviations of the particles velocities from the global motion.
+	 * It makes the cloth deformations more cohesive and reduces jitter without affecting the overall movement.
+	 * It can also produce synchronization artifacts where part of the cloth mesh are disconnected (which might well be desirable, or not), and is more expensive than global damping.
+	 */
+	UPROPERTY(EditAnywhere, Category = "Damping Properties", Meta = (UIMin = "0", UIMax = "1", ClampMin = "0", ClampMax = "1"))
+	float LocalDampingCoefficient = 0.f;
+
+	FChaosClothAssetSimulationDampingConfigNode(const Dataflow::FNodeParameters& InParam, FGuid InGuid = FGuid::NewGuid());
+
+private:
+	virtual void AddProperties(::Chaos::Softs::FCollectionPropertyMutableFacade& Properties) const override;
+};
diff --git a/Engine/Plugins/Experimental/ChaosClothAssetEditor/Source/ChaosClothAssetDataflowNodes/Public/ChaosClothAsset/SimulationGravityConfigNode.h b/Engine/Plugins/Experimental/ChaosClothAssetEditor/Source/ChaosClothAssetDataflowNodes/Public/ChaosClothAsset/SimulationGravityConfigNode.h
new file mode 100644
--- /dev/null
+++ b/Engine/Plugins/Experimental/ChaosClothAssetEditor/Source/ChaosClothAssetDataflowNodes/Public/ChaosClothAsset/SimulationGravityConfigNode.h
@@ -0,0 +1,32 @@
+// Copyright Epic Games, Inc. All Rights Reserved.
+
+#pragma once 
+
+#include "ChaosClothAsset/SimulationBaseConfigNode.h"
+#include "SimulationGravityConfigNode.generated.h"
+
+/** Gravity properties configuration node. */
+USTRUCT(Meta = (DataflowCloth))
+struct FChaosClothAssetSimulationGravityConfigNode : public FChaosClothAssetSimulationBaseConfigNode
+{
+	GENERATED_USTRUCT_BODY()
+	DATAFLOW_NODE_DEFINE_INTERNAL(FChaosClothAssetSimulationGravityConfigNode, "SimulationGravityConfig", "Cloth", "Cloth Simulation Gravity Config")
+
+public:
+	// Use the config gravity value instead of world gravity.
+	UPROPERTY(EditAnywhere, Category = "Gravity Properties", Meta = (InlineEditConditionToggle))
+	bool bUseGravityOverride = false;
+
+	// Scale factor applied to the world gravity and also to the clothing simulation interactor gravity. Does not affect the gravity if set using the override below.
+	UPROPERTY(EditAnywhere, Category = "Gravity Properties", Meta = (UIMin = "0", UIMax = "10", EditCondition = "!bUseGravityOverride"))
+	float GravityScale = 1.f;
+
+	// The gravitational acceleration vector [cm/s^2]
+	UPROPERTY(EditAnywhere, Category = "Gravity Properties", Meta = (UIMin = "0", UIMax = "10", EditCondition = "bUseGravityOverride"))
+	FVector3f GravityOverride = { 0.f, 0.f, -980.665f };  // TODO: Should we make this a S.I. unit?
+
+	FChaosClothAssetSimulationGravityConfigNode(const Dataflow::FNodeParameters& InParam, FGuid InGuid = FGuid::NewGuid());
+
+private:
+	virtual void AddProperties(::Chaos::Softs::FCollectionPropertyMutableFacade& Properties) const override;
+};
diff --git a/Engine/Plugins/Experimental/ChaosClothAssetEditor/Source/ChaosClothAssetDataflowNodes/Public/ChaosClothAsset/SimulationMassConfigNode.h b/Engine/Plugins/Experimental/ChaosClothAssetEditor/Source/ChaosClothAssetDataflowNodes/Public/ChaosClothAsset/SimulationMassConfigNode.h
new file mode 100644
--- /dev/null
+++ b/Engine/Plugins/Experimental/ChaosClothAssetEditor/Source/ChaosClothAssetDataflowNodes/Public/ChaosClothAsset/SimulationMassConfigNode.h
@@ -0,0 +1,51 @@
+// Copyright Epic Games, Inc. All Rights Reserved.
+
+#pragma once 
+
+#include "ChaosClothAsset/SimulationBaseConfigNode.h"
+#include "ClothConfig.h"  // For EClothMassMode
+#include "SimulationMassConfigNode.generated.h"
+
+/** Mass properties configuration node. */
+USTRUCT(Meta = (DataflowCloth))
+struct FChaosClothAssetSimulationMassConfigNode : public FChaosClothAssetSimulationBaseConfigNode
+{
+	GENERATED_USTRUCT_BODY()
+	DATAFLOW_NODE_DEFINE_INTERNAL(FChaosClothAssetSimulationMassConfigNode, "SimulationMassConfig", "Cloth", "Cloth Simulation Mass Config")
+
+public:
+	/**
+	 * How cloth particle mass is determined
+	 * -	Uniform Mass: Every particle's mass will be set to the value specified in the UniformMass setting. Mostly to avoid as it can cause some serious issues with irregular tessellations.
+	 * -	Total Mass: The total mass is distributed equally over all the particles. Useful when referencing a specific garment size and feel.
+	 * -	Density: A constant mass density is used. Density is usually the preferred way of setting mass since it allows matching real life materials' density values.
+	 */
+	UPROPERTY(EditAnywhere, Category = "Mass Properties")
+	EClothMassMode MassMode = EClothMassMode::Density;
+
+	/** The value used when the Mass Mode is set to Uniform Mass. */
+	UPROPERTY(EditAnywhere, Category = "Mass Properties", Meta = (UIMin = "0.000001", UIMax = "0.001", ClampMin = "0", EditCondition = "MassMode == EClothMassMode::UniformMass"))
+	float UniformMass = 0.00015f;
+
+	/** The value used when Mass Mode is set to TotalMass. */
+	UPROPERTY(EditAnywhere, Category = "Mass Properties", Meta = (UIMin = "0.001", UIMax = "10", ClampMin = "0", EditCondition = "MassMode == EClothMassMode::TotalMass"))
+	float TotalMass = 0.5f;
+
+	/**
+	 * Density in kg/m^2.
+	 * Melton Wool: 0.7
+	 * Heavy leather: 0.6
+	 * Polyurethane: 0.5
+	 * Denim: 0.4
+	 * Light leather: 0.3
+	 * Cotton: 0.2
+	 * Silk: 0.1
+	 */
+	UPROPERTY(EditAnywhere, Category = "Mass Properties", meta = (UIMin = "0.001", UIMax = "1", ClampMin = "0", EditCondition = "MassMode == EClothMassMode::Density"))
+	float Density = 0.35f;
+
+	FChaosClothAssetSimulationMassConfigNode(const Dataflow::FNodeParameters& InParam, FGuid InGuid = FGuid::NewGuid());
+
+private:
+	virtual void AddProperties(::Chaos::Softs::FCollectionPropertyMutableFacade& Properties) const override;
+};
diff --git a/Engine/Plugins/Experimental/ChaosClothAssetEditor/Source/ChaosClothAssetDataflowNodes/Public/ChaosClothAsset/SimulationMaxDistanceConfigNode.h b/Engine/Plugins/Experimental/ChaosClothAssetEditor/Source/ChaosClothAssetDataflowNodes/Public/ChaosClothAsset/SimulationMaxDistanceConfigNode.h
new file mode 100644
--- /dev/null
+++ b/Engine/Plugins/Experimental/ChaosClothAssetEditor/Source/ChaosClothAssetDataflowNodes/Public/ChaosClothAsset/SimulationMaxDistanceConfigNode.h
@@ -0,0 +1,31 @@
+// Copyright Epic Games, Inc. All Rights Reserved.
+
+#pragma once 
+
+#include "ChaosClothAsset/SimulationBaseConfigNode.h"
+#include "ChaosClothAsset/WeightedValue.h"
+#include "SimulationMaxDistanceConfigNode.generated.h"
+
+/** Maximum distance constraint property configuration node. */
+USTRUCT(Meta = (DataflowCloth))
+struct FChaosClothAssetSimulationMaxDistanceConfigNode : public FChaosClothAssetSimulationBaseConfigNode
+{
+	GENERATED_USTRUCT_BODY()
+	DATAFLOW_NODE_DEFINE_INTERNAL(FChaosClothAssetSimulationMaxDistanceConfigNode, "SimulationMaxDistanceConfig", "Cloth", "Cloth Simulation MaxDistance Config")
+
+public:
+	/**
+	 * The maximum distance a simulated particle can reach from its animated skinned cloth mesh position.
+	 * If a particle has 0 for its maximum distance, it is no longer considered dynamic, and becomes kinematic to follow its animated position.
+	 * If a valid weight map is found with the given Weight Map name, then both Low and High values
+	 * are interpolated with the per particle weight to make the final value used for the simulation.
+	 * Otherwise all particles are considered to have a zero weight, and only the Low value is meaningful.
+	 */
+	UPROPERTY(EditAnywhere, Category = "MaxDistance Properties", Meta = (UIMin = "0", UIMax = "100", ClampMin = "0"))
+	FChaosClothAssetWeightedValue MaxDistance = { true, 0.f, 100.f, TEXT("MaxDistance") };
+
+	FChaosClothAssetSimulationMaxDistanceConfigNode(const Dataflow::FNodeParameters& InParam, FGuid InGuid = FGuid::NewGuid());
+
+private:
+	virtual void AddProperties(::Chaos::Softs::FCollectionPropertyMutableFacade& Properties) const override;
+};
diff --git a/Engine/Plugins/Experimental/ChaosClothAssetEditor/Source/ChaosClothAssetDataflowNodes/Public/ChaosClothAsset/SimulationPBDAreaSpringConfigNode.h b/Engine/Plugins/Experimental/ChaosClothAssetEditor/Source/ChaosClothAssetDataflowNodes/Public/ChaosClothAsset/SimulationPBDAreaSpringConfigNode.h
new file mode 100644
--- /dev/null
+++ b/Engine/Plugins/Experimental/ChaosClothAssetEditor/Source/ChaosClothAssetDataflowNodes/Public/ChaosClothAsset/SimulationPBDAreaSpringConfigNode.h
@@ -0,0 +1,31 @@
+// Copyright Epic Games, Inc. All Rights Reserved.
+
+#pragma once 
+
+#include "ChaosClothAsset/SimulationBaseConfigNode.h"
+#include "ChaosClothAsset/WeightedValue.h"
+#include "SimulationPBDAreaSpringConfigNode.generated.h"
+
+/** Area spring constraint property configuration node. */
+USTRUCT(Meta = (DataflowCloth))
+struct FChaosClothAssetSimulationPBDAreaSpringConfigNode : public FChaosClothAssetSimulationBaseConfigNode
+{
+	GENERATED_USTRUCT_BODY()
+	DATAFLOW_NODE_DEFINE_INTERNAL(FChaosClothAssetSimulationPBDAreaSpringConfigNode, "SimulationPBDAreaSpringConfig", "Cloth", "Cloth Simulation PBD Area Spring Config")
+
+public:
+	/**
+	 * The stiffness of the surface area preservation constraints.
+	 * Increase the solver iteration count for stiffer materials.
+	 * If a valid weight map is found with the given Weight Map name, then both Low and High values
+	 * are interpolated with the per particle weight to make the final value used for the simulation.
+	 * Otherwise all particles are considered to have a zero weight, and only the Low value is meaningful.
+	 */
+	UPROPERTY(EditAnywhere, Category = "PBDAreaSpring Properties", Meta = (UIMin = "0", UIMax = "1", ClampMin = "0", ClampMax = "1"))
+	FChaosClothAssetWeightedValue AreaSpringStiffness = { true, 1.f, 1.f, TEXT("AreaSpringStiffness") };
+
+	FChaosClothAssetSimulationPBDAreaSpringConfigNode(const Dataflow::FNodeParameters& InParam, FGuid InGuid = FGuid::NewGuid());
+
+private:
+	virtual void AddProperties(::Chaos::Softs::FCollectionPropertyMutableFacade& Properties) const override;
+};
diff --git a/Engine/Plugins/Experimental/ChaosClothAssetEditor/Source/ChaosClothAssetDataflowNodes/Public/ChaosClothAsset/SimulationPBDBendingElementConfigNode.h b/Engine/Plugins/Experimental/ChaosClothAssetEditor/Source/ChaosClothAssetDataflowNodes/Public/ChaosClothAsset/SimulationPBDBendingElementConfigNode.h
new file mode 100644
--- /dev/null
+++ b/Engine/Plugins/Experimental/ChaosClothAssetEditor/Source/ChaosClothAssetDataflowNodes/Public/ChaosClothAsset/SimulationPBDBendingElementConfigNode.h
@@ -0,0 +1,49 @@
+// Copyright Epic Games, Inc. All Rights Reserved.
+
+#pragma once 
+
+#include "ChaosClothAsset/SimulationBaseConfigNode.h"
+#include "ChaosClothAsset/WeightedValue.h"
+#include "SimulationPBDBendingElementConfigNode.generated.h"
+
+/** Bending element constraint property configuration node. */
+USTRUCT(Meta = (DataflowCloth))
+struct FChaosClothAssetSimulationPBDBendingElementConfigNode : public FChaosClothAssetSimulationBaseConfigNode
+{
+	GENERATED_USTRUCT_BODY()
+	DATAFLOW_NODE_DEFINE_INTERNAL(FChaosClothAssetSimulationPBDBendingElementConfigNode, "SimulationPBDBendingElementConfig", "Cloth", "Cloth Simulation PBD Bending Element Config")
+
+public:
+	/**
+	 * The Stiffness of the bending elements constraints. Increase the iteration count for stiffer materials.
+	 * If a valid weight map is found with the given Weight Map name, then both Low and High values
+	 * are interpolated with the per particle weight to make the final value used for the simulation.
+	 * Otherwise all particles are considered to have a zero weight, and only the Low value is meaningful.
+	 */
+	UPROPERTY(EditAnywhere, Category = "PBDBendingElement Properties", Meta = (UIMin = "0", UIMax = "1", ClampMin = "0", ClampMax = "1"))
+	FChaosClothAssetWeightedValue BendingElementStiffness = { true, 1.f, 1.f, TEXT("BendingElementStiffness") };
+
+	/**
+	 * Once the element has bent such that it's folded more than this ratio from its rest angle ("buckled"), switch to using Buckling Stiffness instead of BendingElement Stiffness.
+	 * When Buckling Ratio = 0, the Buckling Stiffness will never be used. When BucklingRatio = 1, the Buckling Stiffness will be used as soon as its bent past its rest configuration.
+	 */
+	UPROPERTY(EditAnywhere, Category = "PBDBendingElement Properties", Meta = (UIMin = "0", UIMax = "1", ClampMin = "0", ClampMax = "1"))
+	float BucklingRatio = 0.5f;
+
+	/**
+	 * The stiffness after bucking.
+	 * The constraint will use this stiffness instead of element Stiffness once the cloth has buckled, i.e., bent beyond a certain angle.
+	 * Typically, Buckling Stiffness is set to be less than BendingElement Stiffness.
+	 * Buckling Ratio determines the switch point between using BendingElement Stiffness and Buckling Stiffness.
+	 * If a valid weight map is found with the given Weight Map name, then both Low and High values
+	 * are interpolated with the per particle weight to make the final value used for the simulation.
+	 * Otherwise all particles are considered to have a zero weight, and only the Low value is meaningful.
+	 */
+	UPROPERTY(EditAnywhere, Category = "PBDBendingElement Properties", Meta = (UIMin = "0", UIMax = "1", ClampMin = "0", ClampMax = "1", EditCondition = "BucklingRatio != 0"))
+	FChaosClothAssetWeightedValue BucklingStiffness = { true, 0.9f, 0.9f, TEXT("BucklingStiffness") };
+
+	FChaosClothAssetSimulationPBDBendingElementConfigNode(const Dataflow::FNodeParameters& InParam, FGuid InGuid = FGuid::NewGuid());
+
+private:
+	virtual void AddProperties(::Chaos::Softs::FCollectionPropertyMutableFacade& Properties) const override;
+};
diff --git a/Engine/Plugins/Experimental/ChaosClothAssetEditor/Source/ChaosClothAssetDataflowNodes/Public/ChaosClothAsset/SimulationPBDBendingSpringConfigNode.h b/Engine/Plugins/Experimental/ChaosClothAssetEditor/Source/ChaosClothAssetDataflowNodes/Public/ChaosClothAsset/SimulationPBDBendingSpringConfigNode.h
new file mode 100644
--- /dev/null
+++ b/Engine/Plugins/Experimental/ChaosClothAssetEditor/Source/ChaosClothAssetDataflowNodes/Public/ChaosClothAsset/SimulationPBDBendingSpringConfigNode.h
@@ -0,0 +1,31 @@
+// Copyright Epic Games, Inc. All Rights Reserved.
+
+#pragma once 
+
+#include "ChaosClothAsset/SimulationBaseConfigNode.h"
+#include "ChaosClothAsset/WeightedValue.h"
+#include "SimulationPBDBendingSpringConfigNode.generated.h"
+
+/** Bending spring constraint property configuration node. */
+USTRUCT(Meta = (DataflowCloth))
+struct FChaosClothAssetSimulationPBDBendingSpringConfigNode : public FChaosClothAssetSimulationBaseConfigNode
+{
+	GENERATED_USTRUCT_BODY()
+	DATAFLOW_NODE_DEFINE_INTERNAL(FChaosClothAssetSimulationPBDBendingSpringConfigNode, "SimulationPBDBendingSpringConfig", "Cloth", "Cloth Simulation PBD Bending Spring Config")
+
+public:
+	/**
+	 * The Stiffness of cross segments and bending elements constraints
+	 * Increase the solver iteration count for stiffer materials.
+	 * If a valid weight map is found with the given Weight Map name, then both Low and High values
+	 * are interpolated with the per particle weight to make the final value used for the simulation.
+	 * Otherwise all particles are considered to have a zero weight, and only the Low value is meaningful.
+	 */
+	UPROPERTY(EditAnywhere, Category = "PBDBendingSpring Properties", Meta = (UIMin = "0", UIMax = "1", ClampMin = "0", ClampMax = "1"))
+	FChaosClothAssetWeightedValue BendingSpringStiffness = { true, 1.f, 1.f, TEXT("BendingSpringStiffness") };
+
+	FChaosClothAssetSimulationPBDBendingSpringConfigNode(const Dataflow::FNodeParameters& InParam, FGuid InGuid = FGuid::NewGuid());
+
+private:
+	virtual void AddProperties(::Chaos::Softs::FCollectionPropertyMutableFacade& Properties) const override;
+};
diff --git a/Engine/Plugins/Experimental/ChaosClothAssetEditor/Source/ChaosClothAssetDataflowNodes/Public/ChaosClothAsset/SimulationPBDEdgeSpringConfigNode.h b/Engine/Plugins/Experimental/ChaosClothAssetEditor/Source/ChaosClothAssetDataflowNodes/Public/ChaosClothAsset/SimulationPBDEdgeSpringConfigNode.h
new file mode 100644
--- /dev/null
+++ b/Engine/Plugins/Experimental/ChaosClothAssetEditor/Source/ChaosClothAssetDataflowNodes/Public/ChaosClothAsset/SimulationPBDEdgeSpringConfigNode.h
@@ -0,0 +1,31 @@
+// Copyright Epic Games, Inc. All Rights Reserved.
+
+#pragma once 
+
+#include "ChaosClothAsset/SimulationBaseConfigNode.h"
+#include "ChaosClothAsset/WeightedValue.h"
+#include "SimulationPBDEdgeSpringConfigNode.generated.h"
+
+/** Edge spring constraint property configuration node. */
+USTRUCT(Meta = (DataflowCloth))
+struct FChaosClothAssetSimulationPBDEdgeSpringConfigNode : public FChaosClothAssetSimulationBaseConfigNode
+{
+	GENERATED_USTRUCT_BODY()
+	DATAFLOW_NODE_DEFINE_INTERNAL(FChaosClothAssetSimulationPBDEdgeSpringConfigNode, "SimulationPBDEdgeSpringConfig", "Cloth", "Cloth Simulation PBD Edge Spring Config")
+
+public:
+	/**
+	 * The Stiffness of segments constraints. Increase the iteration count for stiffer materials.
+	 * Increase the solver iteration count for stiffer materials.
+	 * If a valid weight map is found with the given Weight Map name, then both Low and High values
+	 * are interpolated with the per particle weight to make the final value used for the simulation.
+	 * Otherwise all particles are considered to have a zero weight, and only the Low value is meaningful.
+	 */
+	UPROPERTY(EditAnywhere, Category = "PBDEdgeSpring Properties", Meta = (UIMin = "0", UIMax = "1", ClampMin = "0", ClampMax = "1"))
+	FChaosClothAssetWeightedValue EdgeSpringStiffness = { true, 1.f, 1.f, TEXT("EdgeSpringStiffness") };
+
+	FChaosClothAssetSimulationPBDEdgeSpringConfigNode(const Dataflow::FNodeParameters& InParam, FGuid InGuid = FGuid::NewGuid());
+
+private:
+	virtual void AddProperties(::Chaos::Softs::FCollectionPropertyMutableFacade& Properties) const override;
+};
diff --git a/Engine/Plugins/Experimental/ChaosClothAssetEditor/Source/ChaosClothAssetDataflowNodes/Public/ChaosClothAsset/SimulationPressureConfigNode.h b/Engine/Plugins/Experimental/ChaosClothAssetEditor/Source/ChaosClothAssetDataflowNodes/Public/ChaosClothAsset/SimulationPressureConfigNode.h
new file mode 100644
--- /dev/null
+++ b/Engine/Plugins/Experimental/ChaosClothAssetEditor/Source/ChaosClothAssetDataflowNodes/Public/ChaosClothAsset/SimulationPressureConfigNode.h
@@ -0,0 +1,30 @@
+// Copyright Epic Games, Inc. All Rights Reserved.
+
+#pragma once 
+
+#include "ChaosClothAsset/SimulationBaseConfigNode.h"
+#include "ChaosClothAsset/WeightedValue.h"
+#include "SimulationPressureConfigNode.generated.h"
+
+/** Pressure properties configuration node. */
+USTRUCT(Meta = (DataflowCloth))
+struct FChaosClothAssetSimulationPressureConfigNode : public FChaosClothAssetSimulationBaseConfigNode
+{
+	GENERATED_USTRUCT_BODY()
+	DATAFLOW_NODE_DEFINE_INTERNAL(FChaosClothAssetSimulationPressureConfigNode, "SimulationPressureConfig", "Cloth", "Cloth Simulation Pressure Config")
+
+public:
+	/**
+	 * Pressure force strength applied in the normal direction(use negative value to push toward backface)
+	 * If a valid weight map is found with the given Weight Map name, then both Low and High values
+	 * are interpolated with the per particle weight to make the final value used for the simulation.
+	 * Otherwise all particles are considered to have a zero weight, and only the Low value is meaningful.
+	 */
+	UPROPERTY(EditAnywhere, Category = "Pressure Properties", meta = (UIMin = "-10", UIMax = "10", ClampMin = "-100", ClampMax = "100"))
+	FChaosClothAssetWeightedValue Pressure = { true, 0.0f, 1.f, TEXT("Pressure") };
+
+	FChaosClothAssetSimulationPressureConfigNode(const Dataflow::FNodeParameters& InParam, FGuid InGuid = FGuid::NewGuid());
+
+private:
+	virtual void AddProperties(::Chaos::Softs::FCollectionPropertyMutableFacade& Properties) const override;
+};
diff --git a/Engine/Plugins/Experimental/ChaosClothAssetEditor/Source/ChaosClothAssetDataflowNodes/Public/ChaosClothAsset/SimulationSelfCollisionConfigNode.h b/Engine/Plugins/Experimental/ChaosClothAssetEditor/Source/ChaosClothAssetDataflowNodes/Public/ChaosClothAsset/SimulationSelfCollisionConfigNode.h
new file mode 100644
--- /dev/null
+++ b/Engine/Plugins/Experimental/ChaosClothAssetEditor/Source/ChaosClothAssetDataflowNodes/Public/ChaosClothAsset/SimulationSelfCollisionConfigNode.h
@@ -0,0 +1,36 @@
+// Copyright Epic Games, Inc. All Rights Reserved.
+
+#pragma once 
+
+#include "ChaosClothAsset/SimulationBaseConfigNode.h"
+#include "SimulationSelfCollisionConfigNode.generated.h"
+
+/** Self-collision repulsion forces (point-face) properties configuration node. */
+USTRUCT(Meta = (DataflowCloth))
+struct FChaosClothAssetSimulationSelfCollisionConfigNode : public FChaosClothAssetSimulationBaseConfigNode
+{
+	GENERATED_USTRUCT_BODY()
+	DATAFLOW_NODE_DEFINE_INTERNAL(FChaosClothAssetSimulationSelfCollisionConfigNode, "SimulationSelfCollisionConfig", "Cloth", "Cloth Simulation Self Collision Config")
+
+public:
+	/** The radius of the spheres used in self collision (i.e., offset per side. total thickness of cloth is 2x this value). */
+	UPROPERTY(EditAnywhere, Category = "Self-Collision Properties", meta = (UIMin = "0", UIMax = "100", ClampMin = "0", ClampMax = "1000"))
+	float SelfCollisionThickness = 0.5f;
+
+	/** The stiffness of the springs used to control self collision. */
+	UPROPERTY(EditAnywhere, Category = "Self-Collision Properties", meta = (UIMin = "0", UIMax = "100", ClampMin = "0", ClampMax = "1000"))
+	float SelfCollisionStiffness = 0.5f;
+
+	/** Friction coefficient for cloth - cloth interaction. */
+	UPROPERTY(EditAnywhere, Category = "Self-Collision Properties", meta = (UIMin = "0", UIMax = "1", ClampMin = "0", ClampMax = "10"))
+	float SelfCollisionFriction = 0.0f;
+
+	/** Enable self intersection resolution. This will try to fix any cloth intersections that are not handled by collision repulsions. */
+	UPROPERTY(EditAnywhere, Category = "Self-Collision Properties")
+	bool bUseSelfIntersections = false;
+
+	FChaosClothAssetSimulationSelfCollisionConfigNode(const Dataflow::FNodeParameters& InParam, FGuid InGuid = FGuid::NewGuid());
+
+private:
+	virtual void AddProperties(::Chaos::Softs::FCollectionPropertyMutableFacade& Properties) const override;
+};
diff --git a/Engine/Plugins/Experimental/ChaosClothAssetEditor/Source/ChaosClothAssetDataflowNodes/Public/ChaosClothAsset/SimulationSolverConfigNode.h b/Engine/Plugins/Experimental/ChaosClothAssetEditor/Source/ChaosClothAssetDataflowNodes/Public/ChaosClothAsset/SimulationSolverConfigNode.h
new file mode 100644
--- /dev/null
+++ b/Engine/Plugins/Experimental/ChaosClothAssetEditor/Source/ChaosClothAssetDataflowNodes/Public/ChaosClothAsset/SimulationSolverConfigNode.h
@@ -0,0 +1,43 @@
+// Copyright Epic Games, Inc. All Rights Reserved.
+
+#pragma once 
+
+#include "ChaosClothAsset/SimulationBaseConfigNode.h"
+#include "SimulationSolverConfigNode.generated.h"
+
+/** Solver properties configuration node. */
+USTRUCT(Meta = (DataflowCloth))
+struct FChaosClothAssetSimulationSolverConfigNode : public FChaosClothAssetSimulationBaseConfigNode
+{
+	GENERATED_USTRUCT_BODY()
+	DATAFLOW_NODE_DEFINE_INTERNAL(FChaosClothAssetSimulationSolverConfigNode, "SimulationSolverConfig", "Cloth", "Cloth Simulation Solver Config")
+
+public:
+	/**
+	 * The number of time step dependent solver iterations. This sets iterations at 60fps.
+	 * NumIterations can never be bigger than MaxNumIterations.
+	 * At lower fps up to MaxNumIterations may be used instead. At higher fps as low as one single iteration might be used.
+	 * Higher number of iterations will increase the stiffness of all constraints and improve convergence, but will also increase the CPU cost of the simulation.
+	 */
+	UPROPERTY(EditAnywhere, Category = Simulation, meta = (UIMin = "1", UIMax = "10", ClampMin = "1", ClampMax = "100"))
+	int32 NumIterations = 1;
+
+	/**
+	 * The maximum number of solver iterations.
+	 * This is the upper limit of the number of iterations set in solver, when the frame rate is lower than 60fps.
+	 */
+	UPROPERTY(EditAnywhere, Category = Simulation, meta = (UIMin = "1", UIMax = "10", ClampMin = "1", ClampMax = "100"))
+	int32 MaxNumIterations = 6;
+
+	/**
+	 * The number of solver substeps.
+	 * This will increase the precision of the collision inputs and help with constraint resolutions but will increase the CPU cost.
+	 */
+	UPROPERTY(EditAnywhere, Category = Simulation, meta = (UIMin = "1", UIMax = "10", ClampMin = "1", ClampMax = "100"))
+	int32 NumSubsteps = 1;
+
+	FChaosClothAssetSimulationSolverConfigNode(const Dataflow::FNodeParameters& InParam, FGuid InGuid = FGuid::NewGuid());
+
+private:
+	virtual void AddProperties(::Chaos::Softs::FCollectionPropertyMutableFacade& Properties) const override;
+};
diff --git a/Engine/Plugins/Experimental/ChaosClothAssetEditor/Source/ChaosClothAssetDataflowNodes/Public/ChaosClothAsset/SimulationVelocityScaleConfigNode.h b/Engine/Plugins/Experimental/ChaosClothAssetEditor/Source/ChaosClothAssetDataflowNodes/Public/ChaosClothAsset/SimulationVelocityScaleConfigNode.h
new file mode 100644
--- /dev/null
+++ b/Engine/Plugins/Experimental/ChaosClothAssetEditor/Source/ChaosClothAssetDataflowNodes/Public/ChaosClothAsset/SimulationVelocityScaleConfigNode.h
@@ -0,0 +1,45 @@
+// Copyright Epic Games, Inc. All Rights Reserved.
+
+#pragma once 
+
+#include "ChaosClothAsset/SimulationBaseConfigNode.h"
+#include "ChaosClothAsset/WeightedValue.h"
+#include "SimulationVelocityScaleConfigNode.generated.h"
+
+/** Velocity scale properties configuration node. */
+USTRUCT(Meta = (DataflowCloth))
+struct FChaosClothAssetSimulationVelocityScaleConfigNode : public FChaosClothAssetSimulationBaseConfigNode
+{
+	GENERATED_USTRUCT_BODY()
+	DATAFLOW_NODE_DEFINE_INTERNAL(FChaosClothAssetSimulationVelocityScaleConfigNode, "SimulationVelocityScaleConfig", "Cloth", "Cloth Simulation Velocity Scale Config")
+
+public:
+	/**
+	 * The amount of linear velocities sent to the local cloth space from the reference bone
+	 * (the closest bone to the root on which the cloth section has been skinned, or the root itself if the cloth isn't skinned).
+	 */
+	UPROPERTY(EditAnywhere, Category = "Animation Properties", meta = (UIMin = "0", UIMax = "1", ClampMin = "0", ClampMax = "1"))
+	FVector3f LinearVelocityScale = { 0.75f, 0.75f, 0.75f };
+
+	/**
+	 * The amount of angular velocities sent to the local cloth space from the reference bone
+	 * (the closest bone to the root on which the cloth section has been skinned, or the root itself if the cloth isn't skinned).
+	 */
+	UPROPERTY(EditAnywhere, Category = "Animation Properties", meta = (UIMin = "0", UIMax = "1", ClampMin = "0", ClampMax = "1"))
+	float AngularVelocityScale = 0.75f;
+
+	/**
+	 * The portion of the angular velocity that is used to calculate the strength of all fictitious forces (e.g. centrifugal force).
+	 * This parameter is only having an effect on the portion of the reference bone's angular velocity that has been removed from the
+	 * simulation via the Angular Velocity Scale parameter. This means it has no effect when AngularVelocityScale is set to 1 in which
+	 * case the cloth is simulated with full world space angular velocities and subjected to the true physical world inertial forces.
+	 * Values range from 0 to 2, with 0 showing no centrifugal effect, 1 full centrifugal effect, and 2 an overdriven centrifugal effect.
+	 */
+	UPROPERTY(EditAnywhere, Category = "Animation Properties", meta = (UIMin = "0", UIMax = "1", ClampMin = "0", ClampMax = "2"))
+	float FictitiousAngularScale = 1.f;
+
+	FChaosClothAssetSimulationVelocityScaleConfigNode(const Dataflow::FNodeParameters& InParam, FGuid InGuid = FGuid::NewGuid());
+
+private:
+	virtual void AddProperties(::Chaos::Softs::FCollectionPropertyMutableFacade& Properties) const override;
+};
diff --git a/Engine/Plugins/Experimental/ChaosClothAssetEditor/Source/ChaosClothAssetDataflowNodes/Public/ChaosClothAsset/SimulationXPBDAnisoBendingConfigNode.h b/Engine/Plugins/Experimental/ChaosClothAssetEditor/Source/ChaosClothAssetDataflowNodes/Public/ChaosClothAsset/SimulationXPBDAnisoBendingConfigNode.h
new file mode 100644
--- /dev/null
+++ b/Engine/Plugins/Experimental/ChaosClothAssetEditor/Source/ChaosClothAssetDataflowNodes/Public/ChaosClothAsset/SimulationXPBDAnisoBendingConfigNode.h
@@ -0,0 +1,100 @@
+// Copyright Epic Games, Inc. All Rights Reserved.
+
+#pragma once 
+
+#include "ChaosClothAsset/SimulationBaseConfigNode.h"
+#include "ChaosClothAsset/WeightedValue.h"
+#include "SimulationXPBDAnisoBendingConfigNode.generated.h"
+
+/** XPBD anisotropic bending constraint property configuration node. */
+USTRUCT(Meta = (DataflowCloth))
+struct FChaosClothAssetSimulationXPBDAnisoBendingConfigNode : public FChaosClothAssetSimulationBaseConfigNode
+{
+	GENERATED_USTRUCT_BODY()
+	DATAFLOW_NODE_DEFINE_INTERNAL(FChaosClothAssetSimulationXPBDAnisoBendingConfigNode, "SimulationXPBDAnisoBendingConfig", "Cloth", "Cloth Simulation XPBD Anisotropic Bending Config")
+
+public:
+	/**
+	 * The stiffness of the bending constraints in the warp (vertical) direction.
+	 * If a valid weight map is found with the given Weight Map name, then both Low and High values
+	 * are interpolated with the per particle weight to make the final value used for the simulation.
+	 * Otherwise all particles are considered to have a zero weight, and only the Low value is meaningful.
+	 */
+	UPROPERTY(EditAnywhere, Category = "XPBDAnisoBending Properties", Meta = (UIMin = "0", UIMax = "10000", ClampMin = "0", ClampMax = "10000000"))
+	FChaosClothAssetWeightedValue XPBDAnisoBendingStiffnessWarp = { true, 100.f, 100.f, TEXT("XPBDAnisoBendingStiffnessWarp") };
+
+	/**
+	 * The stiffness of the bending constraints in the weft (horizontal) direction.
+	 * If a valid weight map is found with the given Weight Map name, then both Low and High values
+	 * are interpolated with the per particle weight to make the final value used for the simulation.
+	 * Otherwise all particles are considered to have a zero weight, and only the Low value is meaningful.
+	 */
+	UPROPERTY(EditAnywhere, Category = "XPBDAnisoBending Properties", Meta = (UIMin = "0", UIMax = "10000", ClampMin = "0", ClampMax = "10000000"))
+	FChaosClothAssetWeightedValue XPBDAnisoBendingStiffnessWeft = { true, 100.f, 100.f, TEXT("XPBDAnisoBendingStiffnessWeft") };
+
+	/**
+	 * The stiffness of the bending constraints in the bias (diagonal) direction.
+	 * If a valid weight map is found with the given Weight Map name, then both Low and High values
+	 * are interpolated with the per particle weight to make the final value used for the simulation.
+	 * Otherwise all particles are considered to have a zero weight, and only the Low value is meaningful.
+	 */
+	UPROPERTY(EditAnywhere, Category = "XPBDAnisoBending Properties", Meta = (UIMin = "0", UIMax = "10000", ClampMin = "0", ClampMax = "10000000"))
+	FChaosClothAssetWeightedValue XPBDAnisoBendingStiffnessBias = { true, 100.f, 100.f, TEXT("XPBDAnisoBendingStiffnessBias") };
+
+	/**
+	 * The damping of the bending constraints.
+	 * If a valid weight map is found with the given Weight Map name, then both Low and High values
+	 * are interpolated with the per particle weight to make the final value used for the simulation.
+	 * Otherwise all particles are considered to have a zero weight, and only the Low value is meaningful.
+	 */
+	UPROPERTY(EditAnywhere, Category = "XPBDAnisoBending Properties", Meta = (UIMin = "0", UIMax = "10", ClampMin = "0", ClampMax = "1000"))
+	FChaosClothAssetWeightedValue XPBDAnisoBendingDamping = { true, 1.f, 1.f, TEXT("XPBDAnisoBendingDamping") };
+
+	/**
+	 * Once the element has bent such that it's folded more than this ratio from its rest angle ("buckled"), switch to using Buckling Stiffness instead of BendingElement Stiffness.
+	 * When Buckling Ratio = 0, the Buckling Stiffness will never be used. When BucklingRatio = 1, the Buckling Stiffness will be used as soon as its bent past its rest configuration.
+	 */
+	UPROPERTY(EditAnywhere, Category = "XPBDAnisoBending Properties", Meta = (UIMin = "0", UIMax = "1", ClampMin = "0", ClampMax = "1"))
+	float XPBDAnisoBucklingRatio = 0.5f;
+
+	/**
+	 * The stiffness after bucking in the warp (vertical) direction.
+	 * The constraint will use this stiffness instead of element Stiffness once the cloth has buckled, i.e., bent beyond a certain angle.
+	 * Typically, Buckling Stiffness is set to be less than BendingElement Stiffness.
+	 * Buckling Ratio determines the switch point between using BendingElement Stiffness and Buckling Stiffness.
+	 * If a valid weight map is found with the given Weight Map name, then both Low and High values
+	 * are interpolated with the per particle weight to make the final value used for the simulation.
+	 * Otherwise all particles are considered to have a zero weight, and only the Low value is meaningful.
+	 */
+	UPROPERTY(EditAnywhere, Category = "XPBDBendingElement Properties", Meta = (UIMin = "0", UIMax = "10000", ClampMin = "0", ClampMax = "10000000", EditCondition = "XPBDBucklingRatio != 0"))
+	FChaosClothAssetWeightedValue XPBDAnisoBucklingStiffnessWarp = { true, 50.f, 50.f, TEXT("XPBDAnisoBucklingStiffnessWarp") };
+
+	/**
+	 * The stiffness after bucking in the weft (horizontal) direction.
+	 * The constraint will use this stiffness instead of element Stiffness once the cloth has buckled, i.e., bent beyond a certain angle.
+	 * Typically, Buckling Stiffness is set to be less than BendingElement Stiffness.
+	 * Buckling Ratio determines the switch point between using BendingElement Stiffness and Buckling Stiffness.
+	 * If a valid weight map is found with the given Weight Map name, then both Low and High values
+	 * are interpolated with the per particle weight to make the final value used for the simulation.
+	 * Otherwise all particles are considered to have a zero weight, and only the Low value is meaningful.
+	 */
+	UPROPERTY(EditAnywhere, Category = "XPBDBendingElement Properties", Meta = (UIMin = "0", UIMax = "10000", ClampMin = "0", ClampMax = "10000000", EditCondition = "XPBDBucklingRatio != 0"))
+	FChaosClothAssetWeightedValue XPBDAnisoBucklingStiffnessWeft = { true, 50.f, 50.f, TEXT("XPBDAnisoBucklingStiffnessWeft") };
+
+	/**
+	 * The stiffness after bucking in the bias (diagonal) direction.
+	 * The constraint will use this stiffness instead of element Stiffness once the cloth has buckled, i.e., bent beyond a certain angle.
+	 * Typically, Buckling Stiffness is set to be less than BendingElement Stiffness.
+	 * Buckling Ratio determines the switch point between using BendingElement Stiffness and Buckling Stiffness.
+	 * If a valid weight map is found with the given Weight Map name, then both Low and High values
+	 * are interpolated with the per particle weight to make the final value used for the simulation.
+	 * Otherwise all particles are considered to have a zero weight, and only the Low value is meaningful.
+	 */
+	UPROPERTY(EditAnywhere, Category = "XPBDBendingElement Properties", Meta = (UIMin = "0", UIMax = "10000", ClampMin = "0", ClampMax = "10000000", EditCondition = "XPBDBucklingRatio != 0"))
+	FChaosClothAssetWeightedValue XPBDAnisoBucklingStiffnessBias = { true, 50.f, 50.f, TEXT("XPBDAnisoBucklingStiffnessBias") };
+
+	FChaosClothAssetSimulationXPBDAnisoBendingConfigNode(const Dataflow::FNodeParameters& InParam, FGuid InGuid = FGuid::NewGuid());
+
+private:
+	virtual void AddProperties(::Chaos::Softs::FCollectionPropertyMutableFacade& Properties) const override;
+};
diff --git a/Engine/Plugins/Experimental/ChaosClothAssetEditor/Source/ChaosClothAssetDataflowNodes/Public/ChaosClothAsset/SimulationXPBDAnisoStretchConfigNode.h b/Engine/Plugins/Experimental/ChaosClothAssetEditor/Source/ChaosClothAssetDataflowNodes/Public/ChaosClothAsset/SimulationXPBDAnisoStretchConfigNode.h
new file mode 100644
--- /dev/null
+++ b/Engine/Plugins/Experimental/ChaosClothAssetEditor/Source/ChaosClothAssetDataflowNodes/Public/ChaosClothAsset/SimulationXPBDAnisoStretchConfigNode.h
@@ -0,0 +1,79 @@
+// Copyright Epic Games, Inc. All Rights Reserved.
+
+#pragma once 
+
+#include "ChaosClothAsset/SimulationBaseConfigNode.h"
+#include "ChaosClothAsset/WeightedValue.h"
+#include "SimulationXPBDAnisoStretchConfigNode.generated.h"
+
+/** XPBD anisotropic stretch constraint property configuration node. */
+USTRUCT(Meta = (DataflowCloth))
+struct FChaosClothAssetSimulationXPBDAnisoStretchConfigNode : public FChaosClothAssetSimulationBaseConfigNode
+{
+	GENERATED_USTRUCT_BODY()
+	DATAFLOW_NODE_DEFINE_INTERNAL(FChaosClothAssetSimulationXPBDAnisoStretchConfigNode, "SimulationXPBDAnisoStretchConfig", "Cloth", "Cloth Simulation XPBD Anisotropic Stretch Config")
+
+public:
+	/** Whether to use the 3D drapped space as rest lengths, or use the 2D pattern space instead. */
+	UPROPERTY(EditAnywhere, Category = "XPBDAnisoStretch Properties")
+	bool bXPBDAnisoStretchUse3dRestLengths = true;
+
+	/**
+	 * The stiffness of the stretch constraints in the warp (vertical) direction.
+	 * If a valid weight map is found with the given Weight Map name, then both Low and High values
+	 * are interpolated with the per particle weight to make the final value used for the simulation.
+	 * Otherwise all particles are considered to have a zero weight, and only the Low value is meaningful.
+	 */
+	UPROPERTY(EditAnywhere, Category = "XPBDAnisoStretch Properties", Meta = (UIMin = "0", UIMax = "10000", ClampMin = "0", ClampMax = "10000000"))
+	FChaosClothAssetWeightedValue XPBDAnisoStretchStiffnessWarp = { true, 100.f, 100.f, TEXT("XPBDAnisoStretchStiffnessWarp") };
+
+	/**
+	 * The stiffness of the stretch constraints in the weft (horizontal) direction.
+	 * If a valid weight map is found with the given Weight Map name, then both Low and High values
+	 * are interpolated with the per particle weight to make the final value used for the simulation.
+	 * Otherwise all particles are considered to have a zero weight, and only the Low value is meaningful.
+	 */
+	UPROPERTY(EditAnywhere, Category = "XPBDAnisoStretch Properties", Meta = (UIMin = "0", UIMax = "10000", ClampMin = "0", ClampMax = "10000000"))
+	FChaosClothAssetWeightedValue XPBDAnisoStretchStiffnessWeft = { true, 100.f, 100.f, TEXT("XPBDAnisoStretchStiffnessWeft") };
+
+	/**
+	 * The stiffness of the stretch constraints in the bias (diagonal) direction.
+	 * If a valid weight map is found with the given Weight Map name, then both Low and High values
+	 * are interpolated with the per particle weight to make the final value used for the simulation.
+	 * Otherwise all particles are considered to have a zero weight, and only the Low value is meaningful.
+	 */
+	UPROPERTY(EditAnywhere, Category = "XPBDAnisoStretch Properties", Meta = (UIMin = "0", UIMax = "10000", ClampMin = "0", ClampMax = "10000000"))
+	FChaosClothAssetWeightedValue XPBDAnisoStretchStiffnessBias = { true, 100.f, 100.f, TEXT("XPBDAnisoStretchStiffnessBias") };
+
+	/**
+	 * The damping of the stretch constraints.
+	 * If a valid weight map is found with the given Weight Map name, then both Low and High values
+	 * are interpolated with the per particle weight to make the final value used for the simulation.
+	 * Otherwise all particles are considered to have a zero weight, and only the Low value is meaningful.
+	 */
+	UPROPERTY(EditAnywhere, Category = "XPBDAnisoStretch Properties", Meta = (UIMin = "0", UIMax = "10", ClampMin = "0", ClampMax = "1000"))
+	FChaosClothAssetWeightedValue XPBDAnisoStretchDamping = { true, 1.f, 1.f, TEXT("XPBDAnisoStretchDamping") };
+
+	/**
+	 * The scale of the stretch constraints at rest in the warp (vertical) direction.
+	 * If a valid weight map is found with the given Weight Map name, then both Low and High values
+	 * are interpolated with the per particle weight to make the final value used for the simulation.
+	 * Otherwise all particles are considered to have a zero weight, and only the Low value is meaningful.
+	 */
+	UPROPERTY(EditAnywhere, Category = "XPBDAnisoStretch Properties", Meta = (UIMin = "0", UIMax = "10", ClampMin = "0", ClampMax = "10000000"))
+	FChaosClothAssetWeightedValue XPBDAnisoStretchWarpScale = { true, 1.f, 1.f, TEXT("XPBDAnisoStretchWarpScale") };
+
+	/**
+	 * The scale of the stretch constraints at rest in the weft (horizontal) direction.
+	 * If a valid weight map is found with the given Weight Map name, then both Low and High values
+	 * are interpolated with the per particle weight to make the final value used for the simulation.
+	 * Otherwise all particles are considered to have a zero weight, and only the Low value is meaningful.
+	 */
+	UPROPERTY(EditAnywhere, Category = "XPBDAnisoStretch Properties", Meta = (UIMin = "0", UIMax = "10", ClampMin = "0", ClampMax = "10000000"))
+	FChaosClothAssetWeightedValue XPBDAnisoStretchWeftScale = { true, 1.f, 1.f, TEXT("XPBDAnisoStretchWeftScale") };
+
+	FChaosClothAssetSimulationXPBDAnisoStretchConfigNode(const Dataflow::FNodeParameters& InParam, FGuid InGuid = FGuid::NewGuid());
+
+private:
+	virtual void AddProperties(::Chaos::Softs::FCollectionPropertyMutableFacade& Properties) const override;
+};
diff --git a/Engine/Plugins/Experimental/ChaosClothAssetEditor/Source/ChaosClothAssetDataflowNodes/Public/ChaosClothAsset/SimulationXPBDAreaSpringConfigNode.h b/Engine/Plugins/Experimental/ChaosClothAssetEditor/Source/ChaosClothAssetDataflowNodes/Public/ChaosClothAsset/SimulationXPBDAreaSpringConfigNode.h
new file mode 100644
--- /dev/null
+++ b/Engine/Plugins/Experimental/ChaosClothAssetEditor/Source/ChaosClothAssetDataflowNodes/Public/ChaosClothAsset/SimulationXPBDAreaSpringConfigNode.h
@@ -0,0 +1,31 @@
+// Copyright Epic Games, Inc. All Rights Reserved.
+
+#pragma once 
+
+#include "ChaosClothAsset/SimulationBaseConfigNode.h"
+#include "ChaosClothAsset/WeightedValue.h"
+#include "SimulationXPBDAreaSpringConfigNode.generated.h"
+
+/** XPBD area spring constraint property configuration node. */
+USTRUCT(Meta = (DataflowCloth))
+struct FChaosClothAssetSimulationXPBDAreaSpringConfigNode : public FChaosClothAssetSimulationBaseConfigNode
+{
+	GENERATED_USTRUCT_BODY()
+	DATAFLOW_NODE_DEFINE_INTERNAL(FChaosClothAssetSimulationXPBDAreaSpringConfigNode, "SimulationXPBDAreaSpringConfig", "Cloth", "Cloth Simulation XPBD Area Spring Config")
+
+public:
+	/**
+	 * The stiffness of the surface area preservation constraints.
+	 * Increase the solver iteration count for stiffer materials.
+	 * If a valid weight map is found with the given Weight Map name, then both Low and High values
+	 * are interpolated with the per particle weight to make the final value used for the simulation.
+	 * Otherwise all particles are considered to have a zero weight, and only the Low value is meaningful.
+	 */
+	UPROPERTY(EditAnywhere, Category = "XPBDAreaSpring Properties", Meta = (UIMin = "0", UIMax = "10000", ClampMin = "0", ClampMax = "10000000"))
+	FChaosClothAssetWeightedValue XPBDAreaSpringStiffness = { true, 100.f, 100.f, TEXT("XPBDAreaSpringStiffness") };
+
+	FChaosClothAssetSimulationXPBDAreaSpringConfigNode(const Dataflow::FNodeParameters& InParam, FGuid InGuid = FGuid::NewGuid());
+
+private:
+	virtual void AddProperties(::Chaos::Softs::FCollectionPropertyMutableFacade& Properties) const override;
+};
diff --git a/Engine/Plugins/Experimental/ChaosClothAssetEditor/Source/ChaosClothAssetDataflowNodes/Public/ChaosClothAsset/SimulationXPBDBendingElementConfigNode.h b/Engine/Plugins/Experimental/ChaosClothAssetEditor/Source/ChaosClothAssetDataflowNodes/Public/ChaosClothAsset/SimulationXPBDBendingElementConfigNode.h
new file mode 100644
--- /dev/null
+++ b/Engine/Plugins/Experimental/ChaosClothAssetEditor/Source/ChaosClothAssetDataflowNodes/Public/ChaosClothAsset/SimulationXPBDBendingElementConfigNode.h
@@ -0,0 +1,58 @@
+// Copyright Epic Games, Inc. All Rights Reserved.
+
+#pragma once 
+
+#include "ChaosClothAsset/SimulationBaseConfigNode.h"
+#include "ChaosClothAsset/WeightedValue.h"
+#include "SimulationXPBDBendingElementConfigNode.generated.h"
+
+/** XPBD bending element constraint property configuration node. */
+USTRUCT(Meta = (DataflowCloth))
+struct FChaosClothAssetSimulationXPBDBendingElementConfigNode : public FChaosClothAssetSimulationBaseConfigNode
+{
+	GENERATED_USTRUCT_BODY()
+	DATAFLOW_NODE_DEFINE_INTERNAL(FChaosClothAssetSimulationXPBDBendingElementConfigNode, "SimulationXPBDBendingElementConfig", "Cloth", "Cloth Simulation XPBD Bending Element Config")
+
+public:
+	/**
+	 * The stiffness of the bending element constraints.
+	 * If a valid weight map is found with the given Weight Map name, then both Low and High values
+	 * are interpolated with the per particle weight to make the final value used for the simulation.
+	 * Otherwise all particles are considered to have a zero weight, and only the Low value is meaningful.
+	 */
+	UPROPERTY(EditAnywhere, Category = "XPBDBendingElement Properties", Meta = (UIMin = "0", UIMax = "10000", ClampMin = "0", ClampMax = "10000000"))
+	FChaosClothAssetWeightedValue XPBDBendingElementStiffness = { true, 100.f, 100.f, TEXT("XPBDBendingElementStiffness") };
+
+	/**
+	 * The damping of the bending element constraints.
+	 * If a valid weight map is found with the given Weight Map name, then both Low and High values
+	 * are interpolated with the per particle weight to make the final value used for the simulation.
+	 * Otherwise all particles are considered to have a zero weight, and only the Low value is meaningful.
+	 */
+	UPROPERTY(EditAnywhere, Category = "XPBDBendingElement Properties", Meta = (UIMin = "0", UIMax = "10", ClampMin = "0", ClampMax = "1000"))
+	FChaosClothAssetWeightedValue XPBDBendingElementDamping = { true, 1.f, 1.f, TEXT("XPBDBendingElementDamping") };
+
+	/**
+	 * Once the element has bent such that it's folded more than this ratio from its rest angle ("buckled"), switch to using Buckling Stiffness instead of BendingElement Stiffness.
+	 * When Buckling Ratio = 0, the Buckling Stiffness will never be used. When BucklingRatio = 1, the Buckling Stiffness will be used as soon as its bent past its rest configuration.
+	 */
+	UPROPERTY(EditAnywhere, Category = "XPBDBendingElement Properties", Meta = (UIMin = "0", UIMax = "1", ClampMin = "0", ClampMax = "1"))
+	float XPBDBucklingRatio = 0.5f;
+
+	/**
+	 * The stiffness after bucking.
+	 * The constraint will use this stiffness instead of element Stiffness once the cloth has buckled, i.e., bent beyond a certain angle.
+	 * Typically, Buckling Stiffness is set to be less than BendingElement Stiffness.
+	 * Buckling Ratio determines the switch point between using BendingElement Stiffness and Buckling Stiffness.
+	 * If a valid weight map is found with the given Weight Map name, then both Low and High values
+	 * are interpolated with the per particle weight to make the final value used for the simulation.
+	 * Otherwise all particles are considered to have a zero weight, and only the Low value is meaningful.
+	 */
+	UPROPERTY(EditAnywhere, Category = "XPBDBendingElement Properties", Meta = (UIMin = "0", UIMax = "10000", ClampMin = "0", ClampMax = "10000000", EditCondition = "XPBDBucklingRatio != 0"))
+	FChaosClothAssetWeightedValue XPBDBucklingStiffness = { true, 50.f, 50.f, TEXT("XPBDBucklingStiffness") };
+
+	FChaosClothAssetSimulationXPBDBendingElementConfigNode(const Dataflow::FNodeParameters& InParam, FGuid InGuid = FGuid::NewGuid());
+
+private:
+	virtual void AddProperties(::Chaos::Softs::FCollectionPropertyMutableFacade& Properties) const override;
+};
diff --git a/Engine/Plugins/Experimental/ChaosClothAssetEditor/Source/ChaosClothAssetDataflowNodes/Public/ChaosClothAsset/SimulationXPBDBendingSpringConfigNode.h b/Engine/Plugins/Experimental/ChaosClothAssetEditor/Source/ChaosClothAssetDataflowNodes/Public/ChaosClothAsset/SimulationXPBDBendingSpringConfigNode.h
new file mode 100644
--- /dev/null
+++ b/Engine/Plugins/Experimental/ChaosClothAssetEditor/Source/ChaosClothAssetDataflowNodes/Public/ChaosClothAsset/SimulationXPBDBendingSpringConfigNode.h
@@ -0,0 +1,39 @@
+// Copyright Epic Games, Inc. All Rights Reserved.
+
+#pragma once 
+
+#include "ChaosClothAsset/SimulationBaseConfigNode.h"
+#include "ChaosClothAsset/WeightedValue.h"
+#include "SimulationXPBDBendingSpringConfigNode.generated.h"
+
+/** XPBD bending spring constraint property configuration node. */
+USTRUCT(Meta = (DataflowCloth))
+struct FChaosClothAssetSimulationXPBDBendingSpringConfigNode : public FChaosClothAssetSimulationBaseConfigNode
+{
+	GENERATED_USTRUCT_BODY()
+	DATAFLOW_NODE_DEFINE_INTERNAL(FChaosClothAssetSimulationXPBDBendingSpringConfigNode, "SimulationXPBDBendingSpringConfig", "Cloth", "Cloth Simulation XPBD Bending Spring Config")
+
+public:
+	/**
+	 * The stiffness of the bending cross segment constraints.
+	 * If a valid weight map is found with the given Weight Map name, then both Low and High values
+	 * are interpolated with the per particle weight to make the final value used for the simulation.
+	 * Otherwise all particles are considered to have a zero weight, and only the Low value is meaningful.
+	 */
+	UPROPERTY(EditAnywhere, Category = "XPBDBendingSpring Properties", Meta = (UIMin = "0", UIMax = "10000", ClampMin = "0", ClampMax = "10000000"))
+	FChaosClothAssetWeightedValue XPBDBendingSpringStiffness = { true, 100.f, 100.f, TEXT("XPBDBendingSpringStiffness") };
+
+	/**
+	 * The damping of the bending cross segment constraints.
+	 * If a valid weight map is found with the given Weight Map name, then both Low and High values
+	 * are interpolated with the per particle weight to make the final value used for the simulation.
+	 * Otherwise all particles are considered to have a zero weight, and only the Low value is meaningful.
+	 */
+	UPROPERTY(EditAnywhere, Category = "XPBDBendingSpring Properties", Meta = (UIMin = "0", UIMax = "10", ClampMin = "0", ClampMax = "1000"))
+	FChaosClothAssetWeightedValue XPBDBendingSpringDamping = { true, 1.f, 1.f, TEXT("XPBDBendingSpringDamping") };
+
+	FChaosClothAssetSimulationXPBDBendingSpringConfigNode(const Dataflow::FNodeParameters& InParam, FGuid InGuid = FGuid::NewGuid());
+
+private:
+	virtual void AddProperties(::Chaos::Softs::FCollectionPropertyMutableFacade& Properties) const override;
+};
diff --git a/Engine/Plugins/Experimental/ChaosClothAssetEditor/Source/ChaosClothAssetDataflowNodes/Public/ChaosClothAsset/SimulationXPBDConfigNode.h b/Engine/Plugins/Experimental/ChaosClothAssetEditor/Source/ChaosClothAssetDataflowNodes/Public/ChaosClothAsset/SimulationXPBDConfigNode.h
deleted file mode 100644
diff --git a/Engine/Plugins/Experimental/ChaosClothAssetEditor/Source/ChaosClothAssetDataflowNodes/Public/ChaosClothAsset/SimulationXPBDEdgeSpringConfigNode.h b/Engine/Plugins/Experimental/ChaosClothAssetEditor/Source/ChaosClothAssetDataflowNodes/Public/ChaosClothAsset/SimulationXPBDEdgeSpringConfigNode.h
new file mode 100644
--- /dev/null
+++ b/Engine/Plugins/Experimental/ChaosClothAssetEditor/Source/ChaosClothAssetDataflowNodes/Public/ChaosClothAsset/SimulationXPBDEdgeSpringConfigNode.h
@@ -0,0 +1,39 @@
+// Copyright Epic Games, Inc. All Rights Reserved.
+
+#pragma once 
+
+#include "ChaosClothAsset/SimulationBaseConfigNode.h"
+#include "ChaosClothAsset/WeightedValue.h"
+#include "SimulationXPBDEdgeSpringConfigNode.generated.h"
+
+/** XPBD edge spring constraint property configuration node. */
+USTRUCT(Meta = (DataflowCloth))
+struct FChaosClothAssetSimulationXPBDEdgeSpringConfigNode : public FChaosClothAssetSimulationBaseConfigNode
+{
+	GENERATED_USTRUCT_BODY()
+	DATAFLOW_NODE_DEFINE_INTERNAL(FChaosClothAssetSimulationXPBDEdgeSpringConfigNode, "SimulationXPBDEdgeSpringConfig", "Cloth", "Cloth Simulation XPBD Edge Spring Config")
+
+public:
+	/**
+	 * The stiffness of the edge segment constraints.
+	 * If a valid weight map is found with the given Weight Map name, then both Low and High values
+	 * are interpolated with the per particle weight to make the final value used for the simulation.
+	 * Otherwise all particles are considered to have a zero weight, and only the Low value is meaningful.
+	 */
+	UPROPERTY(EditAnywhere, Category = "XPBDEdgeSpring Properties", Meta = (UIMin = "0", UIMax = "10000", ClampMin = "0", ClampMax = "10000000"))
+	FChaosClothAssetWeightedValue XPBDEdgeSpringStiffness = { true, 100.f, 100.f, TEXT("XPBDEdgeSpringStiffness") };
+
+	/**
+	 * The damping of the edge segment constraints.
+	 * If a valid weight map is found with the given Weight Map name, then both Low and High values
+	 * are interpolated with the per particle weight to make the final value used for the simulation.
+	 * Otherwise all particles are considered to have a zero weight, and only the Low value is meaningful.
+	 */
+	UPROPERTY(EditAnywhere, Category = "XPBDEdgeSpring Properties", Meta = (UIMin = "0", UIMax = "10", ClampMin = "0", ClampMax = "1000"))
+	FChaosClothAssetWeightedValue XPBDEdgeSpringDamping = { true, 1.f, 1.f, TEXT("XPBDEdgeSpringDamping") };
+
+	FChaosClothAssetSimulationXPBDEdgeSpringConfigNode(const Dataflow::FNodeParameters& InParam, FGuid InGuid = FGuid::NewGuid());
+
+private:
+	virtual void AddProperties(::Chaos::Softs::FCollectionPropertyMutableFacade& Properties) const override;
+};
diff --git a/Engine/Plugins/Experimental/ChaosClothAssetEditor/Source/ChaosClothAssetDataflowNodes/Public/ChaosClothAsset/WeightedValue.h b/Engine/Plugins/Experimental/ChaosClothAssetEditor/Source/ChaosClothAssetDataflowNodes/Public/ChaosClothAsset/WeightedValue.h
new file mode 100644
--- /dev/null
+++ b/Engine/Plugins/Experimental/ChaosClothAssetEditor/Source/ChaosClothAssetDataflowNodes/Public/ChaosClothAsset/WeightedValue.h
@@ -0,0 +1,43 @@
+// Copyright Epic Games, Inc. All Rights Reserved.
+#pragma once
+
+#include "UObject/ObjectMacros.h"
+#include "Containers/UnrealString.h"
+#include "WeightedValue.generated.h"
+
+USTRUCT()
+struct FChaosClothAssetWeightedValue
+{
+	GENERATED_BODY()
+
+	/**
+	 * Whether the property can ever be updated/animated in real time.
+	 * This could make the simulation takes more CPU time, even more so if the weight maps needs updating.
+	 */
+	UPROPERTY(EditAnywhere, Category = "Weighted Value")
+	bool bIsAnimatable = true;
+
+	/**
+	 * Property value corresponding to the lower bound of the Weight Map.
+	 * A Weight Map stores a series of Weight values assigned to each point, all between 0 and 1.
+	 * The weights are used to interpolate the individual values from Low to High assigned to each different point.
+	 * A Weight of 0 always corresponds to the Low parameter value, and a Weight of 1 to the High parameter value.
+	 * The value for Low can be set to be bigger than for High in order to reverse the effect of the Weight Map.
+	 */
+	UPROPERTY(EditAnywhere, Category = "Weighted Value", Meta = (DisplayName = "Low Weight", ChaosClothAssetShortName = "Lo"))
+	float Low = 0.f;
+
+	/**
+	 * Property value corresponding to the upper bound of the Weight Map.
+	 * A Weight Map stores a series of Weight values assigned to each point, all between 0 and 1.
+	 * The weights are used to interpolate the individual values from Low to High assigned to each different point.
+	 * A Weight of 0 always corresponds to the Low parameter value, and a Weight of 1 to the High parameter value.
+	 * The value for Low can be set to be bigger than for High in order to reverse the effect of the Weight Map.
+	 */
+	UPROPERTY(EditAnywhere, Category = "Weighted Value", Meta = (DisplayName = "High Weight", ChaosClothAssetShortName = "Hi"))
+	float High = 1.f;
+
+	/** The name of the weight map for this property. */
+	UPROPERTY(EditAnywhere, Category = "Weighted Value")
+	FString WeightMap = TEXT("WeightMap");
+};
diff --git a/Engine/Source/Runtime/Experimental/Chaos/Private/Chaos/XPBDAnisotropicBendingConstraints.cpp b/Engine/Source/Runtime/Experimental/Chaos/Private/Chaos/XPBDAnisotropicBendingConstraints.cpp
--- a/Engine/Source/Runtime/Experimental/Chaos/Private/Chaos/XPBDAnisotropicBendingConstraints.cpp
+++ b/Engine/Source/Runtime/Experimental/Chaos/Private/Chaos/XPBDAnisotropicBendingConstraints.cpp
@@ -35,16 +35,16 @@ FXPBDAnisotropicBendingConstraints::FXPBDAnisotropicBendingConstraints(const FSo
 		InParticleOffset,
 		InParticleCount,
 		TriangleMesh.GetUniqueAdjacentElements(),
-		WeightMaps.FindRef(GetXPBDBendingElementStiffnessWarpString(PropertyCollection, XPBDBendingElementStiffnessWarpName.ToString())),
-		WeightMaps.FindRef(GetXPBDBucklingStiffnessWarpString(PropertyCollection, XPBDBucklingStiffnessWarpName.ToString())),
-		FSolverVec2(GetWeightedFloatXPBDBendingElementStiffnessWarp(PropertyCollection, MaxStiffness)),
-		(FSolverReal)GetXPBDBucklingRatio(PropertyCollection, 0.f),
-		FSolverVec2(GetWeightedFloatXPBDBucklingStiffnessWarp(PropertyCollection, MaxStiffness)),
+		WeightMaps.FindRef(GetXPBDAnisoBendingStiffnessWarpString(PropertyCollection, XPBDAnisoBendingStiffnessWarpName.ToString())),
+		WeightMaps.FindRef(GetXPBDAnisoBucklingStiffnessWarpString(PropertyCollection, XPBDAnisoBucklingStiffnessWarpName.ToString())),
+		FSolverVec2(GetWeightedFloatXPBDAnisoBendingStiffnessWarp(PropertyCollection, MaxStiffness)),
+		(FSolverReal)GetXPBDAnisoBucklingRatio(PropertyCollection, 0.f),
+		FSolverVec2(GetWeightedFloatXPBDAnisoBucklingStiffnessWarp(PropertyCollection, MaxStiffness)),
 		bTrimKinematicConstraints,
 		MaxStiffness)
 	, StiffnessWeft(
-		FSolverVec2(GetWeightedFloatXPBDBendingElementStiffnessWeft(PropertyCollection, MaxStiffness)),
-		WeightMaps.FindRef(GetXPBDBendingElementStiffnessWeftString(PropertyCollection, XPBDBendingElementStiffnessWeftName.ToString())),
+		FSolverVec2(GetWeightedFloatXPBDAnisoBendingStiffnessWeft(PropertyCollection, MaxStiffness)),
+		WeightMaps.FindRef(GetXPBDAnisoBendingStiffnessWeftString(PropertyCollection, XPBDAnisoBendingStiffnessWeftName.ToString())),
 		TConstArrayView<TVec2<int32>>(ConstraintSharedEdges),
 		ParticleOffset,
 		ParticleCount,
@@ -52,8 +52,8 @@ FXPBDAnisotropicBendingConstraints::FXPBDAnisotropicBendingConstraints(const FSo
 		FPBDStiffness::DefaultParameterFitBase,
 		MaxStiffness)
 	, StiffnessBias(
-		FSolverVec2(GetWeightedFloatXPBDBendingElementStiffnessBias(PropertyCollection, MaxStiffness)),
-		WeightMaps.FindRef(GetXPBDBendingElementStiffnessBiasString(PropertyCollection, XPBDBendingElementStiffnessBiasName.ToString())),
+		FSolverVec2(GetWeightedFloatXPBDAnisoBendingStiffnessBias(PropertyCollection, MaxStiffness)),
+		WeightMaps.FindRef(GetXPBDAnisoBendingStiffnessBiasString(PropertyCollection, XPBDAnisoBendingStiffnessBiasName.ToString())),
 		TConstArrayView<TVec2<int32>>(ConstraintSharedEdges),
 		ParticleOffset,
 		ParticleCount,
@@ -61,8 +61,8 @@ FXPBDAnisotropicBendingConstraints::FXPBDAnisotropicBendingConstraints(const FSo
 		FPBDStiffness::DefaultParameterFitBase,
 		MaxStiffness)
 	, BucklingStiffnessWeft(
-		FSolverVec2(GetWeightedFloatXPBDBucklingStiffnessWeft(PropertyCollection, MaxStiffness)),
-		WeightMaps.FindRef(GetXPBDBucklingStiffnessWeftString(PropertyCollection, XPBDBucklingStiffnessWeftName.ToString())),
+		FSolverVec2(GetWeightedFloatXPBDAnisoBucklingStiffnessWeft(PropertyCollection, MaxStiffness)),
+		WeightMaps.FindRef(GetXPBDAnisoBucklingStiffnessWeftString(PropertyCollection, XPBDAnisoBucklingStiffnessWeftName.ToString())),
 		TConstArrayView<TVec2<int32>>(ConstraintSharedEdges),
 		ParticleOffset,
 		ParticleCount,
@@ -70,8 +70,8 @@ FXPBDAnisotropicBendingConstraints::FXPBDAnisotropicBendingConstraints(const FSo
 		FPBDStiffness::DefaultParameterFitBase,
 		MaxStiffness)
 	, BucklingStiffnessBias(
-		FSolverVec2(GetWeightedFloatXPBDBucklingStiffnessBias(PropertyCollection, MaxStiffness)),
-		WeightMaps.FindRef(GetXPBDBucklingStiffnessBiasString(PropertyCollection, XPBDBucklingStiffnessBiasName.ToString())),
+		FSolverVec2(GetWeightedFloatXPBDAnisoBucklingStiffnessBias(PropertyCollection, MaxStiffness)),
+		WeightMaps.FindRef(GetXPBDAnisoBucklingStiffnessBiasString(PropertyCollection, XPBDAnisoBucklingStiffnessBiasName.ToString())),
 		TConstArrayView<TVec2<int32>>(ConstraintSharedEdges),
 		ParticleOffset,
 		ParticleCount,
@@ -79,8 +79,8 @@ FXPBDAnisotropicBendingConstraints::FXPBDAnisotropicBendingConstraints(const FSo
 		FPBDStiffness::DefaultParameterFitBase,
 		MaxStiffness)
 	, DampingRatio(
-		FSolverVec2(GetWeightedFloatXPBDBendingElementDamping(PropertyCollection, MinDamping)).ClampAxes(MinDamping, MaxDamping),
-		WeightMaps.FindRef(GetXPBDBendingElementDampingString(PropertyCollection, XPBDBendingElementDampingName.ToString())),
+		FSolverVec2(GetWeightedFloatXPBDAnisoBendingDamping(PropertyCollection, MinDamping)).ClampAxes(MinDamping, MaxDamping),
+		WeightMaps.FindRef(GetXPBDAnisoBendingDampingString(PropertyCollection, XPBDAnisoBendingDampingName.ToString())),
 		TConstArrayView<TVec2<int32>>(ConstraintSharedEdges),
 		ParticleOffset,
 		ParticleCount)
@@ -284,12 +284,12 @@ void FXPBDAnisotropicBendingConstraints::SetProperties(
 	const FCollectionPropertyConstFacade& PropertyCollection,
 	const TMap<FString, TConstArrayView<FRealSingle>>& WeightMaps)
 {
-	if (IsXPBDBendingElementStiffnessWarpMutable(PropertyCollection))
+	if (IsXPBDAnisoBendingStiffnessWarpMutable(PropertyCollection))
 	{
-		const FSolverVec2 WeightedValue(GetWeightedFloatXPBDBendingElementStiffnessWarp(PropertyCollection));
-		if (IsXPBDBendingElementStiffnessWarpStringDirty(PropertyCollection))
+		const FSolverVec2 WeightedValue(GetWeightedFloatXPBDAnisoBendingStiffnessWarp(PropertyCollection));
+		if (IsXPBDAnisoBendingStiffnessWarpStringDirty(PropertyCollection))
 		{
-			const FString& WeightMapName = GetXPBDBendingElementStiffnessWarpString(PropertyCollection);
+			const FString& WeightMapName = GetXPBDAnisoBendingStiffnessWarpString(PropertyCollection);
 			Stiffness = FPBDStiffness(
 				WeightedValue,
 				WeightMaps.FindRef(WeightMapName),
@@ -305,12 +305,12 @@ void FXPBDAnisotropicBendingConstraints::SetProperties(
 			Stiffness.SetWeightedValue(WeightedValue, MaxStiffness);
 		}
 	}
-	if (IsXPBDBendingElementStiffnessWeftMutable(PropertyCollection))
+	if (IsXPBDAnisoBendingStiffnessWeftMutable(PropertyCollection))
 	{
-		const FSolverVec2 WeightedValue(GetWeightedFloatXPBDBendingElementStiffnessWeft(PropertyCollection));
-		if (IsXPBDBendingElementStiffnessWeftStringDirty(PropertyCollection))
+		const FSolverVec2 WeightedValue(GetWeightedFloatXPBDAnisoBendingStiffnessWeft(PropertyCollection));
+		if (IsXPBDAnisoBendingStiffnessWeftStringDirty(PropertyCollection))
 		{
-			const FString& WeightMapName = GetXPBDBendingElementStiffnessWeftString(PropertyCollection);
+			const FString& WeightMapName = GetXPBDAnisoBendingStiffnessWeftString(PropertyCollection);
 			StiffnessWeft = FPBDStiffness(
 				WeightedValue,
 				WeightMaps.FindRef(WeightMapName),
@@ -326,12 +326,12 @@ void FXPBDAnisotropicBendingConstraints::SetProperties(
 			StiffnessWeft.SetWeightedValue(WeightedValue, MaxStiffness);
 		}
 	}
-	if (IsXPBDBendingElementStiffnessBiasMutable(PropertyCollection))
+	if (IsXPBDAnisoBendingStiffnessBiasMutable(PropertyCollection))
 	{
-		const FSolverVec2 WeightedValue(GetWeightedFloatXPBDBendingElementStiffnessBias(PropertyCollection));
-		if (IsXPBDBendingElementStiffnessBiasStringDirty(PropertyCollection))
+		const FSolverVec2 WeightedValue(GetWeightedFloatXPBDAnisoBendingStiffnessBias(PropertyCollection));
+		if (IsXPBDAnisoBendingStiffnessBiasStringDirty(PropertyCollection))
 		{
-			const FString& WeightMapName = GetXPBDBendingElementStiffnessBiasString(PropertyCollection);
+			const FString& WeightMapName = GetXPBDAnisoBendingStiffnessBiasString(PropertyCollection);
 			StiffnessBias = FPBDStiffness(
 				WeightedValue,
 				WeightMaps.FindRef(WeightMapName),
@@ -347,16 +347,16 @@ void FXPBDAnisotropicBendingConstraints::SetProperties(
 			StiffnessBias.SetWeightedValue(WeightedValue, MaxStiffness);
 		}
 	}
-	if (IsXPBDBucklingRatioMutable(PropertyCollection))
+	if (IsXPBDAnisoBucklingRatioMutable(PropertyCollection))
 	{
-		BucklingRatio = FMath::Clamp(GetXPBDBucklingRatio(PropertyCollection), (FSolverReal)0., (FSolverReal)1.);
+		BucklingRatio = FMath::Clamp(GetXPBDAnisoBucklingRatio(PropertyCollection), (FSolverReal)0., (FSolverReal)1.);
 	}
-	if (IsXPBDBucklingStiffnessWarpMutable(PropertyCollection))
+	if (IsXPBDAnisoBucklingStiffnessWarpMutable(PropertyCollection))
 	{
-		const FSolverVec2 WeightedValue(GetWeightedFloatXPBDBucklingStiffnessWarp(PropertyCollection));
-		if (IsXPBDBucklingStiffnessWarpStringDirty(PropertyCollection))
+		const FSolverVec2 WeightedValue(GetWeightedFloatXPBDAnisoBucklingStiffnessWarp(PropertyCollection));
+		if (IsXPBDAnisoBucklingStiffnessWarpStringDirty(PropertyCollection))
 		{
-			const FString& WeightMapName = GetXPBDBucklingStiffnessWarpString(PropertyCollection);
+			const FString& WeightMapName = GetXPBDAnisoBucklingStiffnessWarpString(PropertyCollection);
 			BucklingStiffness = FPBDStiffness(
 				WeightedValue,
 				WeightMaps.FindRef(WeightMapName),
@@ -372,12 +372,12 @@ void FXPBDAnisotropicBendingConstraints::SetProperties(
 			BucklingStiffness.SetWeightedValue(WeightedValue, MaxStiffness);
 		}
 	}
-	if (IsXPBDBucklingStiffnessWeftMutable(PropertyCollection))
+	if (IsXPBDAnisoBucklingStiffnessWeftMutable(PropertyCollection))
 	{
-		const FSolverVec2 WeightedValue(GetWeightedFloatXPBDBucklingStiffnessWeft(PropertyCollection));
-		if (IsXPBDBucklingStiffnessWeftStringDirty(PropertyCollection))
+		const FSolverVec2 WeightedValue(GetWeightedFloatXPBDAnisoBucklingStiffnessWeft(PropertyCollection));
+		if (IsXPBDAnisoBucklingStiffnessWeftStringDirty(PropertyCollection))
 		{
-			const FString& WeightMapName = GetXPBDBucklingStiffnessWeftString(PropertyCollection);
+			const FString& WeightMapName = GetXPBDAnisoBucklingStiffnessWeftString(PropertyCollection);
 			BucklingStiffnessWeft = FPBDStiffness(
 				WeightedValue,
 				WeightMaps.FindRef(WeightMapName),
@@ -393,12 +393,12 @@ void FXPBDAnisotropicBendingConstraints::SetProperties(
 			BucklingStiffnessWeft.SetWeightedValue(WeightedValue, MaxStiffness);
 		}
 	}
-	if (IsXPBDBucklingStiffnessBiasMutable(PropertyCollection))
+	if (IsXPBDAnisoBucklingStiffnessBiasMutable(PropertyCollection))
 	{
-		const FSolverVec2 WeightedValue(GetWeightedFloatXPBDBucklingStiffnessBias(PropertyCollection));
-		if (IsXPBDBucklingStiffnessBiasStringDirty(PropertyCollection))
+		const FSolverVec2 WeightedValue(GetWeightedFloatXPBDAnisoBucklingStiffnessBias(PropertyCollection));
+		if (IsXPBDAnisoBucklingStiffnessBiasStringDirty(PropertyCollection))
 		{
-			const FString& WeightMapName = GetXPBDBucklingStiffnessBiasString(PropertyCollection);
+			const FString& WeightMapName = GetXPBDAnisoBucklingStiffnessBiasString(PropertyCollection);
 			BucklingStiffnessBias = FPBDStiffness(
 				WeightedValue,
 				WeightMaps.FindRef(WeightMapName),
@@ -414,12 +414,12 @@ void FXPBDAnisotropicBendingConstraints::SetProperties(
 			BucklingStiffnessBias.SetWeightedValue(WeightedValue, MaxStiffness);
 		}
 	}
-	if (IsXPBDBendingElementDampingMutable(PropertyCollection))
+	if (IsXPBDAnisoBendingDampingMutable(PropertyCollection))
 	{
-		const FSolverVec2 WeightedValue = FSolverVec2(GetXPBDBendingElementDamping(PropertyCollection)).ClampAxes(MinDamping, MaxDamping);
-		if (IsXPBDBendingElementDampingStringDirty(PropertyCollection))
+		const FSolverVec2 WeightedValue = FSolverVec2(GetXPBDAnisoBendingDamping(PropertyCollection)).ClampAxes(MinDamping, MaxDamping);
+		if (IsXPBDAnisoBendingDampingStringDirty(PropertyCollection))
 		{
-			const FString& WeightMapName = GetXPBDBendingElementDampingString(PropertyCollection);
+			const FString& WeightMapName = GetXPBDAnisoBendingDampingString(PropertyCollection);
 			DampingRatio = FPBDWeightMap(
 				WeightedValue,
 				WeightMaps.FindRef(WeightMapName),
diff --git a/Engine/Source/Runtime/Experimental/Chaos/Private/Chaos/XPBDStretchBiasElementConstraints.cpp b/Engine/Source/Runtime/Experimental/Chaos/Private/Chaos/XPBDStretchBiasElementConstraints.cpp
--- a/Engine/Source/Runtime/Experimental/Chaos/Private/Chaos/XPBDStretchBiasElementConstraints.cpp
+++ b/Engine/Source/Runtime/Experimental/Chaos/Private/Chaos/XPBDStretchBiasElementConstraints.cpp
@@ -38,8 +38,8 @@ FXPBDStretchBiasElementConstraints::FXPBDStretchBiasElementConstraints(const FSo
 	: ParticleOffset(InParticleOffset)
 	, ParticleCount(InParticleCount)
 	, StiffnessWarp(
-		FSolverVec2(GetWeightedFloatXPBDStretchBiasElementStiffnessWarp(PropertyCollection, MaxStiffness)),
-		WeightMaps.FindRef(GetXPBDStretchBiasElementStiffnessWarpString(PropertyCollection, XPBDStretchBiasElementStiffnessWarpName.ToString())),
+		FSolverVec2(GetWeightedFloatXPBDAnisoStretchStiffnessWarp(PropertyCollection, MaxStiffness)),
+		WeightMaps.FindRef(GetXPBDAnisoStretchStiffnessWarpString(PropertyCollection, XPBDAnisoStretchStiffnessWarpName.ToString())),
 		TConstArrayView<TVec3<int32>>(Constraints),
 		ParticleOffset,
 		ParticleCount,
@@ -47,8 +47,8 @@ FXPBDStretchBiasElementConstraints::FXPBDStretchBiasElementConstraints(const FSo
 		FPBDStiffness::DefaultParameterFitBase,
 		MaxStiffness)
 	, StiffnessWeft(
-		FSolverVec2(GetWeightedFloatXPBDStretchBiasElementStiffnessWeft(PropertyCollection, MaxStiffness)),
-		WeightMaps.FindRef(GetXPBDStretchBiasElementStiffnessWeftString(PropertyCollection, XPBDStretchBiasElementStiffnessWeftName.ToString())),
+		FSolverVec2(GetWeightedFloatXPBDAnisoStretchStiffnessWeft(PropertyCollection, MaxStiffness)),
+		WeightMaps.FindRef(GetXPBDAnisoStretchStiffnessWeftString(PropertyCollection, XPBDAnisoStretchStiffnessWeftName.ToString())),
 		TConstArrayView<TVec3<int32>>(Constraints),
 		ParticleOffset,
 		ParticleCount,
@@ -56,8 +56,8 @@ FXPBDStretchBiasElementConstraints::FXPBDStretchBiasElementConstraints(const FSo
 		FPBDStiffness::DefaultParameterFitBase,
 		MaxStiffness)
 	, StiffnessBias(
-		FSolverVec2(GetWeightedFloatXPBDStretchBiasElementStiffnessBias(PropertyCollection, MaxStiffness)),
-		WeightMaps.FindRef(GetXPBDStretchBiasElementStiffnessBiasString(PropertyCollection, XPBDStretchBiasElementStiffnessBiasName.ToString())),
+		FSolverVec2(GetWeightedFloatXPBDAnisoStretchStiffnessBias(PropertyCollection, MaxStiffness)),
+		WeightMaps.FindRef(GetXPBDAnisoStretchStiffnessBiasString(PropertyCollection, XPBDAnisoStretchStiffnessBiasName.ToString())),
 		TConstArrayView<TVec3<int32>>(Constraints),
 		ParticleOffset,
 		ParticleCount,
@@ -65,26 +65,26 @@ FXPBDStretchBiasElementConstraints::FXPBDStretchBiasElementConstraints(const FSo
 		FPBDStiffness::DefaultParameterFitBase,
 		MaxStiffness)
 	, DampingRatio(
-		FSolverVec2(GetWeightedFloatXPBDStretchBiasElementDamping(PropertyCollection, MinDamping)).ClampAxes(MinDamping, MaxDamping),
-		WeightMaps.FindRef(GetXPBDStretchBiasElementDampingString(PropertyCollection, XPBDStretchBiasElementDampingName.ToString())),
+		FSolverVec2(GetWeightedFloatXPBDAnisoStretchDamping(PropertyCollection, MinDamping)).ClampAxes(MinDamping, MaxDamping),
+		WeightMaps.FindRef(GetXPBDAnisoStretchDampingString(PropertyCollection, XPBDAnisoStretchDampingName.ToString())),
 		TConstArrayView<TVec3<int32>>(Constraints),
 		ParticleOffset,
 		ParticleCount)
 	, WarpScale(
-		FSolverVec2(GetWeightedFloatXPBDStretchBiasElementWarpScale(PropertyCollection, DefaultWarpWeftScale)).ClampAxes(MinWarpWeftScale, MaxWarpWeftScale),
-		WeightMaps.FindRef(GetXPBDStretchBiasElementWarpScaleString(PropertyCollection, XPBDStretchBiasElementWarpScaleName.ToString())),
+		FSolverVec2(GetWeightedFloatXPBDAnisoStretchWarpScale(PropertyCollection, DefaultWarpWeftScale)).ClampAxes(MinWarpWeftScale, MaxWarpWeftScale),
+		WeightMaps.FindRef(GetXPBDAnisoStretchWarpScaleString(PropertyCollection, XPBDAnisoStretchWarpScaleName.ToString())),
 		TConstArrayView<TVec3<int32>>(Constraints),
 		ParticleOffset,
 		ParticleCount)
 	, WeftScale(
-		FSolverVec2(GetWeightedFloatXPBDStretchBiasElementWeftScale(PropertyCollection, DefaultWarpWeftScale)).ClampAxes(MinWarpWeftScale, MaxWarpWeftScale),
-		WeightMaps.FindRef(GetXPBDStretchBiasElementWeftScaleString(PropertyCollection, XPBDStretchBiasElementWeftScaleName.ToString())),
+		FSolverVec2(GetWeightedFloatXPBDAnisoStretchWeftScale(PropertyCollection, DefaultWarpWeftScale)).ClampAxes(MinWarpWeftScale, MaxWarpWeftScale),
+		WeightMaps.FindRef(GetXPBDAnisoStretchWeftScaleString(PropertyCollection, XPBDAnisoStretchWeftScaleName.ToString())),
 		TConstArrayView<TVec3<int32>>(Constraints),
 		ParticleOffset,
 		ParticleCount)
 {
 	Lambdas.Init(FSolverVec3(0.), Constraints.Num());
-	InitConstraintsAndRestData(InParticles, TriangleMesh, FaceVertexUVs, GetXPBDStretchBiasElementUse3dRestLengths(PropertyCollection, bDefaultUse3dRestLengths), bTrimKinematicConstraints);
+	InitConstraintsAndRestData(InParticles, TriangleMesh, FaceVertexUVs, GetXPBDAnisoStretchUse3dRestLengths(PropertyCollection, bDefaultUse3dRestLengths), bTrimKinematicConstraints);
 	InitColor(InParticles);
 }
 
@@ -302,12 +302,12 @@ void FXPBDStretchBiasElementConstraints::SetProperties(
 	const FCollectionPropertyConstFacade& PropertyCollection,
 	const TMap<FString, TConstArrayView<FRealSingle>>& WeightMaps)
 {
-	if (IsXPBDStretchBiasElementStiffnessWarpMutable(PropertyCollection))
+	if (IsXPBDAnisoStretchStiffnessWarpMutable(PropertyCollection))
 	{
-		const FSolverVec2 WeightedValue(GetWeightedFloatXPBDStretchBiasElementStiffnessWarp(PropertyCollection));
-		if (IsXPBDStretchBiasElementStiffnessWarpStringDirty(PropertyCollection))
+		const FSolverVec2 WeightedValue(GetWeightedFloatXPBDAnisoStretchStiffnessWarp(PropertyCollection));
+		if (IsXPBDAnisoStretchStiffnessWarpStringDirty(PropertyCollection))
 		{
-			const FString& WeightMapName = GetXPBDStretchBiasElementStiffnessWarpString(PropertyCollection);
+			const FString& WeightMapName = GetXPBDAnisoStretchStiffnessWarpString(PropertyCollection);
 			StiffnessWarp = FPBDStiffness(
 				WeightedValue,
 				WeightMaps.FindRef(WeightMapName),
@@ -323,12 +323,12 @@ void FXPBDStretchBiasElementConstraints::SetProperties(
 			StiffnessWarp.SetWeightedValue(WeightedValue, MaxStiffness);
 		}
 	}
-	if (IsXPBDStretchBiasElementStiffnessWeftMutable(PropertyCollection))
+	if (IsXPBDAnisoStretchStiffnessWeftMutable(PropertyCollection))
 	{
-		const FSolverVec2 WeightedValue(GetWeightedFloatXPBDStretchBiasElementStiffnessWeft(PropertyCollection));
-		if (IsXPBDStretchBiasElementStiffnessWeftStringDirty(PropertyCollection))
+		const FSolverVec2 WeightedValue(GetWeightedFloatXPBDAnisoStretchStiffnessWeft(PropertyCollection));
+		if (IsXPBDAnisoStretchStiffnessWeftStringDirty(PropertyCollection))
 		{
-			const FString& WeightMapName = GetXPBDStretchBiasElementStiffnessWeftString(PropertyCollection);
+			const FString& WeightMapName = GetXPBDAnisoStretchStiffnessWeftString(PropertyCollection);
 			StiffnessWeft = FPBDStiffness(
 				WeightedValue,
 				WeightMaps.FindRef(WeightMapName),
@@ -344,12 +344,12 @@ void FXPBDStretchBiasElementConstraints::SetProperties(
 			StiffnessWeft.SetWeightedValue(WeightedValue, MaxStiffness);
 		}
 	}
-	if (IsXPBDStretchBiasElementStiffnessBiasMutable(PropertyCollection))
+	if (IsXPBDAnisoStretchStiffnessBiasMutable(PropertyCollection))
 	{
-		const FSolverVec2 WeightedValue(GetWeightedFloatXPBDStretchBiasElementStiffnessBias(PropertyCollection));
-		if (IsXPBDStretchBiasElementStiffnessBiasStringDirty(PropertyCollection))
+		const FSolverVec2 WeightedValue(GetWeightedFloatXPBDAnisoStretchStiffnessBias(PropertyCollection));
+		if (IsXPBDAnisoStretchStiffnessBiasStringDirty(PropertyCollection))
 		{
-			const FString& WeightMapName = GetXPBDStretchBiasElementStiffnessBiasString(PropertyCollection);
+			const FString& WeightMapName = GetXPBDAnisoStretchStiffnessBiasString(PropertyCollection);
 			StiffnessBias = FPBDStiffness(
 				WeightedValue,
 				WeightMaps.FindRef(WeightMapName),
@@ -365,12 +365,12 @@ void FXPBDStretchBiasElementConstraints::SetProperties(
 			StiffnessBias.SetWeightedValue(WeightedValue, MaxStiffness);
 		}
 	}
-	if (IsXPBDStretchBiasElementDampingMutable(PropertyCollection))
+	if (IsXPBDAnisoStretchDampingMutable(PropertyCollection))
 	{
-		const FSolverVec2 WeightedValue = FSolverVec2(GetWeightedFloatXPBDStretchBiasElementDamping(PropertyCollection)).ClampAxes(MinDamping, MaxDamping);
-		if (IsXPBDStretchBiasElementDampingStringDirty(PropertyCollection))
+		const FSolverVec2 WeightedValue = FSolverVec2(GetWeightedFloatXPBDAnisoStretchDamping(PropertyCollection)).ClampAxes(MinDamping, MaxDamping);
+		if (IsXPBDAnisoStretchDampingStringDirty(PropertyCollection))
 		{
-			const FString& WeightMapName = GetXPBDStretchBiasElementDampingString(PropertyCollection);
+			const FString& WeightMapName = GetXPBDAnisoStretchDampingString(PropertyCollection);
 			DampingRatio = FPBDWeightMap(
 				WeightedValue,
 				WeightMaps.FindRef(WeightMapName),
@@ -383,12 +383,12 @@ void FXPBDStretchBiasElementConstraints::SetProperties(
 			DampingRatio.SetWeightedValue(WeightedValue);
 		}
 	}
-	if (IsXPBDStretchBiasElementWarpScaleMutable(PropertyCollection))
+	if (IsXPBDAnisoStretchWarpScaleMutable(PropertyCollection))
 	{
-		const FSolverVec2 WeightedValue = FSolverVec2(GetWeightedFloatXPBDStretchBiasElementWarpScale(PropertyCollection)).ClampAxes(MinWarpWeftScale, MaxWarpWeftScale);
-		if (IsXPBDStretchBiasElementWarpScaleStringDirty(PropertyCollection))
+		const FSolverVec2 WeightedValue = FSolverVec2(GetWeightedFloatXPBDAnisoStretchWarpScale(PropertyCollection)).ClampAxes(MinWarpWeftScale, MaxWarpWeftScale);
+		if (IsXPBDAnisoStretchWarpScaleStringDirty(PropertyCollection))
 		{
-			const FString& WeightMapName = GetXPBDStretchBiasElementWarpScaleString(PropertyCollection);
+			const FString& WeightMapName = GetXPBDAnisoStretchWarpScaleString(PropertyCollection);
 			WarpScale = FPBDWeightMap(
 				WeightedValue,
 				WeightMaps.FindRef(WeightMapName),
@@ -401,12 +401,12 @@ void FXPBDStretchBiasElementConstraints::SetProperties(
 			WarpScale.SetWeightedValue(WeightedValue);
 		}
 	}
-	if (IsXPBDStretchBiasElementWeftScaleMutable(PropertyCollection))
+	if (IsXPBDAnisoStretchWeftScaleMutable(PropertyCollection))
 	{
-		const FSolverVec2 WeightedValue = FSolverVec2(GetWeightedFloatXPBDStretchBiasElementWeftScale(PropertyCollection)).ClampAxes(MinWarpWeftScale, MaxWarpWeftScale);
-		if (IsXPBDStretchBiasElementWeftScaleStringDirty(PropertyCollection))
+		const FSolverVec2 WeightedValue = FSolverVec2(GetWeightedFloatXPBDAnisoStretchWeftScale(PropertyCollection)).ClampAxes(MinWarpWeftScale, MaxWarpWeftScale);
+		if (IsXPBDAnisoStretchWeftScaleStringDirty(PropertyCollection))
 		{
-			const FString& WeightMapName = GetXPBDStretchBiasElementWeftScaleString(PropertyCollection);
+			const FString& WeightMapName = GetXPBDAnisoStretchWeftScaleString(PropertyCollection);
 			WeftScale = FPBDWeightMap(
 				WeightedValue,
 				WeightMaps.FindRef(WeightMapName),
diff --git a/Engine/Source/Runtime/Experimental/Chaos/Public/Chaos/XPBDAnisotropicBendingConstraints.h b/Engine/Source/Runtime/Experimental/Chaos/Public/Chaos/XPBDAnisotropicBendingConstraints.h
--- a/Engine/Source/Runtime/Experimental/Chaos/Public/Chaos/XPBDAnisotropicBendingConstraints.h
+++ b/Engine/Source/Runtime/Experimental/Chaos/Public/Chaos/XPBDAnisotropicBendingConstraints.h
@@ -26,7 +26,7 @@ public:
 
 	static bool IsEnabled(const FCollectionPropertyConstFacade& PropertyCollection)
 	{
-		return IsXPBDBendingElementStiffnessWarpEnabled(PropertyCollection, false);
+		return IsXPBDAnisoBendingStiffnessWarpEnabled(PropertyCollection, false);
 	}
 
 	FXPBDAnisotropicBendingConstraints(const FSolverParticles& InParticles,
@@ -121,14 +121,14 @@ private:
 
 	TArray<FSolverVec3> WarpWeftBiasBaseMultipliers;
 
-	UE_CHAOS_DECLARE_PROPERTYCOLLECTION_NAME(XPBDBendingElementStiffnessWarp, float);
-	UE_CHAOS_DECLARE_PROPERTYCOLLECTION_NAME(XPBDBendingElementStiffnessWeft, float);
-	UE_CHAOS_DECLARE_PROPERTYCOLLECTION_NAME(XPBDBendingElementStiffnessBias, float);
-	UE_CHAOS_DECLARE_PROPERTYCOLLECTION_NAME(XPBDBendingElementDamping, float);
-	UE_CHAOS_DECLARE_PROPERTYCOLLECTION_NAME(XPBDBucklingRatio, float);
-	UE_CHAOS_DECLARE_PROPERTYCOLLECTION_NAME(XPBDBucklingStiffnessWarp, float);
-	UE_CHAOS_DECLARE_PROPERTYCOLLECTION_NAME(XPBDBucklingStiffnessWeft, float);
-	UE_CHAOS_DECLARE_PROPERTYCOLLECTION_NAME(XPBDBucklingStiffnessBias, float);
+	UE_CHAOS_DECLARE_PROPERTYCOLLECTION_NAME(XPBDAnisoBendingStiffnessWarp, float);
+	UE_CHAOS_DECLARE_PROPERTYCOLLECTION_NAME(XPBDAnisoBendingStiffnessWeft, float);
+	UE_CHAOS_DECLARE_PROPERTYCOLLECTION_NAME(XPBDAnisoBendingStiffnessBias, float);
+	UE_CHAOS_DECLARE_PROPERTYCOLLECTION_NAME(XPBDAnisoBendingDamping, float);
+	UE_CHAOS_DECLARE_PROPERTYCOLLECTION_NAME(XPBDAnisoBucklingRatio, float);
+	UE_CHAOS_DECLARE_PROPERTYCOLLECTION_NAME(XPBDAnisoBucklingStiffnessWarp, float);
+	UE_CHAOS_DECLARE_PROPERTYCOLLECTION_NAME(XPBDAnisoBucklingStiffnessWeft, float);
+	UE_CHAOS_DECLARE_PROPERTYCOLLECTION_NAME(XPBDAnisoBucklingStiffnessBias, float);
 };
 
 }  // End namespace Chaos::Softs
diff --git a/Engine/Source/Runtime/Experimental/Chaos/Public/Chaos/XPBDStretchBiasElementConstraints.h b/Engine/Source/Runtime/Experimental/Chaos/Public/Chaos/XPBDStretchBiasElementConstraints.h
--- a/Engine/Source/Runtime/Experimental/Chaos/Public/Chaos/XPBDStretchBiasElementConstraints.h
+++ b/Engine/Source/Runtime/Experimental/Chaos/Public/Chaos/XPBDStretchBiasElementConstraints.h
@@ -31,7 +31,7 @@ public:
 
 	static bool IsEnabled(const FCollectionPropertyConstFacade& PropertyCollection)
 	{
-		return IsXPBDStretchBiasElementStiffnessWarpEnabled(PropertyCollection, false);
+		return IsXPBDAnisoStretchStiffnessWarpEnabled(PropertyCollection, false);
 	}
 
 	FXPBDStretchBiasElementConstraints(const FSolverParticles& InParticles,
@@ -130,13 +130,13 @@ private:
 	TArray<FSolverVec2> RestStretchLengths;
 	TArray<FSolverVec3> StiffnessScales; // Used to make everything resolution independent.
 
-	UE_CHAOS_DECLARE_PROPERTYCOLLECTION_NAME(XPBDStretchBiasElementUse3dRestLengths, bool);
-	UE_CHAOS_DECLARE_PROPERTYCOLLECTION_NAME(XPBDStretchBiasElementStiffnessWarp, float);
-	UE_CHAOS_DECLARE_PROPERTYCOLLECTION_NAME(XPBDStretchBiasElementStiffnessWeft, float);
-	UE_CHAOS_DECLARE_PROPERTYCOLLECTION_NAME(XPBDStretchBiasElementStiffnessBias, float);
-	UE_CHAOS_DECLARE_PROPERTYCOLLECTION_NAME(XPBDStretchBiasElementDamping, float);
-	UE_CHAOS_DECLARE_PROPERTYCOLLECTION_NAME(XPBDStretchBiasElementWarpScale, float);
-	UE_CHAOS_DECLARE_PROPERTYCOLLECTION_NAME(XPBDStretchBiasElementWeftScale, float);
+	UE_CHAOS_DECLARE_PROPERTYCOLLECTION_NAME(XPBDAnisoStretchUse3dRestLengths, bool);
+	UE_CHAOS_DECLARE_PROPERTYCOLLECTION_NAME(XPBDAnisoStretchStiffnessWarp, float);
+	UE_CHAOS_DECLARE_PROPERTYCOLLECTION_NAME(XPBDAnisoStretchStiffnessWeft, float);
+	UE_CHAOS_DECLARE_PROPERTYCOLLECTION_NAME(XPBDAnisoStretchStiffnessBias, float);
+	UE_CHAOS_DECLARE_PROPERTYCOLLECTION_NAME(XPBDAnisoStretchDamping, float);
+	UE_CHAOS_DECLARE_PROPERTYCOLLECTION_NAME(XPBDAnisoStretchWarpScale, float);
+	UE_CHAOS_DECLARE_PROPERTYCOLLECTION_NAME(XPBDAnisoStretchWeftScale, float);
 };
 
 }  // End namespace Chaos::Softs
