CIS - Fixed cannot instantiate abstract class error in FChaosClothAssetSimulationBaseConfigNode.

#rnx

diff --git a/Engine/Plugins/Experimental/ChaosClothAssetEditor/Source/ChaosClothAssetDataflowNodes/Public/ChaosClothAsset/SimulationBaseConfigNode.h b/Engine/Plugins/Experimental/ChaosClothAssetEditor/Source/ChaosClothAssetDataflowNodes/Public/ChaosClothAsset/SimulationBaseConfigNode.h
--- a/Engine/Plugins/Experimental/ChaosClothAssetEditor/Source/ChaosClothAssetDataflowNodes/Public/ChaosClothAsset/SimulationBaseConfigNode.h
+++ b/Engine/Plugins/Experimental/ChaosClothAssetEditor/Source/ChaosClothAssetDataflowNodes/Public/ChaosClothAsset/SimulationBaseConfigNode.h
@@ -42,3 +42,12 @@ protected:
 		bool bIsAnimatable = true,
 		const TArray<FName>& SimilarPropertyNames = TArray<FName>()) const;
 };
+
+template<>
+struct TStructOpsTypeTraits<FChaosClothAssetSimulationBaseConfigNode> : public TStructOpsTypeTraitsBase2<FChaosClothAssetSimulationBaseConfigNode>
+{
+	enum
+	{
+		WithPureVirtual = true,
+	};
+};
