CIS - Fixed error declaration shadows a field of FChaosClothAssetMergeClothCollectionsNode in Clang builds.

#rnx

diff --git a/Engine/Plugins/Experimental/ChaosClothAssetEditor/Source/ChaosClothAssetDataflowNodes/Private/ChaosClothAsset/MergeClothCollectionsNode.cpp b/Engine/Plugins/Experimental/ChaosClothAssetEditor/Source/ChaosClothAssetDataflowNodes/Private/ChaosClothAsset/MergeClothCollectionsNode.cpp
--- a/Engine/Plugins/Experimental/ChaosClothAssetEditor/Source/ChaosClothAssetDataflowNodes/Private/ChaosClothAsset/MergeClothCollectionsNode.cpp
+++ b/Engine/Plugins/Experimental/ChaosClothAssetEditor/Source/ChaosClothAssetDataflowNodes/Private/ChaosClothAsset/MergeClothCollectionsNode.cpp
@@ -56,10 +56,10 @@ void FChaosClothAssetMergeClothCollectionsNode::Evaluate(Dataflow::FContext& Con
 
 Dataflow::FPin FChaosClothAssetMergeClothCollectionsNode::AddPin()
 {
-	auto AddInput = [this](const FManagedArrayCollection* Collection) -> Dataflow::FPin
+	auto AddInput = [this](const FManagedArrayCollection* InCollection) -> Dataflow::FPin
 	{
-		RegisterInputConnection(Collection);
-		const FDataflowInput* const Input = FindInput(Collection);
+		RegisterInputConnection(InCollection);
+		const FDataflowInput* const Input = FindInput(InCollection);
 		return { Dataflow::FPin::EDirection::INPUT, Input->GetType(), Input->GetName() };
 	};
 
@@ -78,12 +78,12 @@ Dataflow::FPin FChaosClothAssetMergeClothCollectionsNode::AddPin()
 
 Dataflow::FPin FChaosClothAssetMergeClothCollectionsNode::RemovePin()
 {
-	auto RemoveInput = [this](const FManagedArrayCollection* Collection) -> Dataflow::FPin
+	auto RemoveInput = [this](const FManagedArrayCollection* InCollection) -> Dataflow::FPin
 	{
-		const FDataflowInput* const Input = FindInput(Collection);
+		const FDataflowInput* const Input = FindInput(InCollection);
 		check(Input);
 		Dataflow::FPin Pin = { Dataflow::FPin::EDirection::INPUT, Input->GetType(), Input->GetName() };
-		UnregisterInputConnection(Collection);  // This will delete the input, so set the pin before that
+		UnregisterInputConnection(InCollection);  // This will delete the input, so set the pin before that
 		return Pin;
 	};
 
