Chaos Clothing Asset: Update static mesh import to be able to import render data. Added option to import sim data into single pattern or multiple patterns based on uv islands.

Render vertices are determined by welding vertex instances (face vertex data) corresponding with the same vertex if they're identical.

Minor change: UVChannel default changed from INDEX_NONE to 0, so default is now to import sim mesh positions from UVs rather than generating them from 3d positions and unwrapping.

diff --git a/Engine/Plugins/Experimental/ChaosClothAsset/Source/ChaosClothAsset/Private/ChaosClothAsset/CollectionClothLodFacade.cpp b/Engine/Plugins/Experimental/ChaosClothAsset/Source/ChaosClothAsset/Private/ChaosClothAsset/CollectionClothLodFacade.cpp
--- a/Engine/Plugins/Experimental/ChaosClothAsset/Source/ChaosClothAsset/Private/ChaosClothAsset/CollectionClothLodFacade.cpp
+++ b/Engine/Plugins/Experimental/ChaosClothAsset/Source/ChaosClothAsset/Private/ChaosClothAsset/CollectionClothLodFacade.cpp
@@ -126,7 +126,7 @@ namespace UE::Chaos::ClothAsset::Private
 	}
 
 	template<bool bWeldNearlyCoincidentVertices>
-	static void UnwrapDynamicMesh(const UE::Geometry::FDynamicMesh3& DynamicMesh, TArray<FIsland>& OutIslands)
+	static void UnwrapDynamicMesh(const UE::Geometry::FDynamicMesh3& DynamicMesh, bool bForceSingleIsland, TArray<FIsland>& OutIslands)
 	{
 		using namespace UE::Geometry;
 
@@ -138,6 +138,11 @@ namespace UE::Chaos::ClothAsset::Private
 		TSet<int32> VisitedTriangles;
 		VisitedTriangles.Reserve(NumTriangles);
 
+		if (bForceSingleIsland)
+		{
+			OutIslands.AddDefaulted();
+		}
+
 		for (int32 SeedTriangle : DynamicMesh.TriangleIndicesItr())
 		{
 			if (VisitedTriangles.Contains(SeedTriangle))
@@ -159,7 +164,7 @@ namespace UE::Chaos::ClothAsset::Private
 			}
 
 			// Setup first visitor from seed, and add the first two points
-			FIsland& Island = OutIslands.AddDefaulted_GetRef();
+			FIsland& Island = bForceSingleIsland ? OutIslands[0] : OutIslands.AddDefaulted_GetRef();
 
 			Island.RestPositions.Add(Position0);
 			Island.RestPositions.Add(Position1);
@@ -313,7 +318,7 @@ namespace UE::Chaos::ClothAsset::Private
 		}
 	}
 
-	static void BuildIslandsFromDynamicMeshUVs(const UE::Geometry::FDynamicMeshUVOverlay& UVOverlay, TArray<FIsland>& OutIslands, const FVector2f& UVScale)
+	static void BuildIslandsFromDynamicMeshUVs(const UE::Geometry::FDynamicMeshUVOverlay& UVOverlay, const FVector2f& UVScale, bool bForceSingleIsland, TArray<FIsland>& OutIslands)
 	{
 		using namespace UE::Geometry;
 
@@ -330,6 +335,12 @@ namespace UE::Chaos::ClothAsset::Private
 		// This is reused for each island, but only allocate once.
 		TArray<int32> SourceElementIndexToNewIndex;
 
+		if (bForceSingleIsland)
+		{
+			OutIslands.AddDefaulted();
+			SourceElementIndexToNewIndex.Init(INDEX_NONE, UVOverlay.MaxElementID());
+		}
+
 		for (int32 SeedTriangle : DynamicMesh->TriangleIndicesItr())
 		{
 			if (VisitedTriangles.Contains(SeedTriangle))
@@ -338,9 +349,12 @@ namespace UE::Chaos::ClothAsset::Private
 			}
 
 			// Setup first visitor from seed
-			FIsland& Island = OutIslands.AddDefaulted_GetRef();
+			FIsland& Island = bForceSingleIsland ? OutIslands[0] : OutIslands.AddDefaulted_GetRef();
 
-			SourceElementIndexToNewIndex.Init(INDEX_NONE, UVOverlay.MaxElementID());
+			if (!bForceSingleIsland)
+			{
+				SourceElementIndexToNewIndex.Init(INDEX_NONE, UVOverlay.MaxElementID());
+			}
 
 			struct FVisitor
 			{
@@ -929,7 +943,7 @@ namespace UE::Chaos::ClothAsset
 	}
 
 	template<bool bWeldNearlyCoincidentVertices>
-	void FCollectionClothLodFacade::InitializeFromDynamicMeshInternal(const UE::Geometry::FDynamicMesh3& DynamicMesh, int32 UVChannelIndex, const FVector2f& UVScale)
+	void FCollectionClothLodFacade::InitializeFromDynamicMeshInternal(const UE::Geometry::FDynamicMesh3& DynamicMesh, int32 UVChannelIndex, const FVector2f& UVScale, bool bGenerateSimPatternsFromUVIslands)
 	{
 		using namespace UE::Chaos::ClothAsset::Private;
 		Reset();
@@ -937,14 +951,16 @@ namespace UE::Chaos::ClothAsset
 		const UE::Geometry::FDynamicMeshAttributeSet* const AttributeSet = DynamicMesh.Attributes();
 		const UE::Geometry::FDynamicMeshUVOverlay* const UVOverlay = AttributeSet ? AttributeSet->GetUVLayer(UVChannelIndex) : nullptr;
 
+		const bool bForceSingleIsland = !bGenerateSimPatternsFromUVIslands;
+
 		TArray<FIsland> Islands;
 		if (UVOverlay)
 		{
-			BuildIslandsFromDynamicMeshUVs(*UVOverlay, Islands, UVScale);
+			BuildIslandsFromDynamicMeshUVs(*UVOverlay, UVScale, bForceSingleIsland, Islands);
 		}
 		else
 		{
-			UnwrapDynamicMesh<bWeldNearlyCoincidentVertices>(DynamicMesh, Islands);
+			UnwrapDynamicMesh<bWeldNearlyCoincidentVertices>(DynamicMesh, bForceSingleIsland, Islands);
 		}
 
 		for (FIsland& Island : Islands)
@@ -992,12 +1008,13 @@ namespace UE::Chaos::ClothAsset
 		ToDynamicMesh.Convert(DynamicMesh, SimpleSrc, [](FSimpleSrcMeshInterface::TriIDType) {return 0; });
 		UE::Geometry::FNonManifoldMappingSupport::AttachNonManifoldVertexMappingData(ToDynamicMesh.ToSrcVertIDMap, DynamicMesh);
 
-		InitializeFromDynamicMeshInternal<true>(DynamicMesh, INDEX_NONE, FVector2f(1.f));
+		constexpr bool bGenerateSimPatternsFromUVIslands = false;
+		InitializeFromDynamicMeshInternal<true>(DynamicMesh, INDEX_NONE, FVector2f(1.f), bGenerateSimPatternsFromUVIslands);
 	}
 
-	void FCollectionClothLodFacade::Initialize(const UE::Geometry::FDynamicMesh3& DynamicMesh, int32 UVChannelIndex, const FVector2f& UVScale)
+	void FCollectionClothLodFacade::Initialize(const UE::Geometry::FDynamicMesh3& DynamicMesh, int32 UVChannelIndex, const FVector2f& UVScale, bool bGenerateSimPatternsFromUVIslands)
 	{
-		InitializeFromDynamicMeshInternal<false>(DynamicMesh, UVChannelIndex, UVScale);
+		InitializeFromDynamicMeshInternal<false>(DynamicMesh, UVChannelIndex, UVScale, bGenerateSimPatternsFromUVIslands);
 	}
 
 	void FCollectionClothLodFacade::Initialize(const FCollectionClothLodConstFacade& Other)
diff --git a/Engine/Plugins/Experimental/ChaosClothAsset/Source/ChaosClothAsset/Public/ChaosClothAsset/CollectionClothLodFacade.h b/Engine/Plugins/Experimental/ChaosClothAsset/Source/ChaosClothAsset/Public/ChaosClothAsset/CollectionClothLodFacade.h
--- a/Engine/Plugins/Experimental/ChaosClothAsset/Source/ChaosClothAsset/Public/ChaosClothAsset/CollectionClothLodFacade.h
+++ b/Engine/Plugins/Experimental/ChaosClothAsset/Source/ChaosClothAsset/Public/ChaosClothAsset/CollectionClothLodFacade.h
@@ -140,7 +140,7 @@ namespace UE::Chaos::ClothAsset
 		void Initialize(const FCollectionClothLodConstFacade& Other);
 
 		/** Initialize the cloth LOD from a DynamicMesh (possibly using UV islands to define patterns). */
-		void Initialize(const UE::Geometry::FDynamicMesh3& DynamicMesh, int32 UVChannelIndex, const FVector2f& UVScale);
+		void Initialize(const UE::Geometry::FDynamicMesh3& DynamicMesh, int32 UVChannelIndex, const FVector2f& UVScale, bool bGenerateSimPatternsFromUVIslands);
 
 		/** Append to this cloth LOD another cloth LOD. */
 		void Append(const FCollectionClothLodConstFacade& Other);
@@ -215,7 +215,7 @@ namespace UE::Chaos::ClothAsset
 		void SetDefaults();
 
 		template<bool bWeldNearlyCoincidentVertices>
-		void InitializeFromDynamicMeshInternal(const UE::Geometry::FDynamicMesh3& DynamicMesh, int32 UVChannelIndex, const FVector2f& UVScale);
+		void InitializeFromDynamicMeshInternal(const UE::Geometry::FDynamicMesh3& DynamicMesh, int32 UVChannelIndex, const FVector2f& UVScale, bool bGenerateSimPatternsFromUVIslands);
 
 		TSharedPtr<FClothCollection> GetClothCollection() { return ConstCastSharedPtr<FClothCollection>(ClothCollection); }
 	};
diff --git a/Engine/Plugins/Experimental/ChaosClothAssetEditor/Source/ChaosClothAssetDataflowNodes/ChaosClothAssetDataflowNodes.Build.cs b/Engine/Plugins/Experimental/ChaosClothAssetEditor/Source/ChaosClothAssetDataflowNodes/ChaosClothAssetDataflowNodes.Build.cs
--- a/Engine/Plugins/Experimental/ChaosClothAssetEditor/Source/ChaosClothAssetDataflowNodes/ChaosClothAssetDataflowNodes.Build.cs
+++ b/Engine/Plugins/Experimental/ChaosClothAssetEditor/Source/ChaosClothAssetDataflowNodes/ChaosClothAssetDataflowNodes.Build.cs
@@ -38,6 +38,7 @@ public class ChaosClothAssetDataflowNodes : ModuleRules
 				"SkeletalMeshDescription",
 				"Slate",
 				"SlateCore",
+				"StaticMeshDescription",
 				"UnrealEd",
 				// ... add private dependencies that you statically link with here ...	
 			}
diff --git a/Engine/Plugins/Experimental/ChaosClothAssetEditor/Source/ChaosClothAssetDataflowNodes/Private/ChaosClothAsset/StaticMeshImportNode.cpp b/Engine/Plugins/Experimental/ChaosClothAssetEditor/Source/ChaosClothAssetDataflowNodes/Private/ChaosClothAsset/StaticMeshImportNode.cpp
--- a/Engine/Plugins/Experimental/ChaosClothAssetEditor/Source/ChaosClothAssetDataflowNodes/Private/ChaosClothAsset/StaticMeshImportNode.cpp
+++ b/Engine/Plugins/Experimental/ChaosClothAssetEditor/Source/ChaosClothAssetDataflowNodes/Private/ChaosClothAsset/StaticMeshImportNode.cpp
@@ -7,11 +7,216 @@
 #include "ChaosClothAsset/ClothCollection.h"
 #include "Engine/StaticMesh.h"
 #include "MeshDescriptionToDynamicMesh.h"
+#include "StaticMeshAttributes.h"
+#include "Materials/MaterialInterface.h"
+#include "StaticMeshOperations.h"
 
 #include UE_INLINE_GENERATED_CPP_BY_NAME(StaticMeshImportNode)
 
 #define LOCTEXT_NAMESPACE "ChaosClothAssetStaticMeshImportNode"
 
+namespace UE::Chaos::ClothAsset::Private
+{
+struct FBuildVertex
+{
+	FVertexID OrigVertID;
+	FVertexInstanceID OrigVertInstanceID;
+
+	FVector3f Position;
+	TArray<FVector2f> UVs;
+	FVector3f Normal;
+	FVector3f Tangent;
+	float BiNormalSign;
+	FVector4f Color;
+
+	bool operator==(const FBuildVertex& Other) const
+	{
+		if (OrigVertID != Other.OrigVertID)
+		{
+			return false;
+		}
+
+		// No need to check position since it's from OrigVertID
+
+		if (!ensure(UVs.Num() == Other.UVs.Num()))
+		{
+			return false;
+		}
+		for (int32 UVChannelIndex = 0; UVChannelIndex < UVs.Num(); ++UVChannelIndex)
+		{
+			if (!UVs[UVChannelIndex].Equals(Other.UVs[UVChannelIndex], UE_THRESH_UVS_ARE_SAME))
+			{
+				return false;
+			}
+		}
+		if (!Normal.Equals(Other.Normal, UE_THRESH_NORMALS_ARE_SAME))
+		{
+			return false;
+		}
+		if (!Tangent.Equals(Other.Tangent, UE_THRESH_NORMALS_ARE_SAME))
+		{
+			return false;
+		}
+		if (BiNormalSign != Other.BiNormalSign) // I think these are just -1 or 1, so just comparing them
+		{
+			return false;
+		}
+		// It looks like we just quantize the 0-1 color values to uint8s. (Call FLinearColor::ToFColor(bSRGB = false)) when we consume these.
+		// Going to be a little more strict in case we ever decide to switch to the gamma conversion.
+		if (!Color.Equals(Other.Color, UE_THRESH_NORMALS_ARE_SAME))
+		{
+			return false;
+		}
+		return true;
+	}
+};
+
+static void MergeVertexInstances(const FMeshDescription* const MeshDescription, TArray<FBuildVertex>& MergedVertices, TArray<int32>& VertexInstanceToMerged)
+{
+	FStaticMeshConstAttributes Attributes(*MeshDescription);
+
+	TVertexAttributesConstRef<FVector3f> VertexPositions = Attributes.GetVertexPositions();
+
+	TVertexInstanceAttributesConstRef<FVector2f> VertexInstanceUVs = Attributes.GetVertexInstanceUVs();
+	TVertexInstanceAttributesConstRef<FVector3f> VertexInstanceNormals = Attributes.GetVertexInstanceNormals();
+	TVertexInstanceAttributesConstRef<FVector3f> VertexInstanceTangents = Attributes.GetVertexInstanceTangents();
+	TVertexInstanceAttributesConstRef<float> VertexInstanceBiNormalSigns = Attributes.GetVertexInstanceBinormalSigns();
+	TVertexInstanceAttributesConstRef<FVector4f> VertexInstanceColors = Attributes.GetVertexInstanceColors();
+
+	// Merge identical vertex instances that correspond with the same vertex.
+	auto MakeBuildVertex = [&VertexPositions, &VertexInstanceUVs, &VertexInstanceNormals, &VertexInstanceTangents, &VertexInstanceBiNormalSigns, &VertexInstanceColors](FVertexID VertID, FVertexInstanceID VertexInstanceID)
+	{
+		FBuildVertex BuildVertex;
+		BuildVertex.OrigVertID = VertID;
+		BuildVertex.OrigVertInstanceID = VertexInstanceID;
+		BuildVertex.Position = VertexPositions[VertID];
+		BuildVertex.UVs.SetNumUninitialized(VertexInstanceUVs.GetNumChannels());
+		for (int32 UVChannelIndex = 0; UVChannelIndex < VertexInstanceUVs.GetNumChannels(); ++UVChannelIndex)
+		{
+			BuildVertex.UVs[UVChannelIndex] = VertexInstanceUVs.Get(VertexInstanceID, UVChannelIndex);
+		}
+		BuildVertex.Normal = VertexInstanceNormals[VertexInstanceID];
+		BuildVertex.Tangent = VertexInstanceTangents[VertexInstanceID];
+		BuildVertex.BiNormalSign = VertexInstanceBiNormalSigns[VertexInstanceID];
+		BuildVertex.Color = VertexInstanceColors[VertexInstanceID];
+		return BuildVertex;
+	};
+	MergedVertices.Reset(MeshDescription->VertexInstances().GetArraySize());
+	VertexInstanceToMerged.Init(INDEX_NONE, MeshDescription->VertexInstances().GetArraySize());
+
+	for (const FVertexID& VertID : MeshDescription->Vertices().GetElementIDs())
+	{
+		TConstArrayView<FVertexInstanceID> VertexInstances = MeshDescription->GetVertexVertexInstanceIDs(VertID);
+		if (VertexInstances.IsEmpty())
+		{
+			continue;
+		}
+		const int32 FirstMergedVert = MergedVertices.Add(MakeBuildVertex(VertID, VertexInstances[0]));
+		VertexInstanceToMerged[VertexInstances[0].GetValue()] = FirstMergedVert;
+		for (int32 LocalInstance = 1; LocalInstance < VertexInstances.Num(); ++LocalInstance)
+		{
+			FBuildVertex BuildVert = MakeBuildVertex(VertID, VertexInstances[LocalInstance]);
+			bool bFoundDuplicate = false;
+			for (int32 CompareIndex = FirstMergedVert; CompareIndex < MergedVertices.Num(); ++CompareIndex)
+			{
+				if (BuildVert == MergedVertices[CompareIndex])
+				{
+					bFoundDuplicate = true;
+					VertexInstanceToMerged[VertexInstances[LocalInstance].GetValue()] = CompareIndex;
+					break;
+				}
+			}
+
+			if (!bFoundDuplicate)
+			{
+				const int32 AddedMergedVert = MergedVertices.Add(MoveTemp(BuildVert));
+				VertexInstanceToMerged[VertexInstances[LocalInstance].GetValue()] = AddedMergedVert;
+			}
+		}
+	}
+}
+
+static void InitializeRenderPatternDataFromMeshDescription(const FMeshDescription* const InMeshDescription, const FMeshBuildSettings& BuildSettings, FCollectionClothPatternFacade& RenderPattern)
+{
+	const FMeshDescription* MeshDescription = InMeshDescription;
+	FMeshDescription WritableMeshDescription;
+	if (BuildSettings.bRecomputeNormals || BuildSettings.bRecomputeTangents)
+	{
+		// check if any are invalid
+		bool bHasInvalidNormals, bHasInvalidTangents;
+		FStaticMeshOperations::AreNormalsAndTangentsValid(*InMeshDescription, bHasInvalidNormals, bHasInvalidTangents);
+
+		// if neither are invalid we are not going to recompute
+		if (bHasInvalidNormals || bHasInvalidTangents)
+		{
+			WritableMeshDescription = *InMeshDescription;
+			MeshDescription = &WritableMeshDescription;
+			FStaticMeshAttributes Attributes(WritableMeshDescription);
+			if (!Attributes.GetTriangleNormals().IsValid() || !Attributes.GetTriangleTangents().IsValid())
+			{
+				// If these attributes don't exist, create them and compute their values for each triangle
+				FStaticMeshOperations::ComputeTriangleTangentsAndNormals(WritableMeshDescription);
+			}
+
+			EComputeNTBsFlags ComputeNTBsOptions = EComputeNTBsFlags::BlendOverlappingNormals;
+			ComputeNTBsOptions |= BuildSettings.bRecomputeNormals ? EComputeNTBsFlags::Normals : EComputeNTBsFlags::None;
+			ComputeNTBsOptions |= BuildSettings.bRecomputeTangents ? EComputeNTBsFlags::Tangents : EComputeNTBsFlags::None;
+			ComputeNTBsOptions |= BuildSettings.bUseMikkTSpace ? EComputeNTBsFlags::UseMikkTSpace : EComputeNTBsFlags::None;
+			ComputeNTBsOptions |= BuildSettings.bComputeWeightedNormals ? EComputeNTBsFlags::WeightedNTBs : EComputeNTBsFlags::None;
+			ComputeNTBsOptions |= BuildSettings.bRemoveDegenerates ? EComputeNTBsFlags::IgnoreDegenerateTriangles : EComputeNTBsFlags::None;
+
+			FStaticMeshOperations::ComputeTangentsAndNormals(WritableMeshDescription, ComputeNTBsOptions);
+		}
+	}
+
+	// Merge vertex instances that share the same vertex. These will become the pattern vertices.
+	TArray<FBuildVertex> MergedVertices;
+	TArray<int32> VertexInstanceToMerged;
+	MergeVertexInstances(MeshDescription, MergedVertices, VertexInstanceToMerged);
+
+	RenderPattern.SetNumRenderVertices(MergedVertices.Num());
+	RenderPattern.SetNumRenderFaces(MeshDescription->Triangles().Num());
+
+	// Vertex data (this will MoveTemp stuff out of MergedVertices!!)
+	TArrayView<FVector3f> RenderPosition = RenderPattern.GetRenderPosition();
+	TArrayView<FVector3f> RenderNormal = RenderPattern.GetRenderNormal();
+	TArrayView<FVector3f> RenderTangentU = RenderPattern.GetRenderTangentU();
+	TArrayView<FVector3f> RenderTangentV = RenderPattern.GetRenderTangentV();
+	TArrayView<TArray<FVector2f>> RenderUVs = RenderPattern.GetRenderUVs();
+	TArrayView<FLinearColor> RenderColor = RenderPattern.GetRenderColor();
+	for (int32 VertexIndex = 0; VertexIndex < MergedVertices.Num(); ++VertexIndex)
+	{
+		FBuildVertex& BuildVertex = MergedVertices[VertexIndex];
+		RenderPosition[VertexIndex] = BuildVertex.Position;
+		RenderNormal[VertexIndex] = BuildVertex.Normal;
+		RenderTangentU[VertexIndex] = BuildVertex.Tangent;
+		RenderTangentV[VertexIndex] = FVector3f::CrossProduct(BuildVertex.Normal, BuildVertex.Tangent).GetSafeNormal() * BuildVertex.BiNormalSign;
+		RenderUVs[VertexIndex] = MoveTemp(BuildVertex.UVs);
+		RenderColor[VertexIndex] = FLinearColor(BuildVertex.Color);
+	}
+
+	// Face data
+	TArrayView<FIntVector3> RenderIndices = RenderPattern.GetRenderIndices();
+	TArrayView<int32> RenderMaterialIndex = RenderPattern.GetRenderMaterialIndex();
+	int32 FaceIndex = 0;
+	for (const FTriangleID TriangleID : MeshDescription->Triangles().GetElementIDs())
+	{
+		const FPolygonGroupID PolygonGroupID = MeshDescription->GetTrianglePolygonGroup(TriangleID);
+		RenderMaterialIndex[FaceIndex] = PolygonGroupID.GetValue();
+
+		const TConstArrayView<FVertexInstanceID> VertexInstances = MeshDescription->GetTriangleVertexInstances(TriangleID);
+		check(VertexInstances.Num() == 3);
+		check(MergedVertices.IsValidIndex(VertexInstanceToMerged[VertexInstances[0].GetValue()]));
+		check(MergedVertices.IsValidIndex(VertexInstanceToMerged[VertexInstances[1].GetValue()]));
+		check(MergedVertices.IsValidIndex(VertexInstanceToMerged[VertexInstances[2].GetValue()]));
+		RenderIndices[FaceIndex] = FIntVector3(VertexInstanceToMerged[VertexInstances[0].GetValue()],
+			VertexInstanceToMerged[VertexInstances[1].GetValue()], VertexInstanceToMerged[VertexInstances[2].GetValue()]);
+		++FaceIndex;
+	}
+}
+
+} // namespace UE::Chaos::ClothAsset::Private
+
 FChaosClothAssetStaticMeshImportNode::FChaosClothAssetStaticMeshImportNode(const Dataflow::FNodeParameters& InParam, FGuid InGuid)
 	: FDataflowNode(InParam, InGuid)
 {
@@ -25,8 +230,7 @@ void FChaosClothAssetStaticMeshImportNode::Evaluate(Dataflow::FContext& Context,
 
 	if (Out->IsA<FManagedArrayCollection>(&Collection))
 	{
-
-		if (StaticMesh)
+		if (StaticMesh && (bImportAsSimMesh || bImportAsRenderMesh))
 		{
 			// Evaluate in collection
 			const TSharedRef<FManagedArrayCollection> ClothCollection = MakeShared<FManagedArrayCollection>();
@@ -36,24 +240,47 @@ void FChaosClothAssetStaticMeshImportNode::Evaluate(Dataflow::FContext& Context,
 			const int32 NumLods = StaticMesh->GetNumSourceModels();
 			for (int32 LodIndex = 0; LodIndex < NumLods; ++LodIndex)
 			{
-				const FMeshDescription* MeshDescription = StaticMesh->GetMeshDescription(LodIndex);
+				const FMeshDescription* const MeshDescription = StaticMesh->GetMeshDescription(LodIndex);
 				check(MeshDescription);
 
-				FMeshDescriptionToDynamicMesh Converter;
-				Converter.bPrintDebugMessages = true; // TODO: testing only
-				Converter.bEnableOutputGroups = false;
-				Converter.bVIDsFromNonManifoldMeshDescriptionAttr = true;
-				UE::Geometry::FDynamicMesh3 DynamicMesh;
+				FCollectionClothLodFacade ClothLodFacade = ClothFacade.AddGetLod();
+				ClothLodFacade.Reset();
 
-				Converter.Convert(MeshDescription, DynamicMesh);
+				if (bImportAsSimMesh)
+				{
+					FMeshDescriptionToDynamicMesh Converter;
+					Converter.bPrintDebugMessages = false;
+					Converter.bEnableOutputGroups = false;
+					Converter.bVIDsFromNonManifoldMeshDescriptionAttr = true;
+					UE::Geometry::FDynamicMesh3 DynamicMesh;
 
-				FCollectionClothLodFacade ClothLodFacade = ClothFacade.AddGetLod();
-				ClothLodFacade.Initialize(DynamicMesh, UVChannel, UVScale);
+					Converter.Convert(MeshDescription, DynamicMesh);
+					ClothLodFacade.Initialize(DynamicMesh, UVChannel, UVScale, bUseUVIslands);
+				}
+				if (bImportAsRenderMesh)
+				{
+					// Add render data into a single pattern for now
+					FCollectionClothPatternFacade RenderPattern = ClothLodFacade.AddGetPattern();
+					UE::Chaos::ClothAsset::Private::InitializeRenderPatternDataFromMeshDescription(MeshDescription, StaticMesh->GetSourceModel(LodIndex).BuildSettings, RenderPattern);
+
+					// Set material path names
+					const TArray<FStaticMaterial>& StaticMaterials = StaticMesh->GetStaticMaterials();
+					ClothLodFacade.SetNumMaterials(StaticMaterials.Num());
+					TArrayView<FString> RenderMaterialPathNames = ClothLodFacade.GetRenderMaterialPathName();
+					for(int32 MaterialIndex = 0; MaterialIndex < StaticMaterials.Num(); ++MaterialIndex)
+					{
+						if (ensure(StaticMaterials[MaterialIndex].MaterialInterface))
+						{
+							RenderMaterialPathNames[MaterialIndex] = StaticMaterials[MaterialIndex].MaterialInterface->GetPathName();
+						}
+					}
+				}
 
 				// Set a default skeleton
 				static const TCHAR* const DefaultSkeletonPathName = TEXT("/Engine/EditorMeshes/SkeletalMesh/DefaultSkeletalMesh_Skeleton.DefaultSkeletalMesh_Skeleton");
 				ClothLodFacade.SetSkeletonAssetPathName(DefaultSkeletonPathName);
 			}
+
 			// Make sure that whatever happens there is always at least one empty LOD to avoid crashing the render data
 			if (ClothFacade.GetNumLods() < 1)
 			{
@@ -61,8 +288,6 @@ void FChaosClothAssetStaticMeshImportNode::Evaluate(Dataflow::FContext& Context,
 			}
 			SetValue<FManagedArrayCollection>(Context, *ClothCollection, &Collection);
 		}
-
-
 	}
 }
 
diff --git a/Engine/Plugins/Experimental/ChaosClothAssetEditor/Source/ChaosClothAssetDataflowNodes/Public/ChaosClothAsset/StaticMeshImportNode.h b/Engine/Plugins/Experimental/ChaosClothAssetEditor/Source/ChaosClothAssetDataflowNodes/Public/ChaosClothAsset/StaticMeshImportNode.h
--- a/Engine/Plugins/Experimental/ChaosClothAssetEditor/Source/ChaosClothAssetDataflowNodes/Public/ChaosClothAsset/StaticMeshImportNode.h
+++ b/Engine/Plugins/Experimental/ChaosClothAssetEditor/Source/ChaosClothAssetDataflowNodes/Public/ChaosClothAsset/StaticMeshImportNode.h
@@ -18,15 +18,30 @@ public:
 	UPROPERTY(Meta = (DataflowOutput, DisplayName = "Collection"))
 	FManagedArrayCollection Collection;
 
+	/* The Static Mesh to import from */
 	UPROPERTY(EditAnywhere, Category = "Static Mesh Import")
 	TObjectPtr<const UStaticMesh> StaticMesh;
 
+	/* Import static mesh data as a simulation mesh data.*/
 	UPROPERTY(EditAnywhere, Category = "Static Mesh Import")
-	int32 UVChannel = INDEX_NONE;
+	bool bImportAsSimMesh = true;
 
-	UPROPERTY(EditAnywhere, Category = "Static Mesh Import", Meta = (AllowPreserveRatio))
+	/* Import static mesh data as render mesh data.*/
+	UPROPERTY(EditAnywhere, Category = "Static Mesh Import")
+	bool bImportAsRenderMesh = true;
+
+	/* UV Channel used to populate Sim Mesh positions */
+	UPROPERTY(EditAnywhere, Category = "Static Mesh Import", meta = (EditCondition = "bImportAsSimMesh"))
+	int32 UVChannel = 0;
+
+	/* Apply this scale to the UVs when populating Sim Mesh positions. */
+	UPROPERTY(EditAnywhere, Category = "Static Mesh Import", Meta = (AllowPreserveRatio, EditCondition = "bImportAsSimMesh && UVChannel != INDEX_NONE"))
 	FVector2f UVScale = { 1.f, 1.f };
 
+	/* Use Sim UV Islands to generate separate patterns (Note: render mesh will always be a single pattern)*/
+	UPROPERTY(EditAnywhere, Category = "Static Mesh", meta = (EditCondition = "bImportAsSimMesh && UVChannel != INDEX_NONE"))
+	bool bUseUVIslands = true;
+
 	FChaosClothAssetStaticMeshImportNode(const Dataflow::FNodeParameters& InParam, FGuid InGuid = FGuid::NewGuid());
 
 	virtual void Evaluate(Dataflow::FContext& Context, const FDataflowOutput* Out) const override;
