Adding a two-sided vertex color material. Same as MeshVertexColorMaterial but with "Two Sided" checked.
Use the new two-sided material in the cloth weight map paint tool so that we can paint on the "back" side of the 2D panels

diff --git a/Engine/Plugins/Experimental/ChaosClothAssetEditor/Source/ChaosClothAssetEditorTools/Private/ChaosClothAsset/ClothWeightMapPaintTool.cpp b/Engine/Plugins/Experimental/ChaosClothAssetEditor/Source/ChaosClothAssetEditorTools/Private/ChaosClothAsset/ClothWeightMapPaintTool.cpp
--- a/Engine/Plugins/Experimental/ChaosClothAssetEditor/Source/ChaosClothAssetEditorTools/Private/ChaosClothAsset/ClothWeightMapPaintTool.cpp
+++ b/Engine/Plugins/Experimental/ChaosClothAssetEditor/Source/ChaosClothAssetEditorTools/Private/ChaosClothAsset/ClothWeightMapPaintTool.cpp
@@ -1277,6 +1277,26 @@ void UClothEditorWeightMapPaintTool::Render(IToolsContextRenderAPI* RenderAPI)
 	}
 }
 
+
+void UClothEditorWeightMapPaintTool::UpdateMaterialMode(EMeshEditingMaterialModes MaterialMode)
+{
+	if (MaterialMode == EMeshEditingMaterialModes::VertexColor)
+	{
+		constexpr bool bUseTwoSidedMaterial = true;
+		ActiveOverrideMaterial = ToolSetupUtil::GetVertexColorMaterial(GetToolManager(), bUseTwoSidedMaterial);
+		if (ensure(ActiveOverrideMaterial != nullptr))
+		{
+			GetSculptMeshComponent()->SetOverrideRenderMaterial(ActiveOverrideMaterial);
+			ActiveOverrideMaterial->SetScalarParameterValue(TEXT("FlatShading"), (ViewProperties->bFlatShading) ? 1.0f : 0.0f);
+		}
+		GetSculptMeshComponent()->SetShadowsEnabled(false);
+	}
+	else
+	{
+		UMeshSculptToolBase::UpdateMaterialMode(MaterialMode);
+	}
+}
+
 void UClothEditorWeightMapPaintTool::UpdateStampPendingState()
 {
 	if (InStroke() == false) return;
diff --git a/Engine/Plugins/Experimental/ChaosClothAssetEditor/Source/ChaosClothAssetEditorTools/Public/ChaosClothAsset/ClothWeightMapPaintTool.h b/Engine/Plugins/Experimental/ChaosClothAssetEditor/Source/ChaosClothAssetEditorTools/Public/ChaosClothAsset/ClothWeightMapPaintTool.h
--- a/Engine/Plugins/Experimental/ChaosClothAssetEditor/Source/ChaosClothAssetEditorTools/Public/ChaosClothAsset/ClothWeightMapPaintTool.h
+++ b/Engine/Plugins/Experimental/ChaosClothAssetEditor/Source/ChaosClothAssetEditorTools/Public/ChaosClothAsset/ClothWeightMapPaintTool.h
@@ -251,6 +251,7 @@ public:
 	virtual void Shutdown(EToolShutdownType ShutdownType) override;
 	virtual void DrawHUD(FCanvas* Canvas, IToolsContextRenderAPI* RenderAPI) override;
 	virtual void Render(IToolsContextRenderAPI* RenderAPI) override;
+	virtual void UpdateMaterialMode(EMeshEditingMaterialModes MaterialMode) override;
 
 	virtual void UpdateStampPendingState() override;
 
diff --git a/Engine/Plugins/Runtime/MeshModelingToolset/Source/ModelingComponents/Private/ToolSetupUtil.cpp b/Engine/Plugins/Runtime/MeshModelingToolset/Source/ModelingComponents/Private/ToolSetupUtil.cpp
--- a/Engine/Plugins/Runtime/MeshModelingToolset/Source/ModelingComponents/Private/ToolSetupUtil.cpp
+++ b/Engine/Plugins/Runtime/MeshModelingToolset/Source/ModelingComponents/Private/ToolSetupUtil.cpp
@@ -30,9 +30,13 @@ UMaterialInterface* ToolSetupUtil::GetDefaultMaterial(UInteractiveToolManager* T
 }
 
 
-UMaterialInstanceDynamic* ToolSetupUtil::GetVertexColorMaterial(UInteractiveToolManager* ToolManager)
+UMaterialInstanceDynamic* ToolSetupUtil::GetVertexColorMaterial(UInteractiveToolManager* ToolManager, bool bTwoSided)
 {
-	UMaterial* Material = LoadObject<UMaterial>(nullptr, TEXT("/MeshModelingToolsetExp/Materials/MeshVertexColorMaterial"));
+	// Unfortunately the two-sided flag is not something that we can give as a runtime 
+	// parameter, so we need separate versions of the material.
+	UMaterial* Material = bTwoSided ? LoadObject<UMaterial>(nullptr, TEXT("/MeshModelingToolsetExp/Materials/MeshVertexColorMaterialTwoSided"))
+		: LoadObject<UMaterial>(nullptr, TEXT("/MeshModelingToolsetExp/Materials/MeshVertexColorMaterial"));
+
 	if (Material != nullptr)
 	{
 		UMaterialInstanceDynamic* MatInstance = UMaterialInstanceDynamic::Create(Material, ToolManager);
diff --git a/Engine/Plugins/Runtime/MeshModelingToolset/Source/ModelingComponents/Public/ToolSetupUtil.h b/Engine/Plugins/Runtime/MeshModelingToolset/Source/ModelingComponents/Public/ToolSetupUtil.h
--- a/Engine/Plugins/Runtime/MeshModelingToolset/Source/ModelingComponents/Public/ToolSetupUtil.h
+++ b/Engine/Plugins/Runtime/MeshModelingToolset/Source/ModelingComponents/Public/ToolSetupUtil.h
@@ -34,7 +34,7 @@ namespace ToolSetupUtil
 	/**
 	 * @return configurable vertex color material
 	 */
-	MODELINGCOMPONENTS_API UMaterialInstanceDynamic* GetVertexColorMaterial(UInteractiveToolManager* ToolManager);
+	MODELINGCOMPONENTS_API UMaterialInstanceDynamic* GetVertexColorMaterial(UInteractiveToolManager* ToolManager, bool bTwoSided = false);
 
 
 	/**
