Cloth Editor: restore the "false ortho" camera for 2D panel view. When switching to 2D mode, position and lock the perspective camera so that +y corresponds to up in screen space, and +x corresponds to right. Use 2D Scroll and MouseWheelZoom behaviors to mimic the ortho controls. (Partial reversion of 22508058)

Also:
- Don't refocus on the mesh when changing node selection in the Dataflow editor
- Auto-select meshes in the Construction viewport when changing node selection

diff --git a/Engine/Plugins/Experimental/ChaosClothAssetEditor/Source/ChaosClothAssetEditor/Private/ChaosClothAsset/ClothEditorMode.cpp b/Engine/Plugins/Experimental/ChaosClothAssetEditor/Source/ChaosClothAssetEditor/Private/ChaosClothAsset/ClothEditorMode.cpp
--- a/Engine/Plugins/Experimental/ChaosClothAssetEditor/Source/ChaosClothAssetEditor/Private/ChaosClothAsset/ClothEditorMode.cpp
+++ b/Engine/Plugins/Experimental/ChaosClothAssetEditor/Source/ChaosClothAssetEditor/Private/ChaosClothAsset/ClothEditorMode.cpp
@@ -398,6 +398,13 @@ void UChaosClothAssetEditorMode::SetSelectedClothCollection(TSharedPtr<FManagedA
 {
 	SelectedClothCollection = Collection;
 	ReinitializeDynamicMeshComponents();
+
+	if (bFirstClothCollection && Collection)
+	{
+		// refocus the viewport if this is the first time a cloth collection has been set
+		RefocusRestSpaceViewportClient();
+		bFirstClothCollection = false;
+	}
 }
 
 TSharedPtr<FManagedArrayCollection> UChaosClothAssetEditorMode::GetClothCollection()
@@ -414,25 +421,16 @@ void UChaosClothAssetEditorMode::ReinitializeDynamicMeshComponents()
 	{
 		MeshComponent.SelectionOverrideDelegate = UPrimitiveComponent::FSelectionOverride::CreateUObject(this, &UChaosClothAssetEditorMode::IsComponentSelected);
 
-		UMaterialInterface* const Material = ToolSetupUtil::GetDefaultSculptMaterial(GetToolManager());
-		UMaterialInstanceDynamic* const MatInstance = UMaterialInstanceDynamic::Create(Material, GetToolManager());
-		MatInstance->SetVectorParameterValue(TEXT("Color"), FLinearColor::Gray);
-		MatInstance->TwoSided = 1;
-		MeshComponent.SetMaterial(0, MatInstance);
-
-		if (bPattern2DMode)
-		{
-			// The LVT_OrthoXY viewport transform results in the positive y axis pointing down in screen space. The following transform rotates 
-			// the mesh around the z axis so that increasing y value in the cloth asset corresponds to "up" in screen space
-			MeshComponent.SetWorldRotation(FQuat(FVector(0, 0, 1), FMathd::Pi));
-		}
+		UMaterialInterface* const Material = bPattern2DMode ?
+			ToolSetupUtil::GetCustomTwoSidedDepthOffsetMaterial(GetToolManager(), FLinearColor{0.6, 0.6, 0.6}, 0.0) :
+			ToolSetupUtil::GetDefaultSculptMaterial(GetToolManager());
 
+		MeshComponent.SetMaterial(0, Material);
 		MeshComponent.RegisterComponentWithWorld(this->GetWorld());
 	};
 	
 	// Clean up any existing DynamicMeshComponents
 	// Save indices of selected mesh components
-	TArray<int32> PreviouslySelectedDynamicMeshComponents;
 	const int32 TotalNumberExistingDynamicMeshComponents = DynamicMeshComponents.Num();
 
 	USelection* SelectedComponents = GetModeManager()->GetSelectedComponents();
@@ -446,8 +444,8 @@ void UChaosClothAssetEditorMode::ReinitializeDynamicMeshComponents()
 
 		if (SelectedComponents->IsSelected(DynamicMeshComp))
 		{
-			PreviouslySelectedDynamicMeshComponents.Add(DynamicMeshComponentIndex);
 			SelectedComponents->Deselect(DynamicMeshComp);
+			DynamicMeshComp->PushSelectionToProxy();
 		}
 	}
 
@@ -557,12 +555,7 @@ void UChaosClothAssetEditorMode::ReinitializeDynamicMeshComponents()
 	{
 		// Set up the wireframe display of the rest space mesh.
 		TObjectPtr<UMeshElementsVisualizer> WireframeDisplay = NewObject<UMeshElementsVisualizer>(this);
-
-		// The LVT_OrthoXY viewport transform results in the positive y axis pointing down in screen space. The following transform rotates 
-		// the mesh around the z axis so that increasing y value in the cloth asset corresponds to "up" in screen space
-		const FTransform WireframeTransform = bPattern2DMode ? FTransform(FQuat(FVector(0, 0, 1), FMathd::Pi)) : FTransform::Identity;
-
-		WireframeDisplay->CreateInWorld(GetWorld(), WireframeTransform);
+		WireframeDisplay->CreateInWorld(GetWorld(), FTransform::Identity);
 
 		WireframeDisplay->Settings->DepthBias = 2.0;
 		WireframeDisplay->Settings->bAdjustDepthBiasUsingMeshSize = false;
@@ -615,26 +608,19 @@ void UChaosClothAssetEditorMode::ReinitializeDynamicMeshComponents()
 
 	}
 
-	const bool bNumberDynamicMeshesUnchanged = (TotalNumberExistingDynamicMeshComponents == DynamicMeshComponents.Num());
-	if (bNumberDynamicMeshesUnchanged && PreviouslySelectedDynamicMeshComponents.Num() > 0)
+	SelectedComponents->DeselectAll();
+	SelectedComponents->BeginBatchSelectOperation();
+	for (TObjectPtr<UDynamicMeshComponent> RestSpaceMeshComponent : DynamicMeshComponents)
 	{
-		SelectedComponents->Modify();
-		SelectedComponents->BeginBatchSelectOperation();
-		for (int32 DynamicMeshComponentIndex : PreviouslySelectedDynamicMeshComponents)
-		{
-			if (DynamicMeshComponentIndex < DynamicMeshComponents.Num())
-			{
-				SelectedComponents->Select(DynamicMeshComponents[DynamicMeshComponentIndex]);
-			}
-		}
-		SelectedComponents->EndBatchSelectOperation();
+		SelectedComponents->Select(RestSpaceMeshComponent);
+		RestSpaceMeshComponent->PushSelectionToProxy();
 	}
-
+	SelectedComponents->EndBatchSelectOperation();
 }
 
 void UChaosClothAssetEditorMode::RefocusRestSpaceViewportClient()
 {
-	TSharedPtr<FEditorViewportClient, ESPMode::ThreadSafe> PinnedVC = RestSpaceViewportClient.Pin();
+	TSharedPtr<UE::Chaos::ClothAsset::FChaosClothEditorRestSpaceViewportClient, ESPMode::ThreadSafe> PinnedVC = RestSpaceViewportClient.Pin();
 	if (PinnedVC.IsValid())
 	{
 		// This will happen in FocusViewportOnBox anyways; do it now to get a consistent end result
@@ -644,7 +630,7 @@ void UChaosClothAssetEditorMode::RefocusRestSpaceViewportClient()
 		if (bPattern2DMode)
 		{
 			// 2D pattern
-			PinnedVC->SetInitialViewTransform(ELevelViewportType::LVT_OrthoXY, FVector(0, 0, 0), FRotator(0, 0, 0), DEFAULT_ORTHOZOOM);
+			PinnedVC->SetInitialViewTransform(ELevelViewportType::LVT_Perspective, FVector(0, 0, -100), FRotator(90, -90, 0), DEFAULT_ORTHOZOOM);
 		}
 		else
 		{
@@ -654,6 +640,9 @@ void UChaosClothAssetEditorMode::RefocusRestSpaceViewportClient()
 
 		constexpr bool bInstant = true;
 		PinnedVC->FocusViewportOnBox(SceneBounds, bInstant);
+
+		// Recompute near/far clip planes
+		PinnedVC->Set2DMode(bPattern2DMode);
 	}
 }
 
diff --git a/Engine/Plugins/Experimental/ChaosClothAssetEditor/Source/ChaosClothAssetEditor/Private/ChaosClothAsset/ClothEditorMode.h b/Engine/Plugins/Experimental/ChaosClothAssetEditor/Source/ChaosClothAssetEditor/Private/ChaosClothAsset/ClothEditorMode.h
--- a/Engine/Plugins/Experimental/ChaosClothAssetEditor/Source/ChaosClothAssetEditor/Private/ChaosClothAsset/ClothEditorMode.h
+++ b/Engine/Plugins/Experimental/ChaosClothAssetEditor/Source/ChaosClothAssetEditor/Private/ChaosClothAsset/ClothEditorMode.h
@@ -175,8 +175,12 @@ private:
 	// Whether we can switch between 2D and 3D rest configuration
 	bool bCanTogglePattern2DMode = true;
 
+	// Whether a cloth collection has ever been set on this Mode object
+	bool bFirstClothCollection = true;
+
 	// Whether the rest space viewport should focus on the rest space mesh on the next tick
 	bool bShouldFocusRestSpaceView = true;
+
 	void RestSpaceViewportResized(FViewport* RestspaceViewport, uint32 Unused);
 
 	// Whether to combine all patterns into a single DynamicMeshComponent, or have separate components for each pattern
diff --git a/Engine/Plugins/Experimental/ChaosClothAssetEditor/Source/ChaosClothAssetEditor/Private/ChaosClothAsset/ClothEditorRestSpaceViewportClient.cpp b/Engine/Plugins/Experimental/ChaosClothAssetEditor/Source/ChaosClothAssetEditor/Private/ChaosClothAsset/ClothEditorRestSpaceViewportClient.cpp
--- a/Engine/Plugins/Experimental/ChaosClothAssetEditor/Source/ChaosClothAssetEditor/Private/ChaosClothAsset/ClothEditorRestSpaceViewportClient.cpp
+++ b/Engine/Plugins/Experimental/ChaosClothAssetEditor/Source/ChaosClothAssetEditor/Private/ChaosClothAsset/ClothEditorRestSpaceViewportClient.cpp
@@ -2,46 +2,100 @@
 
 #include "ChaosClothAsset/ClothEditorRestSpaceViewportClient.h"
 
+#include "BaseBehaviors/ClickDragBehavior.h"
+#include "BaseBehaviors/MouseWheelBehavior.h"
 #include "EngineUtils.h"
 #include "EditorModeManager.h"
 #include "Engine/Selection.h"
+#include "Behaviors/2DViewportBehaviorTargets.h"
 #include "EdModeInteractiveToolsContext.h"
 #include "CameraController.h"
 #include "Components/DynamicMeshComponent.h"
 #include "Framework/Application/SlateApplication.h"
+#include "SceneView.h"
 
 namespace UE::Chaos::ClothAsset
 {
+
 FChaosClothEditorRestSpaceViewportClient::FChaosClothEditorRestSpaceViewportClient(FEditorModeTools* InModeTools,
 	FPreviewScene* InPreviewScene,
 	const TWeakPtr<SEditorViewport>& InEditorViewportWidget) :
 	FEditorViewportClient(InModeTools, InPreviewScene, InEditorViewportWidget)
 {
-	EngineShowFlags.Grid = true;
-	DrawHelper.bDrawGrid = true;
-
-	bDrawAxes = false;
-
-	// Don't automatically switch to wireframe rendering in ortho mode
-	SetViewModes(EViewModeIndex::VMI_Lit, EViewModeIndex::VMI_Lit);
+	OverrideNearClipPlane(UE_KINDA_SMALL_NUMBER);
+	OverrideFarClipPlane(0);
+
+	BehaviorSet = NewObject<UInputBehaviorSet>();
+
+	// We'll have the priority of our viewport manipulation behaviors be lower (i.e. higher
+	// numerically) than both the gizmo default and the tool default.
+	constexpr int ViewportBehaviorPriority = 150;
+
+	ScrollBehaviorTarget = MakeUnique<FEditor2DScrollBehaviorTarget>(this);
+	UClickDragInputBehavior* const ScrollBehavior = NewObject<UClickDragInputBehavior>();
+	ScrollBehavior->Initialize(ScrollBehaviorTarget.Get());
+	ScrollBehavior->SetDefaultPriority(ViewportBehaviorPriority);
+	ScrollBehavior->SetUseRightMouseButton();
+	BehaviorSet->Add(ScrollBehavior);
+	BehaviorsFor2DMode.Add(ScrollBehavior);
+
+	ZoomBehaviorTarget = MakeUnique<FEditor2DMouseWheelZoomBehaviorTarget>(this);
+	constexpr float CameraFarPlaneWorldZ = -10.0f;
+	constexpr float CameraNearPlaneProportionZ = 0.8f;
+	ZoomBehaviorTarget->SetCameraFarPlaneWorldZ(CameraFarPlaneWorldZ);
+	ZoomBehaviorTarget->SetCameraNearPlaneProportionZ(CameraNearPlaneProportionZ);
+	ZoomBehaviorTarget->SetZoomLimits(0.001, 100000);
+	UMouseWheelInputBehavior* const ZoomBehavior = NewObject<UMouseWheelInputBehavior>();
+	ZoomBehavior->Initialize(ZoomBehaviorTarget.Get());
+	ZoomBehavior->SetDefaultPriority(ViewportBehaviorPriority);
+	BehaviorSet->Add(ZoomBehavior);
+	BehaviorsFor2DMode.Add(ZoomBehavior);
 
 	EngineShowFlags.SetSelectionOutline(true);
+	ModeTools->GetInteractiveToolsContext()->InputRouter->RegisterSource(this);
 }
 
 void FChaosClothEditorRestSpaceViewportClient::Set2DMode(bool In2DMode)
 {
 	b2DMode = In2DMode;
 
+	BehaviorSet->RemoveAll();
+	ModeTools->GetInteractiveToolsContext()->InputRouter->DeregisterSource(this);
+
 	if (b2DMode)
 	{
-		SetViewportType(ELevelViewportType::LVT_OrthoXY);
+		for (UInputBehavior* const Behavior : BehaviorsFor2DMode)
+		{
+			BehaviorSet->Add(Behavior);
+		}
+		ModeTools->GetInteractiveToolsContext()->InputRouter->RegisterSource(this);
+
+		const double AbsZ = FMath::Abs(ViewTransformPerspective.GetLocation().Z);
+		constexpr double CameraFarPlaneWorldZ = -10.0;
+		constexpr double CameraNearPlaneProportionZ = 0.8;
+		OverrideFarClipPlane(static_cast<float>(AbsZ - CameraFarPlaneWorldZ));
+		OverrideNearClipPlane(static_cast<float>(AbsZ * (1.0 - CameraNearPlaneProportionZ)));
 	}
 	else
 	{
-		SetViewportType(ELevelViewportType::LVT_Perspective);
+		OverrideFarClipPlane(0);
+		OverrideNearClipPlane(UE_KINDA_SMALL_NUMBER);
 	}
+
+}
+
+const UInputBehaviorSet* FChaosClothEditorRestSpaceViewportClient::GetInputBehaviors() const
+{
+	return BehaviorSet;
 }
 
+// Collects UObjects that we don't want the garbage collecter to clean up
+void FChaosClothEditorRestSpaceViewportClient::AddReferencedObjects(FReferenceCollector& Collector)
+{
+	FEditorViewportClient::AddReferencedObjects(Collector);
+	Collector.AddReferencedObject(BehaviorSet);
+	Collector.AddReferencedObjects(BehaviorsFor2DMode);
+}
 
 bool FChaosClothEditorRestSpaceViewportClient::ShouldOrbitCamera() const
 {
@@ -68,7 +122,79 @@ bool FChaosClothEditorRestSpaceViewportClient::InputKey(const FInputKeyEventArgs
 		}
 	}
 
-	return FEditorViewportClient::InputKey(EventArgs);
+	if (!b2DMode)
+	{
+		return FEditorViewportClient::InputKey(EventArgs);
+	}
+
+
+	// We'll support disabling input like our base class, even if it does not end up being used.
+	if (bDisableInput)
+	{
+		return true;
+	}
+
+	// Our viewport manipulation is placed in the input router that ModeTools manages
+	if (ModeTools->InputKey(this, EventArgs.Viewport, EventArgs.Key, EventArgs.Event))
+	{
+		return true;
+	}
+	
+	//
+	// Recreate the click handling from FEditorViewportClient::Internal_InputKey so that we can still have 
+	// mouse click, keyboard, etc., but skip the code dealing with camera controls
+	//
+
+	bool bHandled = false;
+
+	const FInputEventState InputState(EventArgs.Viewport, EventArgs.Key, EventArgs.Event);
+	const FViewport* const InViewport = EventArgs.Viewport;
+	const EInputEvent Event = EventArgs.Event;
+
+	const bool bWasCursorVisible = InViewport->IsCursorVisible();
+	const bool bWasSoftwareCursorVisible = InViewport->IsSoftwareCursorVisible();
+
+	RequiredCursorVisibiltyAndAppearance.bDontResetCursor = false;
+	UpdateRequiredCursorVisibility();
+
+	if (bWasCursorVisible != RequiredCursorVisibiltyAndAppearance.bHardwareCursorVisible || bWasSoftwareCursorVisible != RequiredCursorVisibiltyAndAppearance.bSoftwareCursorVisible)
+	{
+		bHandled = true;
+	}
+
+	// Compute a view.
+	FSceneViewFamilyContext ViewFamily(FSceneViewFamily::ConstructionValues(
+		InViewport,
+		GetScene(),
+		EngineShowFlags)
+		.SetRealtimeUpdate(IsRealtime()));
+	FSceneView* const View = CalcSceneView(&ViewFamily);
+
+	// Start tracking if any mouse button is down and it was a tracking event (MouseButton/Ctrl/Shift/Alt):
+	if (InputState.IsAnyMouseButtonDown()
+		&& (Event == IE_Pressed || Event == IE_Released)
+		&& (InputState.IsMouseButtonEvent() || InputState.IsCtrlButtonEvent() || InputState.IsAltButtonEvent() || InputState.IsShiftButtonEvent()))
+	{
+		StartTrackingDueToInput(InputState, *View);
+		return true;
+	}
+
+	// If we are tracking and no mouse button is down and this input event released the mouse button stop tracking and process any clicks if necessary
+	if (bIsTracking && !InputState.IsAnyMouseButtonDown() && InputState.IsMouseButtonEvent())
+	{
+		// Handle possible mouse click viewport
+		ProcessClickInViewport(InputState, *View);
+
+		// Stop tracking if no mouse button is down
+		StopTracking();
+
+		bHandled |= true;
+	}
+
+	// apply the visibility and set the cursor positions
+	ApplyRequiredCursorVisibility(true);
+
+	return bHandled;
 }
 
 
diff --git a/Engine/Plugins/Experimental/ChaosClothAssetEditor/Source/ChaosClothAssetEditor/Private/ChaosClothAsset/ClothEditorRestSpaceViewportClient.h b/Engine/Plugins/Experimental/ChaosClothAssetEditor/Source/ChaosClothAssetEditor/Private/ChaosClothAsset/ClothEditorRestSpaceViewportClient.h
--- a/Engine/Plugins/Experimental/ChaosClothAssetEditor/Source/ChaosClothAssetEditor/Private/ChaosClothAsset/ClothEditorRestSpaceViewportClient.h
+++ b/Engine/Plugins/Experimental/ChaosClothAssetEditor/Source/ChaosClothAssetEditor/Private/ChaosClothAsset/ClothEditorRestSpaceViewportClient.h
@@ -11,7 +11,8 @@ class UInputBehaviorSet;
 
 namespace UE::Chaos::ClothAsset
 {
-class CHAOSCLOTHASSETEDITOR_API FChaosClothEditorRestSpaceViewportClient : public FEditorViewportClient
+
+class CHAOSCLOTHASSETEDITOR_API FChaosClothEditorRestSpaceViewportClient : public FEditorViewportClient, public IInputBehaviorSource
 {
 public:
 
@@ -19,7 +20,12 @@ public:
 
 	virtual ~FChaosClothEditorRestSpaceViewportClient() = default;
 
-	// FEditorViewportClient
+	// IInputBehaviorSource
+	virtual const UInputBehaviorSet* GetInputBehaviors() const override;
+
+	// FGCObject
+	virtual void AddReferencedObjects(FReferenceCollector& Collector) override;
+
 	virtual void ProcessClick(FSceneView& View, HHitProxy* HitProxy, FKey Key, EInputEvent Event, uint32 HitX, uint32 HitY) override;
 	virtual bool InputKey(const FInputKeyEventArgs& EventArgs) override;
 
@@ -34,6 +40,13 @@ private:
 
 	bool b2DMode = false;
 
+	TObjectPtr<UInputBehaviorSet> BehaviorSet;
+
+	TArray<TObjectPtr<UInputBehavior>> BehaviorsFor2DMode;
+
+	TUniquePtr<FEditor2DScrollBehaviorTarget> ScrollBehaviorTarget;
+	TUniquePtr<FEditor2DMouseWheelZoomBehaviorTarget> ZoomBehaviorTarget;
+
 	TWeakPtr<FUICommandList> ToolCommandList;
 };
 } // namespace UE::Chaos::ClothAsset
diff --git a/Engine/Plugins/Experimental/ChaosClothAssetEditor/Source/ChaosClothAssetEditor/Private/ChaosClothAsset/ClothEditorToolkit.cpp b/Engine/Plugins/Experimental/ChaosClothAssetEditor/Source/ChaosClothAssetEditor/Private/ChaosClothAsset/ClothEditorToolkit.cpp
--- a/Engine/Plugins/Experimental/ChaosClothAssetEditor/Source/ChaosClothAssetEditor/Private/ChaosClothAsset/ClothEditorToolkit.cpp
+++ b/Engine/Plugins/Experimental/ChaosClothAssetEditor/Source/ChaosClothAssetEditor/Private/ChaosClothAsset/ClothEditorToolkit.cpp
@@ -428,11 +428,17 @@ void FChaosClothAssetEditorToolkit::PostInitAssetEditor()
 
 	SetCommonViewportClientOptions(ViewportClient.Get());
 
+	// Ortho has too many problems with rendering things, unfortunately, so we should use perspective.
+	ViewportClient->SetViewportType(ELevelViewportType::LVT_Perspective);
+
 	// Lit gives us the most options in terms of the materials we can use.
 	ViewportClient->SetViewMode(EViewModeIndex::VMI_Lit);
 
 	UChaosClothAssetEditorMode* const ClothMode = CastChecked<UChaosClothAssetEditorMode>(EditorModeManager->GetActiveScriptableMode(UChaosClothAssetEditorMode::EM_ChaosClothAssetEditorModeId));
-	
+
+	// If exposure isn't set to fixed, it will flash as we stare into the void
+	ViewportClient->ExposureSettings.bFixed = true;
+
 	const TWeakPtr<FViewportClient> WeakViewportClient(ViewportClient);
 	ClothMode->SetRestSpaceViewportClient(StaticCastWeakPtr<FChaosClothEditorRestSpaceViewportClient>(WeakViewportClient));
 
@@ -983,7 +989,6 @@ void FChaosClothAssetEditorToolkit::OnNodeSelectionChanged(const TSet<UObject*>&
 	if (ClothMode)
 	{
 		ClothMode->SetSelectedClothCollection(Collection);
-		ClothMode->RefocusRestSpaceViewportClient();
 	}
 }
 
