Chaos Cloth - Fixed ensure on first creation of a static import node.

#rnx

diff --git a/Engine/Plugins/Experimental/ChaosClothAssetEditor/Source/ChaosClothAssetDataflowNodes/Private/ChaosClothAsset/StaticMeshImportNode.cpp b/Engine/Plugins/Experimental/ChaosClothAssetEditor/Source/ChaosClothAssetDataflowNodes/Private/ChaosClothAsset/StaticMeshImportNode.cpp
--- a/Engine/Plugins/Experimental/ChaosClothAssetEditor/Source/ChaosClothAssetDataflowNodes/Private/ChaosClothAsset/StaticMeshImportNode.cpp
+++ b/Engine/Plugins/Experimental/ChaosClothAssetEditor/Source/ChaosClothAssetDataflowNodes/Private/ChaosClothAsset/StaticMeshImportNode.cpp
@@ -230,13 +230,15 @@ void FChaosClothAssetStaticMeshImportNode::Evaluate(Dataflow::FContext& Context,
 
 	if (Out->IsA<FManagedArrayCollection>(&Collection))
 	{
+		static const TCHAR* const DefaultSkeletonPathName = TEXT("/Engine/EditorMeshes/SkeletalMesh/DefaultSkeletalMesh_Skeleton.DefaultSkeletalMesh_Skeleton");
+
+		// Evaluate out collection
+		const TSharedRef<FManagedArrayCollection> ClothCollection = MakeShared<FManagedArrayCollection>();
+		FCollectionClothFacade ClothFacade(ClothCollection);
+		ClothFacade.DefineSchema();
+
 		if (StaticMesh && (bImportAsSimMesh || bImportAsRenderMesh))
 		{
-			// Evaluate in collection
-			const TSharedRef<FManagedArrayCollection> ClothCollection = MakeShared<FManagedArrayCollection>();
-			FCollectionClothFacade ClothFacade(ClothCollection);
-			ClothFacade.DefineSchema();
-
 			const int32 NumLods = StaticMesh->GetNumSourceModels();
 			for (int32 LodIndex = 0; LodIndex < NumLods; ++LodIndex)
 			{
@@ -277,17 +279,16 @@ void FChaosClothAssetStaticMeshImportNode::Evaluate(Dataflow::FContext& Context,
 				}
 
 				// Set a default skeleton
-				static const TCHAR* const DefaultSkeletonPathName = TEXT("/Engine/EditorMeshes/SkeletalMesh/DefaultSkeletalMesh_Skeleton.DefaultSkeletalMesh_Skeleton");
 				ClothLodFacade.SetSkeletonAssetPathName(DefaultSkeletonPathName);
 			}
+		}
 
-			// Make sure that whatever happens there is always at least one empty LOD to avoid crashing the render data
-			if (ClothFacade.GetNumLods() < 1)
-			{
-				ClothFacade.AddLod();
-			}
-			SetValue<FManagedArrayCollection>(Context, *ClothCollection, &Collection);
+		// Make sure that whatever happens there is always at least one empty LOD to avoid crashing the render data
+		if (ClothFacade.GetNumLods() < 1)
+		{
+			ClothFacade.AddGetLod().SetSkeletonAssetPathName(DefaultSkeletonPathName);
 		}
+		SetValue<FManagedArrayCollection>(Context, *ClothCollection, &Collection);
 	}
 }
 
