Chaos Cloth - Empty cloth asset material list on new terminal node evaluations.

#rnx

diff --git a/Engine/Plugins/Experimental/ChaosClothAssetEditor/Source/ChaosClothAssetDataflowNodes/Private/ChaosClothAsset/TerminalNode.cpp b/Engine/Plugins/Experimental/ChaosClothAssetEditor/Source/ChaosClothAssetDataflowNodes/Private/ChaosClothAsset/TerminalNode.cpp
--- a/Engine/Plugins/Experimental/ChaosClothAssetEditor/Source/ChaosClothAssetDataflowNodes/Private/ChaosClothAsset/TerminalNode.cpp
+++ b/Engine/Plugins/Experimental/ChaosClothAssetEditor/Source/ChaosClothAssetDataflowNodes/Private/ChaosClothAsset/TerminalNode.cpp
@@ -30,6 +30,10 @@ void FChaosClothAssetTerminalNode::SetAssetValue(TObjectPtr<UObject> Asset, Data
 		const TSharedPtr<FManagedArrayCollection> ClothCollection = ClothAsset->GetClothCollection();
 		ClothCollection->Reset();
 
+		// Reset the asset's material list
+		ClothAsset->GetMaterials().Reset();
+
+		// Define the final container's collection schema
 		FCollectionClothFacade ClothFacade(ClothCollection);
 		ClothFacade.DefineSchema();
 
