ChaosClothAsset move StaticMeshImportNode's bUseUVIslands to the correct category.

diff --git a/Engine/Plugins/Experimental/ChaosClothAssetEditor/Source/ChaosClothAssetDataflowNodes/Public/ChaosClothAsset/StaticMeshImportNode.h b/Engine/Plugins/Experimental/ChaosClothAssetEditor/Source/ChaosClothAssetDataflowNodes/Public/ChaosClothAsset/StaticMeshImportNode.h
--- a/Engine/Plugins/Experimental/ChaosClothAssetEditor/Source/ChaosClothAssetDataflowNodes/Public/ChaosClothAsset/StaticMeshImportNode.h
+++ b/Engine/Plugins/Experimental/ChaosClothAssetEditor/Source/ChaosClothAssetDataflowNodes/Public/ChaosClothAsset/StaticMeshImportNode.h
@@ -39,7 +39,7 @@ public:
 	FVector2f UVScale = { 1.f, 1.f };
 
 	/* Use Sim UV Islands to generate separate patterns (Note: render mesh will always be a single pattern)*/
-	UPROPERTY(EditAnywhere, Category = "Static Mesh", meta = (EditCondition = "bImportAsSimMesh && UVChannel != INDEX_NONE"))
+	UPROPERTY(EditAnywhere, Category = "Static Mesh Import", meta = (EditCondition = "bImportAsSimMesh && UVChannel != INDEX_NONE"))
 	bool bUseUVIslands = true;
 
 	FChaosClothAssetStaticMeshImportNode(const Dataflow::FNodeParameters& InParam, FGuid InGuid = FGuid::NewGuid());
