CIS - Fixed duplicate LOC names om SimulationBaseConfigNode.cpp.

#rnx

diff --git a/Engine/Plugins/Experimental/ChaosClothAssetEditor/Source/ChaosClothAssetDataflowNodes/Private/ChaosClothAsset/SimulationBaseConfigNode.cpp b/Engine/Plugins/Experimental/ChaosClothAssetEditor/Source/ChaosClothAssetDataflowNodes/Private/ChaosClothAsset/SimulationBaseConfigNode.cpp
--- a/Engine/Plugins/Experimental/ChaosClothAssetEditor/Source/ChaosClothAssetDataflowNodes/Private/ChaosClothAsset/SimulationBaseConfigNode.cpp
+++ b/Engine/Plugins/Experimental/ChaosClothAssetEditor/Source/ChaosClothAssetDataflowNodes/Private/ChaosClothAsset/SimulationBaseConfigNode.cpp
@@ -28,7 +28,7 @@ namespace UE::Chaos::ClothAsset::Private
 
 		const FText Message = FText::Format(
 			LOCTEXT(
-				"DuplicateProperty",
+				"SimilarProperty",
 				"Cloth collection property '{1}' set in node and '{0}' is similar to the property '{2}' already set in an upstream node, "
 				"which might result in an undefined simulation behavior."),
 			FText::FromName(NodeName),
