Save Dataflow asset when the ClothAsset is saved

diff --git a/Engine/Plugins/Experimental/ChaosClothAssetEditor/Source/ChaosClothAssetEditor/Private/ChaosClothAsset/ClothEditorToolkit.cpp b/Engine/Plugins/Experimental/ChaosClothAssetEditor/Source/ChaosClothAssetEditor/Private/ChaosClothAsset/ClothEditorToolkit.cpp
--- a/Engine/Plugins/Experimental/ChaosClothAssetEditor/Source/ChaosClothAssetEditor/Private/ChaosClothAsset/ClothEditorToolkit.cpp
+++ b/Engine/Plugins/Experimental/ChaosClothAssetEditor/Source/ChaosClothAssetEditor/Private/ChaosClothAsset/ClothEditorToolkit.cpp
@@ -481,6 +481,19 @@ void FChaosClothAssetEditorToolkit::InitToolMenuContext(FToolMenuContext& MenuCo
 	MenuContext.AddObject(ClothEditorContext);
 }
 
+void FChaosClothAssetEditorToolkit::GetSaveableObjects(TArray<UObject*>& OutObjects) const
+{
+	FBaseCharacterFXEditorToolkit::GetSaveableObjects(OutObjects);
+
+	const UChaosClothAsset* const ClothAsset = GetAsset();
+	UDataflow* const DataflowAsset = Private::GetDataflowFrom(ClothAsset);
+	if (DataflowAsset)
+	{
+		check(DataflowAsset->IsAsset());
+		OutObjects.Add(DataflowAsset);
+	}
+}
+
 //~ End FAssetEditorToolkit overrides
 
 
diff --git a/Engine/Plugins/Experimental/ChaosClothAssetEditor/Source/ChaosClothAssetEditor/Private/ChaosClothAsset/ClothEditorToolkit.h b/Engine/Plugins/Experimental/ChaosClothAssetEditor/Source/ChaosClothAssetEditor/Private/ChaosClothAsset/ClothEditorToolkit.h
--- a/Engine/Plugins/Experimental/ChaosClothAssetEditor/Source/ChaosClothAssetEditor/Private/ChaosClothAsset/ClothEditorToolkit.h
+++ b/Engine/Plugins/Experimental/ChaosClothAssetEditor/Source/ChaosClothAssetEditor/Private/ChaosClothAsset/ClothEditorToolkit.h
@@ -84,6 +84,7 @@ private:
 	virtual bool OnRequestClose(EAssetEditorCloseReason InCloseReason) override;
 	virtual void PostInitAssetEditor() override;
 	virtual void InitToolMenuContext(FToolMenuContext& MenuContext) override;
+	virtual void GetSaveableObjects(TArray<UObject*>& OutObjects) const override;
 
 	// IAssetEditorInstance
 	// TODO: If this returns true then the editor cannot re-open after it's closed. Figure out why.
