Set Cloth Editor to be the default editor for Cloth Assets.
(This means that if the user has a cloth asset opened in the editor and they double click on the same asset, it will pop up the current editor, not open a new cloth editor instance for the same asset.)

diff --git a/Engine/Plugins/Experimental/ChaosClothAssetEditor/Source/ChaosClothAssetEditor/Private/ChaosClothAsset/ClothEditorToolkit.h b/Engine/Plugins/Experimental/ChaosClothAssetEditor/Source/ChaosClothAssetEditor/Private/ChaosClothAsset/ClothEditorToolkit.h
--- a/Engine/Plugins/Experimental/ChaosClothAssetEditor/Source/ChaosClothAssetEditor/Private/ChaosClothAsset/ClothEditorToolkit.h
+++ b/Engine/Plugins/Experimental/ChaosClothAssetEditor/Source/ChaosClothAssetEditor/Private/ChaosClothAsset/ClothEditorToolkit.h
@@ -87,8 +87,7 @@ private:
 	virtual void GetSaveableObjects(TArray<UObject*>& OutObjects) const override;
 
 	// IAssetEditorInstance
-	// TODO: If this returns true then the editor cannot re-open after it's closed. Figure out why.
-	virtual bool IsPrimaryEditor() const override { return false; };
+	virtual bool IsPrimaryEditor() const override { return true; };
 
 	// IToolkit
 	virtual FText GetToolkitName() const override;
