Cloth Editor
- when the Preview Scene Actor is selected, select all of its components
- when populating the Preview Scene for the first time, make sure the ClothComponent is RootComponent for the Actor

diff --git a/Engine/Plugins/Experimental/ChaosClothAssetEditor/Source/ChaosClothAssetEditor/Private/ChaosClothAsset/ClothEditor3DViewportClient.cpp b/Engine/Plugins/Experimental/ChaosClothAssetEditor/Source/ChaosClothAssetEditor/Private/ChaosClothAsset/ClothEditor3DViewportClient.cpp
--- a/Engine/Plugins/Experimental/ChaosClothAssetEditor/Source/ChaosClothAssetEditor/Private/ChaosClothAsset/ClothEditor3DViewportClient.cpp
+++ b/Engine/Plugins/Experimental/ChaosClothAssetEditor/Source/ChaosClothAssetEditor/Private/ChaosClothAsset/ClothEditor3DViewportClient.cpp
@@ -304,10 +304,7 @@ void FChaosClothAssetEditor3DViewportClient::ProcessClick(FSceneView& View, HHit
 	SelectedComponents->Modify();
 	SelectedComponents->BeginBatchSelectOperation();
 
-	if (!bIsShiftKeyDown && !bIsCtrlKeyDown)
-	{
-		SelectedComponents->DeselectAll();
-	}
+	SelectedComponents->DeselectAll();
 
 	if (HitProxy && HitProxy->IsA(HActor::StaticGetType()))
 	{
@@ -315,39 +312,12 @@ void FChaosClothAssetEditor3DViewportClient::ProcessClick(FSceneView& View, HHit
 		if (ActorProxy && ActorProxy->Actor)
 		{
 			const AActor* const Actor = ActorProxy->Actor;
-			USceneComponent* RootComponent = Actor->GetRootComponent();
-			if (RootComponent)
+			
+			Actor->ForEachComponent<UPrimitiveComponent>(true, [&](UPrimitiveComponent* Component)
 			{
-				if (bIsShiftKeyDown)
-				{
-					SelectedComponents->Select(RootComponent);
-				}
-				else if (bIsCtrlKeyDown)
-				{
-					// Don't use USelection::ToggleSelect here because that checks membership in GObjectSelection, not the given USelection...
-					TArray<USceneComponent*> Components;
-					SelectedComponents->GetSelectedObjects(Components);
-					const bool bIsSelected = Components.Contains(RootComponent);
-
-					if (bIsSelected)
-					{
-						SelectedComponents->Deselect(RootComponent);
-					}
-					else
-					{
-						SelectedComponents->Select(RootComponent);
-					}
-				}
-				else
-				{
-					SelectedComponents->Select(RootComponent);
-				}
-
-				if (UPrimitiveComponent* const PrimitiveComponent = Cast<UPrimitiveComponent>(RootComponent))
-				{
-					PrimitiveComponent->PushSelectionToProxy();
-				}
-			}
+				SelectedComponents->Select(Component);
+				Component->PushSelectionToProxy();
+			});
 		}
 	}
 
diff --git a/Engine/Plugins/Experimental/ChaosClothAssetEditor/Source/ChaosClothAssetEditor/Private/ChaosClothAsset/ClothEditorPreviewScene.cpp b/Engine/Plugins/Experimental/ChaosClothAssetEditor/Source/ChaosClothAssetEditor/Private/ChaosClothAsset/ClothEditorPreviewScene.cpp
--- a/Engine/Plugins/Experimental/ChaosClothAssetEditor/Source/ChaosClothAssetEditor/Private/ChaosClothAsset/ClothEditorPreviewScene.cpp
+++ b/Engine/Plugins/Experimental/ChaosClothAssetEditor/Source/ChaosClothAssetEditor/Private/ChaosClothAsset/ClothEditorPreviewScene.cpp
@@ -253,14 +253,18 @@ void FChaosClothPreviewScene::SetClothAsset(UChaosClothAsset* Asset)
 	check(Asset);
 	check(SceneActor);
 
-	bool bShouldClothBeRootComponent = false;
+	bool bShouldClothBeRootComponent = true;
 	if (ClothComponent)
 	{
 		if (SceneActor->GetRootComponent() == ClothComponent)
 		{
-			bShouldClothBeRootComponent = true;
 			SceneActor->SetRootComponent(nullptr);
 		}
+		else
+		{
+			// SkeletalMesh is the RootComponent, so don't change it
+			bShouldClothBeRootComponent = false;
+		}
 
 		ClothComponent->SelectionOverrideDelegate.Unbind();
 		SceneActor->RemoveOwnedComponent(ClothComponent);
