Added a new "Inpaint Weights" skin transfer method.

#rnx

diff --git a/Engine/Plugins/Experimental/ChaosClothAssetEditor/Source/ChaosClothAssetDataflowNodes/Private/ChaosClothAsset/TransferSkinWeightsNode.cpp b/Engine/Plugins/Experimental/ChaosClothAssetEditor/Source/ChaosClothAssetDataflowNodes/Private/ChaosClothAsset/TransferSkinWeightsNode.cpp
--- a/Engine/Plugins/Experimental/ChaosClothAssetEditor/Source/ChaosClothAssetDataflowNodes/Private/ChaosClothAsset/TransferSkinWeightsNode.cpp
+++ b/Engine/Plugins/Experimental/ChaosClothAssetEditor/Source/ChaosClothAssetDataflowNodes/Private/ChaosClothAsset/TransferSkinWeightsNode.cpp
@@ -11,58 +11,90 @@
 #include "SkeletalMeshAttributes.h"
 #include "Engine/SkeletalMesh.h"
 #include "Animation/Skeleton.h"
+#include "DynamicMesh/DynamicVertexSkinWeightsAttribute.h"
+#include "DynamicMesh/DynamicMeshAttributeSet.h"
+#include "DynamicMesh/DynamicBoneAttribute.h"
+#include "DynamicMesh/MeshNormals.h"
 
 #include UE_INLINE_GENERATED_CPP_BY_NAME(TransferSkinWeightsNode)
 
 #define LOCTEXT_NAMESPACE "ChaosClothAssetTransferSkinWeightsNode"
 
+using namespace UE::Chaos::ClothAsset;
+using namespace UE::AnimationCore;
+using namespace UE::Geometry;
+
 namespace UE::ChaosClothAsset::Private
 {
-	void SkeletalMeshToDynamicMesh(USkeletalMesh* FromSkeletalMeshAsset, int32 SourceLODIdx, FDynamicMesh3& ToDynamicMesh)
-	{
+	void SkeletalMeshToDynamicMesh(USkeletalMesh* FromSkeletalMeshAsset, int32 SkeletalMeshLOD, FDynamicMesh3& ToDynamicMesh)
+	{	
 		FMeshDescription SourceMesh;
+		FromSkeletalMeshAsset->GetMeshDescription(SkeletalMeshLOD, SourceMesh);
+		FMeshDescriptionToDynamicMesh Converter;
+		Converter.Convert(&SourceMesh, ToDynamicMesh);
+	}
 
-		// Check first if we have bulk data available and non-empty.
-		if (FromSkeletalMeshAsset->IsLODImportedDataBuildAvailable(SourceLODIdx) && !FromSkeletalMeshAsset->IsLODImportedDataEmpty(SourceLODIdx))
+	// Convert the cloth LOD into DynamicMesh by welding the sim mesh patterns
+	// @todo: This should instead be handled by a cloth lod to dynamic mesh converter similar to functions in the ClothPatternToDynamicMesh.h
+	bool SimClothLodToDynamicMesh(const FCollectionClothLodFacade& ClothLodFacade, 
+								  const FReferenceSkeleton& TargetRefSkeleton, // the reference skeleton to add to the dynamic mesh
+								  FDynamicMesh3& WeldedSimMesh,
+								  TArray<TArray<int32>>& WeldedToPatternIndices)
+	{
+		// Compute the welded simulation mesh for this LOD.
+		TArray<FVector3f> Positions;
+		TArray<FVector3f> Normals;
+		TArray<uint32> Indices, PatternsIndices, PatternToWeldedIndices;
+		TArray<FVector2f> PatternsPositions;
+		ClothLodFacade.BuildSimulationMesh(Positions, Normals, Indices, PatternsPositions, PatternsIndices, PatternToWeldedIndices, &WeldedToPatternIndices);
+
+		// Convert the sim mesh to DynamicMesh. 
+		// @todo: FTransferBoneWeights should accept raw data arrays (vertices/triangles) to avoid this conversion
+		WeldedSimMesh.Clear();
+		for (const FVector3f& Pos : Positions)
 		{
-			FSkeletalMeshImportData SkeletalMeshImportData;
-			FromSkeletalMeshAsset->LoadLODImportedData(SourceLODIdx, SkeletalMeshImportData);
-			SkeletalMeshImportData.GetMeshDescription(SourceMesh);
+			WeldedSimMesh.AppendVertex(FVector3d(Pos));
 		}
-		else
+
+		for (int Idx = 0; Idx < Indices.Num(); Idx+=3)
 		{
-			// Fall back on the LOD model directly if no bulk data exists. When we commit
-			// the mesh description, we override using the bulk data. This can happen for older
-			// skeletal meshes, from UE 4.24 and earlier.
-			const FSkeletalMeshModel* SkeletalMeshModel = FromSkeletalMeshAsset->GetImportedModel();
-			if (SkeletalMeshModel && SkeletalMeshModel->LODModels.IsValidIndex(SourceLODIdx))
+			const int TID = WeldedSimMesh.AppendTriangle(Indices[Idx], Indices[Idx+1], Indices[Idx+2]);
+			if (TID < 0) // Failed to add the triangle (non-manifold/duplicate)
 			{
-				SkeletalMeshModel->LODModels[SourceLODIdx].GetMeshDescription(SourceMesh, FromSkeletalMeshAsset);
+				return false;
 			}
 		}
+		
+		WeldedSimMesh.EnableAttributes();
+		FMeshNormals NormalsUtil(&WeldedSimMesh);
+		NormalsUtil.InitializeOverlayToPerVertexNormals(WeldedSimMesh.Attributes()->PrimaryNormals());
+
+		// Setup the skeleton.
+		// @note we can't simply copy the bone attributes from the source USkeletalMesh because the cloth  
+		// asset reference skeleton comes from the USkeleton not the USkeletalMesh
+		WeldedSimMesh.Attributes()->EnableBones(TargetRefSkeleton.GetRawBoneNum());
+		for (int32 BoneIdx = 0; BoneIdx < TargetRefSkeleton.GetRawBoneNum(); ++BoneIdx)
+		{
+			WeldedSimMesh.Attributes()->GetBoneNames()->SetValue(BoneIdx, TargetRefSkeleton.GetRawRefBoneInfo()[BoneIdx].Name);
+		}
 
-		FMeshDescriptionToDynamicMesh Converter;
-		Converter.Convert(&SourceMesh, ToDynamicMesh);
-	};
-
+		return true;
+	}
 }
 
 
+
 FChaosClothAssetTransferSkinWeightsNode::FChaosClothAssetTransferSkinWeightsNode(const Dataflow::FNodeParameters& InParam, FGuid InGuid)
 	: FDataflowNode(InParam, InGuid)
 {
 	RegisterInputConnection(&Collection);
 	RegisterOutputConnection(&Collection, &Collection);
 	RegisterInputConnection(&SkeletalMesh);
-	RegisterInputConnection(&Transform);
-	RegisterInputConnection(&SkeletalMeshLOD);
 }
 
 void FChaosClothAssetTransferSkinWeightsNode::Evaluate(Dataflow::FContext& Context, const FDataflowOutput* Out) const
 {
-	using namespace UE::Chaos::ClothAsset;
-	using namespace UE::AnimationCore;
-	using namespace UE::Geometry;
+	using namespace UE::ChaosClothAsset::Private;
 
 	if (Out->IsA<FManagedArrayCollection>(&Collection))
 	{
@@ -70,113 +102,160 @@ void FChaosClothAssetTransferSkinWeightsNode::Evaluate(Dataflow::FContext& Conte
 		FManagedArrayCollection InputCollection = GetValue<FManagedArrayCollection>(Context, &Collection);
 		const TSharedRef<FManagedArrayCollection> ClothCollection = MakeShared<FManagedArrayCollection>(MoveTemp(InputCollection));
 		const TObjectPtr<USkeletalMesh> InputSkeletalMesh = GetValue<TObjectPtr<USkeletalMesh>>(Context, &SkeletalMesh);
-		const FTransform InputRelativeTransform = GetValue<FTransform>(Context, &Transform);
-		const int32 SourceLODIdx = GetValue<int32>(Context, &SkeletalMeshLOD);
 
 		if (InputSkeletalMesh)
 		{
-			if (!InputSkeletalMesh->IsValidLODIndex(SourceLODIdx))
+			if (!InputSkeletalMesh->IsValidLODIndex(SkeletalMeshLOD))
 			{
-				UE::Chaos::ClothAsset::DataflowNodes::LogAndToastWarning(LOCTEXT("Warning_InvalidSkeletalMeshLOD", "TransferSkinWeightsNode: The specified input SkeletalMesh LOD is not valid"));
+				UE::Chaos::ClothAsset::DataflowNodes::LogAndToastWarning(LOCTEXT("Warning_TransferWeightsInvalidSkeletalMeshLOD", "TransferSkinWeightsNode: The specified input SkeletalMesh LOD is not valid."));
+				SetValue<FManagedArrayCollection>(Context, *ClothCollection, &Collection);
 				return;
 			}
 
-			// Convert source Skeletal Mesh to Dynamic Mesh
+			//
+			// Convert source Skeletal Mesh to Dynamic Mesh.
+			//
 			FDynamicMesh3 SourceDynamicMesh;
-			UE::ChaosClothAsset::Private::SkeletalMeshToDynamicMesh(InputSkeletalMesh, SourceLODIdx, SourceDynamicMesh);
-			MeshTransforms::ApplyTransform(SourceDynamicMesh, InputRelativeTransform, true);
-
-			FCollectionClothFacade ClothFacade(ClothCollection);
+			SkeletalMeshToDynamicMesh(InputSkeletalMesh, SkeletalMeshLOD, SourceDynamicMesh);
+			MeshTransforms::ApplyTransform(SourceDynamicMesh, Transform, true);
+			const FReferenceSkeleton& TargetRefSkeleton = InputSkeletalMesh->GetSkeleton()->GetReferenceSkeleton();
 
-			const USkeleton* const InputSkeleleton = InputSkeletalMesh->GetSkeleton();
 
-			// Compute bone index mappings
-			TMap<FName, FBoneIndexType> TargetBoneToIndex;
-			const FReferenceSkeleton& TargetRefSkeleton = InputSkeleleton->GetReferenceSkeleton();
-			for (int32 Index = 0; Index < TargetRefSkeleton.GetRawBoneNum(); ++Index)
-			{
-				TargetBoneToIndex.Add(TargetRefSkeleton.GetRawRefBoneInfo()[Index].Name, Index);
-			}
-
-			// Setup bone weight transfer operator
+			//
+			// Setup the bone weight transfer operator for the source mesh.
+			//
+			constexpr bool bUseParallel = true;
 			FTransferBoneWeights TransferBoneWeights(&SourceDynamicMesh, FSkeletalMeshAttributes::DefaultSkinWeightProfileName);
-			if (TransferBoneWeights.Validate() != EOperationValidationResult::Ok)
-			{
-				UE::Chaos::ClothAsset::DataflowNodes::LogAndToastWarning(LOCTEXT("Warning_TransferWeightsFailed", "TransferSkinWeightsNode: Transferring skin weights failed"));
-				return;
-			}
+			TransferBoneWeights.bUseParallel = bUseParallel;
+			TransferBoneWeights.TransferMethod = static_cast<FTransferBoneWeights::ETransferBoneWeightsMethod>(TransferMethod);
+
 
-			// Iterate over the LODs and transfer the bone weights from the source Skeletal mesh to the Cloth asset
-			for (int TargetLODIdx = 0; TargetLODIdx < ClothFacade.GetNumLods(); ++TargetLODIdx)
+			//
+			// Iterate over the LODs and transfer the bone weights from the source Skeletal mesh to the Cloth asset.
+			//
+			FCollectionClothFacade ClothFacade(ClothCollection);
+			for (int32 TargetLODIdx = 0; TargetLODIdx < ClothFacade.GetNumLods(); ++TargetLODIdx)
 			{
 				FCollectionClothLodFacade ClothLodFacade = ClothFacade.GetLod(TargetLODIdx);
-
-				// Cloth collection data arrays we are writing to
-				const TArrayView<int32> SimNumBoneInfluences = ClothLodFacade.GetSimNumBoneInfluences();
-				const TArrayView<TArray<int32>> SimBoneIndices = ClothLodFacade.GetSimBoneIndices();
-				const TArrayView<TArray<float>> SimBoneWeights = ClothLodFacade.GetSimBoneWeights();
-				const TConstArrayView<FVector3f> SimPositions = ClothLodFacade.GetSimRestPosition();
-				const int32 NumSimVertices = ClothLodFacade.GetNumSimVertices();
-
-				checkSlow(SimPositions.Num() == SimBoneIndices.Num());
-
-				constexpr bool bUseParallel = true;
-
-				// Iterate over each vertex and write the data from FBoneWeights into cloth collection managed arrays
-
-				ParallelFor(NumSimVertices,
-					[&SimPositions, &TransferBoneWeights, &TargetBoneToIndex, &SimNumBoneInfluences, &SimBoneIndices, &SimBoneWeights](int32 VertexID)
+				ClothLodFacade.SetSkeletonAssetPathName(InputSkeletalMesh->GetSkeleton()->GetPathName());
+
+				if (TransferMethod == EChaosClothAssetTransferSkinWeightsMethod::InpaintWeights)
+				{	
+					//
+					// Convert cloth sim mesh LOD to the welded dynamic sim mesh.
+					//
+					FDynamicMesh3 WeldedSimMesh;
+					TArray<TArray<int32>> WeldedToPatternIndices;
+					if (!ensure(SimClothLodToDynamicMesh(ClothLodFacade, TargetRefSkeleton, WeldedSimMesh, WeldedToPatternIndices)))
 					{
-						const FVector3d PosD = FVector3d(SimPositions[VertexID]);
+						UE::Chaos::ClothAsset::DataflowNodes::LogAndToastWarning(LOCTEXT("Warning_TransferWeightsFailedLodToDynamicMesh", "TransferSkinWeightsNode: Failed to weld the simulation mesh for LOD."));
+						break;
+					}
 
-						UE::AnimationCore::FBoneWeights BoneWeights;
-						TransferBoneWeights.TransferWeightsToPoint(BoneWeights, PosD, &TargetBoneToIndex);
 
-						const int32 NumBones = BoneWeights.Num();
+					//
+					// Transfer the weights from the body to the welded sim mesh.
+					//
+					TransferBoneWeights.NormalThreshold = FMathd::DegToRad * NormalThreshold;
+					TransferBoneWeights.SearchRadius = RadiusPercentage * WeldedSimMesh.GetBounds().DiagonalLength();
+					if (TransferBoneWeights.Validate() != EOperationValidationResult::Ok)
+					{
+						UE::Chaos::ClothAsset::DataflowNodes::LogAndToastWarning(LOCTEXT("Warning_TransferWeightsInpaintWeightsInvalidParameters", "TransferSkinWeightsNode: Transfer method parameters are invalid."));
+						break;
+					}
+					if (!TransferBoneWeights.TransferWeightsToMesh(WeldedSimMesh, FSkeletalMeshAttributes::DefaultSkinWeightProfileName))
+					{
+						UE::Chaos::ClothAsset::DataflowNodes::LogAndToastWarning(LOCTEXT("Warning_TransferWeightsFailed", "TransferSkinWeightsNode: Transferring skin weights failed"));
+						break;
+					}
 
-						SimNumBoneInfluences[VertexID] = NumBones;
-						SimBoneIndices[VertexID].SetNum(NumBones);
-						SimBoneWeights[VertexID].SetNum(NumBones);
 
-						for (int BoneIdx = 0; BoneIdx < NumBones; ++BoneIdx)
-						{
-							SimBoneIndices[VertexID][BoneIdx] = BoneWeights[BoneIdx].GetBoneIndex();
-							SimBoneWeights[VertexID][BoneIdx] = BoneWeights[BoneIdx].GetWeight();
+					//
+					// Copy the new bone weight data from the welded sim mesh back to the cloth patterns.
+					//
+					ParallelFor(WeldedSimMesh.MaxVertexID(), [&ClothLodFacade, &WeldedToPatternIndices, &WeldedSimMesh](int32 WeldedID)
+					{
+						const FDynamicMeshVertexSkinWeightsAttribute* OutAttribute = WeldedSimMesh.Attributes()->GetSkinWeightsAttribute(FSkeletalMeshAttributes::DefaultSkinWeightProfileName);
+						
+						checkSlow(OutAttribute);
+						checkSlow(WeldedSimMesh.IsVertex(WeldedID));
+
+						for (const int32 VertexID : WeldedToPatternIndices[WeldedID])
+						{	
+							checkSlow(VertexID < ClothLodFacade.GetNumSimVertices());
+							
+							OutAttribute->GetValue(WeldedID, ClothLodFacade.GetSimBoneIndices()[VertexID], 
+															 ClothLodFacade.GetSimBoneWeights()[VertexID]);
+							ClothLodFacade.GetSimNumBoneInfluences()[VertexID] = ClothLodFacade.GetSimBoneIndices()[VertexID].Num();
 						}
-
 					}, bUseParallel ? EParallelForFlags::None : EParallelForFlags::ForceSingleThread);
 
-				// Cloth collection data arrays we are writing to
-				const TArrayView<int32> RenderNumBoneInfluences = ClothLodFacade.GetRenderNumBoneInfluences();
-				const TArrayView<TArray<int32>> RenderBoneIndices = ClothLodFacade.GetRenderBoneIndices();
-				const TArrayView<TArray<float>> RenderBoneWeights = ClothLodFacade.GetRenderBoneWeights();
-				const TConstArrayView<FVector3f> RenderPositions = ClothLodFacade.GetRenderPosition();
-				const int32 NumRenderVertices = ClothLodFacade.GetNumRenderVertices();
 
-				// Iterate over each vertex and write the data from FBoneWeights into cloth collection managed arrays
-				ParallelFor(NumRenderVertices,
-					[&RenderPositions, &TransferBoneWeights, &TargetBoneToIndex, &RenderNumBoneInfluences, &RenderBoneIndices, &RenderBoneWeights](int32 VertexID)
+					//
+					// Compute the bone weights for the render mesh by transferring weights from the sim mesh
+					// @todo If render mesh eventually supports welding, we should be transferring weights from the 
+					// body instead, same as we do for the sim mesh.
+					//
+					FTransferBoneWeights SimToRenderMeshTransfer(&WeldedSimMesh, FSkeletalMeshAttributes::DefaultSkinWeightProfileName);
+					SimToRenderMeshTransfer.bUseParallel = bUseParallel;
+					SimToRenderMeshTransfer.TransferMethod = FTransferBoneWeights::ETransferBoneWeightsMethod::ClosestPointOnSurface;
+					ParallelFor(ClothLodFacade.GetNumRenderVertices(), [&SimToRenderMeshTransfer, &ClothLodFacade](int32 VertexID)
 					{
-						const FVector3d PosD = FVector3d(RenderPositions[VertexID]);
+						SimToRenderMeshTransfer.TransferWeightsToPoint(ClothLodFacade.GetRenderBoneIndices()[VertexID],
+																	   ClothLodFacade.GetRenderBoneWeights()[VertexID], 
+																	   ClothLodFacade.GetRenderPosition()[VertexID]);
+						ClothLodFacade.GetRenderNumBoneInfluences()[VertexID] = ClothLodFacade.GetRenderBoneIndices()[VertexID].Num();
 
-						UE::AnimationCore::FBoneWeights BoneWeights;
-						TransferBoneWeights.TransferWeightsToPoint(BoneWeights, PosD, &TargetBoneToIndex);
+					}, bUseParallel ? EParallelForFlags::None : EParallelForFlags::ForceSingleThread);
+				}
+				else if (TransferMethod == EChaosClothAssetTransferSkinWeightsMethod::ClosestPointOnSurface)
+				{	
+					//
+					// Compute the bone index mappings. This allows the transfer operator to retarget weights to the correct skeleton.
+					//
+					TMap<FName, FBoneIndexType> TargetBoneToIndex;
+					TargetBoneToIndex.Reserve(TargetRefSkeleton.GetRawBoneNum());
+					for (int32 BoneIdx = 0; BoneIdx < TargetRefSkeleton.GetRawBoneNum(); ++BoneIdx)
+					{
+						TargetBoneToIndex.Add(TargetRefSkeleton.GetRawRefBoneInfo()[BoneIdx].Name, BoneIdx);
+					}
 
-						const int32 NumBones = BoneWeights.Num();
+					if (TransferBoneWeights.Validate() != EOperationValidationResult::Ok)
+					{
+						UE::Chaos::ClothAsset::DataflowNodes::LogAndToastWarning(LOCTEXT("Warning_TransferWeightsClosestPointOnSurfaceInvalidParameters", "TransferSkinWeightsNode: Transfer method parameters are invalid."));
+						break;
+					}
+					
+					//
+					// Transfer weights to the sim mesh.
+					//
+					ParallelFor(ClothLodFacade.GetNumSimVertices(), [&TransferBoneWeights, &TargetBoneToIndex, &ClothLodFacade](int32 VertexID)
+					{
+						TransferBoneWeights.TransferWeightsToPoint(ClothLodFacade.GetSimBoneIndices()[VertexID],
+																   ClothLodFacade.GetSimBoneWeights()[VertexID],
+																   ClothLodFacade.GetSimRestPosition()[VertexID],
+																   &TargetBoneToIndex);
+					    ClothLodFacade.GetSimNumBoneInfluences()[VertexID] = ClothLodFacade.GetSimBoneIndices()[VertexID].Num();
 
-						RenderNumBoneInfluences[VertexID] = NumBones;
-						RenderBoneIndices[VertexID].SetNum(NumBones);
-						RenderBoneWeights[VertexID].SetNum(NumBones);
+					}, bUseParallel ? EParallelForFlags::None : EParallelForFlags::ForceSingleThread);
 
-						for (int BoneIdx = 0; BoneIdx < NumBones; ++BoneIdx)
-						{
-							RenderBoneIndices[VertexID][BoneIdx] = BoneWeights[BoneIdx].GetBoneIndex();
-							RenderBoneWeights[VertexID][BoneIdx] = BoneWeights[BoneIdx].GetWeight();
-						}
+					//
+					// Transfer weights to the render mesh.
+					//
+					ParallelFor(ClothLodFacade.GetNumRenderVertices(), [&TransferBoneWeights, &TargetBoneToIndex, &ClothLodFacade](int32 VertexID)
+					{
+						TransferBoneWeights.TransferWeightsToPoint(ClothLodFacade.GetRenderBoneIndices()[VertexID],
+																   ClothLodFacade.GetRenderBoneWeights()[VertexID],
+																   ClothLodFacade.GetRenderPosition()[VertexID],
+																   &TargetBoneToIndex);
+						ClothLodFacade.GetRenderNumBoneInfluences()[VertexID] = ClothLodFacade.GetRenderBoneIndices()[VertexID].Num();
 
 					}, bUseParallel ? EParallelForFlags::None : EParallelForFlags::ForceSingleThread);
-
-				ClothLodFacade.SetSkeletonAssetPathName(InputSkeleleton->GetPathName());
+				}
+				else
+				{
+					checkNoEntry();
+				}
 			}
 		}
 
diff --git a/Engine/Plugins/Experimental/ChaosClothAssetEditor/Source/ChaosClothAssetDataflowNodes/Public/ChaosClothAsset/TransferSkinWeightsNode.h b/Engine/Plugins/Experimental/ChaosClothAssetEditor/Source/ChaosClothAssetDataflowNodes/Public/ChaosClothAsset/TransferSkinWeightsNode.h
--- a/Engine/Plugins/Experimental/ChaosClothAssetEditor/Source/ChaosClothAssetDataflowNodes/Public/ChaosClothAsset/TransferSkinWeightsNode.h
+++ b/Engine/Plugins/Experimental/ChaosClothAssetEditor/Source/ChaosClothAssetDataflowNodes/Public/ChaosClothAsset/TransferSkinWeightsNode.h
@@ -9,6 +9,20 @@
 
 class USkeletalMesh;
 
+UENUM(BlueprintType)
+enum class EChaosClothAssetTransferSkinWeightsMethod : uint8
+{
+	// For every vertex on the target mesh, find the closest point on the surface of the source mesh and copy its weights.
+	ClosestPointOnSurface,
+	
+	// For every vertex on the target mesh, find the closest point on the surface of the source mesh. If that point 
+	// is within the SearchRadius, and their normals differ by less than the NormalThreshold, then we directly copy  
+	// the weights from the source point to the target mesh vertex. For all the vertices we didn't copy the weights 
+	// directly, automatically compute the smooth weights.
+	InpaintWeights
+};
+
+
 USTRUCT(Meta = (DataflowCloth))
 struct FChaosClothAssetTransferSkinWeightsNode : public FDataflowNode
 {
@@ -24,13 +38,26 @@ public:
 	UPROPERTY(EditAnywhere, Category = "Source Mesh", Meta = (Dataflowinput, DisplayName = "SkeletalMesh"))
 	TObjectPtr<USkeletalMesh> SkeletalMesh;
 
-	UPROPERTY(EditAnywhere, Category = "Source Mesh", Meta = (Dataflowinput, DisplayName = "SkeletalMeshLOD"))
+	UPROPERTY(EditAnywhere, Category = "Source Mesh", Meta = (DisplayName = "SkeletalMeshLOD"))
 	int32 SkeletalMeshLOD = 0;
 
 	// The relative transform between the SkeletalMesh and ClothCollection
-	UPROPERTY(EditAnywhere, Category = "Source Mesh", Meta = (Dataflowinput, DisplayName = "Transform"))
+	UPROPERTY(EditAnywhere, Category = "Source Mesh", Meta = (DisplayName = "Transform"))
 	FTransform Transform;
 
+	UPROPERTY(EditAnywhere, Category = "Transfer Method", Meta = (DisplayName = "Algorithm"))
+	EChaosClothAssetTransferSkinWeightsMethod TransferMethod = EChaosClothAssetTransferSkinWeightsMethod::ClosestPointOnSurface;
+
+	// Defines the search radius as the RadiusPercentage * BoundingBoxDiagonalLength. All points not within the search
+	// radius will be ignored. If negative, all points are considered. Only used in the InpaintWeights algorithm.
+	UPROPERTY(EditAnywhere, Category = "Transfer Method", Meta = (UIMin = -1, UIMax = 2, ClampMin = -1, ClampMax = 2, DisplayName = "Radius Percentage", EditCondition="TransferMethod==EChaosClothAssetTransferSkinWeightsMethod::InpaintWeights"))
+	double RadiusPercentage = 0.05;
+
+	// Maximum angle (in degress) difference between target and source point normals to be considred a match. 
+	// If negative, normals are ignored. Only used in the InpaintWeights algorithm.
+	UPROPERTY(EditAnywhere, Category = "Transfer Method", Meta = (UIMin = -1, UIMax = 180, ClampMin = -1, ClampMax = 180, DisplayName = "Normal Threshold", EditCondition="TransferMethod==EChaosClothAssetTransferSkinWeightsMethod::InpaintWeights"))
+	double NormalThreshold = 30;
+
 	FChaosClothAssetTransferSkinWeightsNode(const Dataflow::FNodeParameters& InParam, FGuid InGuid = FGuid::NewGuid());
 
 	virtual void Evaluate(Dataflow::FContext& Context, const FDataflowOutput* Out) const override;
diff --git a/Engine/Plugins/Experimental/ChaosClothAssetEditor/Source/ChaosClothAssetEditorTools/Private/ChaosClothAsset/ClothTransferSkinWeightsTool.cpp b/Engine/Plugins/Experimental/ChaosClothAssetEditor/Source/ChaosClothAssetEditorTools/Private/ChaosClothAsset/ClothTransferSkinWeightsTool.cpp
--- a/Engine/Plugins/Experimental/ChaosClothAssetEditor/Source/ChaosClothAssetEditorTools/Private/ChaosClothAsset/ClothTransferSkinWeightsTool.cpp
+++ b/Engine/Plugins/Experimental/ChaosClothAssetEditor/Source/ChaosClothAssetEditorTools/Private/ChaosClothAsset/ClothTransferSkinWeightsTool.cpp
@@ -28,26 +28,7 @@ namespace UE::Chaos::ClothAsset::Private
 	void SkeletalMeshToDynamicMesh(USkeletalMesh* FromSkeletalMeshAsset, int32 SourceLODIdx, FDynamicMesh3& ToDynamicMesh)
 	{
 		FMeshDescription SourceMesh;
-
-		// Check first if we have bulk data available and non-empty.
-		if (FromSkeletalMeshAsset->IsLODImportedDataBuildAvailable(SourceLODIdx) && !FromSkeletalMeshAsset->IsLODImportedDataEmpty(SourceLODIdx))
-		{
-			FSkeletalMeshImportData SkeletalMeshImportData;
-			FromSkeletalMeshAsset->LoadLODImportedData(SourceLODIdx, SkeletalMeshImportData);
-			SkeletalMeshImportData.GetMeshDescription(SourceMesh);
-		}
-		else
-		{
-			// Fall back on the LOD model directly if no bulk data exists. When we commit
-			// the mesh description, we override using the bulk data. This can happen for older
-			// skeletal meshes, from UE 4.24 and earlier.
-			const FSkeletalMeshModel* SkeletalMeshModel = FromSkeletalMeshAsset->GetImportedModel();
-			if (SkeletalMeshModel && SkeletalMeshModel->LODModels.IsValidIndex(SourceLODIdx))
-			{
-				SkeletalMeshModel->LODModels[SourceLODIdx].GetMeshDescription(SourceMesh, FromSkeletalMeshAsset);
-			}
-		}
-
+		FromSkeletalMeshAsset->GetMeshDescription(SourceLODIdx, SourceMesh);
 		FMeshDescriptionToDynamicMesh Converter;
 		Converter.Convert(&SourceMesh, ToDynamicMesh);
 	};
