Cloth Editor: changes to the Construction viewport
- Added ability to show the Render mesh of the currently selected node's ClothCollection output
- Removed the need to select the mesh before starting a tool. If a node with an output ClothCollection is selected, the mesh is auto-selected and tools are enabled
- Added a toggle to control wireframe rendering

diff --git a/Engine/Plugins/Experimental/ChaosClothAssetEditor/Source/ChaosClothAssetEditor/Private/ChaosClothAsset/ClothEditorCommands.cpp b/Engine/Plugins/Experimental/ChaosClothAssetEditor/Source/ChaosClothAssetEditor/Private/ChaosClothAsset/ClothEditorCommands.cpp
--- a/Engine/Plugins/Experimental/ChaosClothAssetEditor/Source/ChaosClothAssetEditor/Private/ChaosClothAsset/ClothEditorCommands.cpp
+++ b/Engine/Plugins/Experimental/ChaosClothAssetEditor/Source/ChaosClothAssetEditor/Private/ChaosClothAsset/ClothEditorCommands.cpp
@@ -14,13 +14,11 @@ const FString FChaosClothAssetEditorCommands::BeginAttributeEditorToolIdentifier
 const FString FChaosClothAssetEditorCommands::BeginWeightMapPaintToolIdentifier = TEXT("BeginWeightMapPaintTool");
 const FString FChaosClothAssetEditorCommands::BeginClothTrainingToolIdentifier = TEXT("BeginClothTrainingTool");
 const FString FChaosClothAssetEditorCommands::BeginTransferSkinWeightsToolIdentifier = TEXT("BeginTransferSkinWeightsTool");
-const FString FChaosClothAssetEditorCommands::ToggleSimMeshWireframeIdentifier = TEXT("ToggleSimMeshWireframe");
-const FString FChaosClothAssetEditorCommands::ToggleRenderMeshWireframeIdentifier = TEXT("ToggleRenderMeshWireframe");
 const FString FChaosClothAssetEditorCommands::ToggleSimulationSuspendedIdentifier = TEXT("ToggleSimulationSuspended");
 const FString FChaosClothAssetEditorCommands::SoftResetSimulationIdentifier = TEXT("SoftResetSimulation");
 const FString FChaosClothAssetEditorCommands::HardResetSimulationIdentifier = TEXT("HardResetSimulation");
-const FString FChaosClothAssetEditorCommands::TogglePatternModeIdentifier = TEXT("TogglePatternMode");
-
+const FString FChaosClothAssetEditorCommands::TogglePreviewWireframeIdentifier = TEXT("TogglePreviewWireframe");
+const FString FChaosClothAssetEditorCommands::ToggleConstructionViewWireframeIdentifier = TEXT("ToggleConstructionViewWireframe");
 
 
 FChaosClothAssetEditorCommands::FChaosClothAssetEditorCommands()
@@ -45,10 +43,12 @@ void FChaosClothAssetEditorCommands::RegisterCommands()
 
 	UI_COMMAND(BeginTransferSkinWeightsTool, "Transfer Skin Weights", "Launch Transfer Skin Weights tool", EUserInterfaceActionType::Button, FInputChord());
 
-	UI_COMMAND(TogglePatternMode, "TogglePatternMode", "Toggle pattern mode", EUserInterfaceActionType::Button, FInputChord());
+	UI_COMMAND(SetConstructionMode2D, "2D Sim", "Switches the viewport to 2D simulation mesh view", EUserInterfaceActionType::RadioButton, FInputChord());
+	UI_COMMAND(SetConstructionMode3D, "3D Sim", "Switches the viewport to 3D simulation mesh view", EUserInterfaceActionType::RadioButton, FInputChord());
+	UI_COMMAND(SetConstructionModeRender, "Render", "Switches the viewport to render mesh view", EUserInterfaceActionType::RadioButton, FInputChord());
 
-	UI_COMMAND(ToggleSimMeshWireframe, "ToggleSimMeshWireframe", "Toggle simulation mesh wireframe", EUserInterfaceActionType::ToggleButton, FInputChord());
-	UI_COMMAND(ToggleRenderMeshWireframe, "ToggleRenderMeshWireframe", "Toggle render mesh wireframe", EUserInterfaceActionType::ToggleButton, FInputChord());
+	UI_COMMAND(TogglePreviewWireframe, "TogglePreviewWireframe", "Toggle preview wireframe", EUserInterfaceActionType::ToggleButton, FInputChord());
+	UI_COMMAND(ToggleConstructionViewWireframe, "ToggleConstructionViewWireframe", "Toggle construction view wireframe", EUserInterfaceActionType::ToggleButton, FInputChord());
 	
 	UI_COMMAND(SoftResetSimulation, "SoftResetSimulation", "Soft reset simulation", EUserInterfaceActionType::ToggleButton, FInputChord());
 	UI_COMMAND(HardResetSimulation, "HardResetClothSimulation", "Hard reset simulation", EUserInterfaceActionType::ToggleButton, FInputChord());
diff --git a/Engine/Plugins/Experimental/ChaosClothAssetEditor/Source/ChaosClothAssetEditor/Private/ChaosClothAsset/ClothEditorCommands.h b/Engine/Plugins/Experimental/ChaosClothAssetEditor/Source/ChaosClothAssetEditor/Private/ChaosClothAsset/ClothEditorCommands.h
--- a/Engine/Plugins/Experimental/ChaosClothAssetEditor/Source/ChaosClothAssetEditor/Private/ChaosClothAsset/ClothEditorCommands.h
+++ b/Engine/Plugins/Experimental/ChaosClothAssetEditor/Source/ChaosClothAssetEditor/Private/ChaosClothAsset/ClothEditorCommands.h
@@ -44,17 +44,17 @@ public:
 	const static FString BeginTransferSkinWeightsToolIdentifier;
 	TSharedPtr<FUICommandInfo> BeginTransferSkinWeightsTool;
 
-	// Rest space viewport commands
-	const static FString TogglePatternModeIdentifier;
-	TSharedPtr<FUICommandInfo> TogglePatternMode;
+	// Construction viewport commands
+	const static FString ToggleConstructionViewWireframeIdentifier;
+	TSharedPtr<FUICommandInfo> ToggleConstructionViewWireframe;
 
-	// Sim viewport commands
-	const static FString ToggleSimMeshWireframeIdentifier;
-	TSharedPtr<FUICommandInfo> ToggleSimMeshWireframe;
-
-	const static FString ToggleRenderMeshWireframeIdentifier;
-	TSharedPtr<FUICommandInfo> ToggleRenderMeshWireframe;
+	TSharedPtr<FUICommandInfo> SetConstructionMode2D;
+	TSharedPtr<FUICommandInfo> SetConstructionMode3D;
+	TSharedPtr<FUICommandInfo> SetConstructionModeRender;
 
+	// Preview viewport commands
+	const static FString TogglePreviewWireframeIdentifier;
+	TSharedPtr<FUICommandInfo> TogglePreviewWireframe;
 
 	const static FString SoftResetSimulationIdentifier;
 	TSharedPtr<FUICommandInfo> SoftResetSimulation;
diff --git a/Engine/Plugins/Experimental/ChaosClothAssetEditor/Source/ChaosClothAssetEditor/Private/ChaosClothAsset/ClothEditorMode.cpp b/Engine/Plugins/Experimental/ChaosClothAssetEditor/Source/ChaosClothAssetEditor/Private/ChaosClothAsset/ClothEditorMode.cpp
--- a/Engine/Plugins/Experimental/ChaosClothAssetEditor/Source/ChaosClothAssetEditor/Private/ChaosClothAsset/ClothEditorMode.cpp
+++ b/Engine/Plugins/Experimental/ChaosClothAssetEditor/Source/ChaosClothAssetEditor/Private/ChaosClothAsset/ClothEditorMode.cpp
@@ -56,6 +56,7 @@
 #include "Dataflow/DataflowGraphEditor.h"
 #include "ContextObjectStore.h"
 #include "BaseGizmos/TransformGizmoUtil.h"
+#include "Materials/Material.h"
 
 #include UE_INLINE_GENERATED_CPP_BY_NAME(ClothEditorMode)
 
@@ -117,34 +118,6 @@ void UChaosClothAssetEditorMode::Enter()
 {
 	UBaseCharacterFXEditorMode::Enter();
 
-	SelectionModifiedEventHandle = USelection::SelectionChangedEvent.AddLambda([this](UObject*)
-	{
-		if (!GetModeManager() || !GetModeManager()->GetSelectedComponents())
-		{
-			return;
-		}
-
-		constexpr FColor UnselectedColor = FColor(128, 128, 128);
-		const FColor SelectedColor = GetDefault<UEditorStyleSettings>()->SelectionColor.ToFColor(true);
-
-		const USelection* const SelectedComponents = GetModeManager()->GetSelectedComponents();
-
-		check(DynamicMeshComponents.Num() == WireframesToTick.Num());
-
-		for (int MeshIndex = 0; MeshIndex < WireframesToTick.Num(); ++MeshIndex)
-		{
-			const TObjectPtr<UMeshElementsVisualizer> WireframeDisplay = WireframesToTick[MeshIndex];
-			if (SelectedComponents->IsSelected(DynamicMeshComponents[MeshIndex]))
-			{
-				WireframeDisplay->Settings->WireframeColor = SelectedColor;
-			}
-			else
-			{
-				WireframeDisplay->Settings->WireframeColor = UnselectedColor;
-			}
-		}
-	});
-
 	// Register gizmo ContextObject for use inside interactive tools
 	UE::TransformGizmoUtil::RegisterTransformGizmoContextObject(GetInteractiveToolsContext());
 }
@@ -254,7 +227,7 @@ void UChaosClothAssetEditorMode::OnToolStarted(UInteractiveToolManager* Manager,
 
 	FChaosClothAssetEditorCommands::UpdateToolCommandBinding(Tool, ToolCommandList, false);
 
-	bCanTogglePattern2DMode = false;
+	bCanChangeConstructionViewMode = false;
 
 	// Don't let the Graph change during tool execution (or even which nodes are selected)
 	if (TSharedPtr<SDataflowGraphEditor> GraphEditor = DataflowGraphEditor.Pin())
@@ -267,7 +240,8 @@ void UChaosClothAssetEditorMode::OnToolEnded(UInteractiveToolManager* Manager, U
 {
 	UE::Chaos::ClothAsset::FChaosClothAssetEditorCommands::UpdateToolCommandBinding(Tool, ToolCommandList, true);
 
-	bCanTogglePattern2DMode = true;
+	bCanChangeConstructionViewMode = true;
+
 	ActiveToolsContext = nullptr;
 
 	ReinitializeDynamicMeshComponents();
@@ -314,7 +288,6 @@ void UChaosClothAssetEditorMode::BindCommands()
 
 void UChaosClothAssetEditorMode::Exit()
 {
-	USelection::SelectionChangedEvent.Remove(SelectionModifiedEventHandle);
 	UActorComponent::MarkRenderStateDirtyEvent.RemoveAll(this);
 
 	for (TObjectPtr<UDynamicMeshComponent> DynamicMeshComp : DynamicMeshComponents)
@@ -417,16 +390,36 @@ TSharedPtr<FManagedArrayCollection> UChaosClothAssetEditorMode::GetClothCollecti
 
 void UChaosClothAssetEditorMode::ReinitializeDynamicMeshComponents()
 {
-	auto SetUpDynamicMeshComponent = [this](UDynamicMeshComponent& MeshComponent)
-	{
-		MeshComponent.SelectionOverrideDelegate = UPrimitiveComponent::FSelectionOverride::CreateUObject(this, &UChaosClothAssetEditorMode::IsComponentSelected);
-
-		UMaterialInterface* const Material = bPattern2DMode ?
-			ToolSetupUtil::GetCustomTwoSidedDepthOffsetMaterial(GetToolManager(), FLinearColor{0.6, 0.6, 0.6}, 0.0) :
-			ToolSetupUtil::GetDefaultSculptMaterial(GetToolManager());
+	using namespace UE::Chaos::ClothAsset;
 
-		MeshComponent.SetMaterial(0, Material);
-		MeshComponent.RegisterComponentWithWorld(this->GetWorld());
+	auto SetUpDynamicMeshComponentMaterial = [this](const UE::Chaos::ClothAsset::FCollectionClothLodConstFacade& ClothLodFacade, UDynamicMeshComponent& MeshComponent)
+	{
+		switch (ConstructionViewMode)
+		{
+			case EClothPatternVertexType::Sim2D:
+			{
+				UMaterialInterface* const Material = ToolSetupUtil::GetCustomTwoSidedDepthOffsetMaterial(GetToolManager(), FLinearColor{ 0.6, 0.6, 0.6 }, 0.0);
+				MeshComponent.SetMaterial(0, Material);
+			}
+			break;
+			case EClothPatternVertexType::Sim3D:
+			{
+				UMaterialInterface* const Material = ToolSetupUtil::GetDefaultSculptMaterial(GetToolManager());
+				MeshComponent.SetMaterial(0, Material);
+			}
+			break;
+			case EClothPatternVertexType::Render:
+			{
+				const TArrayView<const FString> MaterialPaths = ClothLodFacade.GetRenderMaterialPathName();
+				for (int32 MaterialIndex = 0; MaterialIndex < MaterialPaths.Num(); ++MaterialIndex)
+				{
+					const FString& Path = MaterialPaths[MaterialIndex];
+					UMaterialInterface* const Material = LoadObject<UMaterial>(nullptr, *Path);
+					MeshComponent.SetMaterial(MaterialIndex, Material);
+				}
+			}
+			break;
+		}
 	};
 	
 	// Clean up any existing DynamicMeshComponents
@@ -478,15 +471,16 @@ void UChaosClothAssetEditorMode::ReinitializeDynamicMeshComponents()
 			DynamicMeshComponentParentActors.SetNum(1);
 
 			UE::Geometry::FDynamicMesh3 LodMesh;
+			LodMesh.EnableAttributes();
 			UE::Geometry::FDynamicMeshEditor MeshEditor(&LodMesh);
 
 			for (int32 PatternIndex = 0; PatternIndex < ClothLodFacade.GetNumPatterns(); ++PatternIndex)
 			{
 				UE::Geometry::FDynamicMesh3 PatternMesh;
 				FClothPatternToDynamicMesh Converter;
-				Converter.Convert(Collection, LodIndex, PatternIndex, bPattern2DMode, PatternMesh);
+				Converter.Convert(Collection, LodIndex, PatternIndex, ConstructionViewMode, PatternMesh);
 
-				if (bPattern2DMode)
+				if (ConstructionViewMode == EClothPatternVertexType::Sim2D)
 				{
 					// Use per-triangle normals for the 2D view
 					UE::Geometry::FMeshNormals::InitializeMeshToPerTriangleNormals(&PatternMesh);
@@ -503,9 +497,12 @@ void UChaosClothAssetEditorMode::ReinitializeDynamicMeshComponents()
 			DynamicMeshComponentParentActors[0] = ParentActor;
 
 			DynamicMeshComponents[0] = NewObject<UDynamicMeshComponent>(ParentActor);
-			DynamicMeshComponents[0]->SetMesh(MoveTemp(LodMesh));
-			
-			SetUpDynamicMeshComponent(*DynamicMeshComponents[0]);
+			DynamicMeshComponents[0]->SetMesh(MoveTemp(LodMesh));		
+
+			SetUpDynamicMeshComponentMaterial(ClothLodFacade, *DynamicMeshComponents[0]);
+
+			DynamicMeshComponents[0]->SelectionOverrideDelegate = UPrimitiveComponent::FSelectionOverride::CreateUObject(this, &UChaosClothAssetEditorMode::IsComponentSelected);
+			DynamicMeshComponents[0]->RegisterComponentWithWorld(this->GetWorld());
 
 			DynamicMeshSourceInfos.Add(FDynamicMeshSourceInfo{ LodIndex, -1 });
 		}
@@ -519,9 +516,9 @@ void UChaosClothAssetEditorMode::ReinitializeDynamicMeshComponents()
 			{
 				UE::Geometry::FDynamicMesh3 PatternMesh;
 				FClothPatternToDynamicMesh Converter;
-				Converter.Convert(Collection, LodIndex, PatternIndex, bPattern2DMode, PatternMesh);
+				Converter.Convert(Collection, LodIndex, PatternIndex, ConstructionViewMode, PatternMesh);
 
-				if (bPattern2DMode)
+				if (ConstructionViewMode == EClothPatternVertexType::Sim2D)
 				{
 					// Use per-triangle normals for the 2D view
 					UE::Geometry::FMeshNormals::InitializeMeshToPerTriangleNormals(&PatternMesh);
@@ -536,7 +533,10 @@ void UChaosClothAssetEditorMode::ReinitializeDynamicMeshComponents()
 				TObjectPtr<UDynamicMeshComponent> PatternMeshComponent = NewObject<UDynamicMeshComponent>(ParentActor);
 				PatternMeshComponent->SetMesh(MoveTemp(PatternMesh));
 
-				SetUpDynamicMeshComponent(*PatternMeshComponent);
+				SetUpDynamicMeshComponentMaterial(ClothLodFacade, *PatternMeshComponent);
+
+				DynamicMeshComponents[PatternIndex]->SelectionOverrideDelegate = UPrimitiveComponent::FSelectionOverride::CreateUObject(this, &UChaosClothAssetEditorMode::IsComponentSelected);
+				DynamicMeshComponents[PatternIndex]->RegisterComponentWithWorld(this->GetWorld());
 
 				DynamicMeshComponents[PatternIndex] = PatternMeshComponent;
 				DynamicMeshSourceInfos.Add(FDynamicMeshSourceInfo{ LodIndex, PatternIndex });
@@ -601,21 +601,25 @@ void UChaosClothAssetEditorMode::ReinitializeDynamicMeshComponents()
 				return;
 			}
 
-			bool bRestSpaceMeshVisible = RestSpaceMesh->GetVisibleFlag();
-			WireframeDisplay->SetAllVisible(bRestSpaceMeshVisible);
-			WireframeDisplay->Settings->bVisible = bRestSpaceMeshVisible;
+			const bool bRestSpaceMeshVisible = RestSpaceMesh->GetVisibleFlag();
+			WireframeDisplay->Settings->bVisible = bRestSpaceMeshVisible && bConstructionViewWireframe;
 		});
 
+		const bool bRestSpaceMeshVisible = RestSpaceMeshComponent->GetVisibleFlag();
+		WireframeDisplay->Settings->bVisible = bRestSpaceMeshVisible && bConstructionViewWireframe;
 	}
 
 	SelectedComponents->DeselectAll();
-	SelectedComponents->BeginBatchSelectOperation();
-	for (TObjectPtr<UDynamicMeshComponent> RestSpaceMeshComponent : DynamicMeshComponents)
+	if (ConstructionViewMode != EClothPatternVertexType::Render)
 	{
-		SelectedComponents->Select(RestSpaceMeshComponent);
-		RestSpaceMeshComponent->PushSelectionToProxy();
+		SelectedComponents->BeginBatchSelectOperation();
+		for (TObjectPtr<UDynamicMeshComponent> RestSpaceMeshComponent : DynamicMeshComponents)
+		{
+			SelectedComponents->Select(RestSpaceMeshComponent);
+			RestSpaceMeshComponent->PushSelectionToProxy();
+		}
+		SelectedComponents->EndBatchSelectOperation();
 	}
-	SelectedComponents->EndBatchSelectOperation();
 }
 
 void UChaosClothAssetEditorMode::RefocusRestSpaceViewportClient()
@@ -627,6 +631,7 @@ void UChaosClothAssetEditorMode::RefocusRestSpaceViewportClient()
 		PinnedVC->ToggleOrbitCamera(false);
 
 		const FBox SceneBounds = SceneBoundingBox();
+		const bool bPattern2DMode = (ConstructionViewMode == UE::Chaos::ClothAsset::EClothPatternVertexType::Sim2D);
 		if (bPattern2DMode)
 		{
 			// 2D pattern
@@ -642,7 +647,7 @@ void UChaosClothAssetEditorMode::RefocusRestSpaceViewportClient()
 		PinnedVC->FocusViewportOnBox(SceneBounds, bInstant);
 
 		// Recompute near/far clip planes
-		PinnedVC->Set2DMode(bPattern2DMode);
+		PinnedVC->SetConstructionViewMode(ConstructionViewMode);
 	}
 }
 
@@ -651,15 +656,7 @@ void UChaosClothAssetEditorMode::InitializeTargets(const TArray<TObjectPtr<UObje
 	// InitializeContexts needs to have been called first so that we have the 3d preview world ready.
 	check(PreviewScene);
 
-	OriginalObjectsToEdit = AssetsIn;
-
-	CreateToolTargets(AssetsIn);
-
-	// NOTE: This ensure will fire until we are able to build a tool target associated with ClothComponent
-	//if (!ensure(AssetsIn.Num() == ToolTargets.Num()))
-	//{
-	//	return;
-	//}
+	UBaseCharacterFXEditorMode::InitializeTargets(AssetsIn);
 
 	for (TObjectPtr<UObject> Asset : AssetsIn)
 	{
@@ -674,8 +671,6 @@ void UChaosClothAssetEditorMode::InitializeTargets(const TArray<TObjectPtr<UObje
 
 	DataflowComponent = NewObject<UDataflowComponent>();
 	DataflowComponent->RegisterComponentWithWorld(PreviewScene->GetWorld());
-
-	bShouldFocusRestSpaceView = true;
 }
 
 void UChaosClothAssetEditorMode::SoftResetSimulation()
@@ -880,25 +875,40 @@ FBox UChaosClothAssetEditorMode::PreviewBoundingBox() const
 	return FBox(ForceInitToZero);
 }
 
-void UChaosClothAssetEditorMode::TogglePatternMode()
+void UChaosClothAssetEditorMode::SetConstructionViewMode(UE::Chaos::ClothAsset::EClothPatternVertexType InMode)
 {
-	bPattern2DMode = !bPattern2DMode;
+	const bool bSwitching2D3D = (ConstructionViewMode == UE::Chaos::ClothAsset::EClothPatternVertexType::Sim2D) != (InMode == UE::Chaos::ClothAsset::EClothPatternVertexType::Sim2D);
+
+	ConstructionViewMode = InMode;
 	ReinitializeDynamicMeshComponents();
 
 	TSharedPtr<UE::Chaos::ClothAsset::FChaosClothEditorRestSpaceViewportClient> VC = RestSpaceViewportClient.Pin();
 	if (VC.IsValid())
 	{
-		VC->Set2DMode(bPattern2DMode);
+		VC->SetConstructionViewMode(ConstructionViewMode);
 	}
 
-	RefocusRestSpaceViewportClient();
+	if (bSwitching2D3D)
+	{
+		RefocusRestSpaceViewportClient();
+	}
 }
 
-bool UChaosClothAssetEditorMode::CanTogglePatternMode() const
+UE::Chaos::ClothAsset::EClothPatternVertexType UChaosClothAssetEditorMode::GetConstructionViewMode() const
 {
-	return bCanTogglePattern2DMode;
+	return ConstructionViewMode;
 }
 
+bool UChaosClothAssetEditorMode::CanChangeConstructionViewMode() const
+{
+	return bCanChangeConstructionViewMode;
+}
+
+void UChaosClothAssetEditorMode::ToggleConstructionViewWireframe()
+{
+	bConstructionViewWireframe = !bConstructionViewWireframe;
+	ReinitializeDynamicMeshComponents();
+}
 
 void UChaosClothAssetEditorMode::SetRestSpaceViewportClient(TWeakPtr<UE::Chaos::ClothAsset::FChaosClothEditorRestSpaceViewportClient, ESPMode::ThreadSafe> InViewportClient)
 {
@@ -907,7 +917,7 @@ void UChaosClothAssetEditorMode::SetRestSpaceViewportClient(TWeakPtr<UE::Chaos::
 	TSharedPtr<UE::Chaos::ClothAsset::FChaosClothEditorRestSpaceViewportClient> VC = RestSpaceViewportClient.Pin();
 	if (VC.IsValid())
 	{
-		VC->Set2DMode(bPattern2DMode);
+		VC->SetConstructionViewMode(ConstructionViewMode);
 		VC->SetToolCommandList(ToolCommandList);
 
 		if (VC->Viewport)
diff --git a/Engine/Plugins/Experimental/ChaosClothAssetEditor/Source/ChaosClothAssetEditor/Private/ChaosClothAsset/ClothEditorMode.h b/Engine/Plugins/Experimental/ChaosClothAssetEditor/Source/ChaosClothAssetEditor/Private/ChaosClothAsset/ClothEditorMode.h
--- a/Engine/Plugins/Experimental/ChaosClothAssetEditor/Source/ChaosClothAssetEditor/Private/ChaosClothAsset/ClothEditorMode.h
+++ b/Engine/Plugins/Experimental/ChaosClothAssetEditor/Source/ChaosClothAssetEditor/Private/ChaosClothAsset/ClothEditorMode.h
@@ -5,6 +5,7 @@
 #include "BaseCharacterFXEditorMode.h"
 #include "GeometryBase.h"
 #include "Delegates/IDelegateInstance.h"
+#include "ChaosClothAsset/ClothPatternVertexType.h"
 #include "ClothEditorMode.generated.h"
 
 PREDECLARE_GEOMETRY(class FDynamicMesh3);
@@ -59,9 +60,15 @@ public:
 	// Bounding box for sim space meshes
 	FBox PreviewBoundingBox() const;
 
-	// Toggle between 2D pattern and 3D rest space mesh view
-	void TogglePatternMode();
-	bool CanTogglePatternMode() const;
+	void SetConstructionViewMode(UE::Chaos::ClothAsset::EClothPatternVertexType InMode);
+	UE::Chaos::ClothAsset::EClothPatternVertexType GetConstructionViewMode() const;
+	bool CanChangeConstructionViewMode() const;
+
+	void ToggleConstructionViewWireframe();
+	bool IsConstructionViewWireframeActive() const
+	{
+		return bConstructionViewWireframe;
+	}
 
 	// Simulation controls
 	void SoftResetSimulation();
@@ -137,6 +144,8 @@ private:
 	void InitializeContextObject();
 	void DeleteContextObject();
 
+	bool IsComponentSelected(const UPrimitiveComponent* InComponent);
+
 	// Rest space wireframes. They have to get ticked to be able to respond to setting changes. 
 	UPROPERTY()
 	TArray<TObjectPtr<UMeshElementsVisualizer>> WireframesToTick;
@@ -166,15 +175,6 @@ private:
 
 	TWeakPtr<UE::Chaos::ClothAsset::FChaosClothEditorRestSpaceViewportClient, ESPMode::ThreadSafe> RestSpaceViewportClient;
 
-	// Handle to a callback triggered when the current selection changes
-	FDelegateHandle SelectionModifiedEventHandle;
-
-	// Whether to display the 2D pattern or 3D assembly in the rest space viewport
-	bool bPattern2DMode = true;
-
-	// Whether we can switch between 2D and 3D rest configuration
-	bool bCanTogglePattern2DMode = true;
-
 	// Whether a cloth collection has ever been set on this Mode object
 	bool bFirstClothCollection = true;
 
@@ -187,7 +187,10 @@ private:
 	// TODO: Expose this to the user
 	bool bCombineAllPatterns = true;
 
-	bool IsComponentSelected(const UPrimitiveComponent* InComponent);
+	UE::Chaos::ClothAsset::EClothPatternVertexType ConstructionViewMode;
+	bool bCanChangeConstructionViewMode = true;
+
+	bool bConstructionViewWireframe = false;
 
 	// Create dynamic mesh components from the cloth component's rest space info
 	void ReinitializeDynamicMeshComponents();
diff --git a/Engine/Plugins/Experimental/ChaosClothAssetEditor/Source/ChaosClothAssetEditor/Private/ChaosClothAsset/ClothEditorModeToolkit.cpp b/Engine/Plugins/Experimental/ChaosClothAssetEditor/Source/ChaosClothAssetEditor/Private/ChaosClothAsset/ClothEditorModeToolkit.cpp
--- a/Engine/Plugins/Experimental/ChaosClothAssetEditor/Source/ChaosClothAssetEditor/Private/ChaosClothAsset/ClothEditorModeToolkit.cpp
+++ b/Engine/Plugins/Experimental/ChaosClothAssetEditor/Source/ChaosClothAssetEditor/Private/ChaosClothAsset/ClothEditorModeToolkit.cpp
@@ -147,10 +147,6 @@ void FChaosClothAssetEditorModeToolkit::BuildToolPalette(FName PaletteIndex, cla
 	// TODO: make separate palettes to split up the tools tab
 	if (PaletteIndex == FBaseCharacterFXEditorModeToolkit::ToolsTabName)
 	{
-		// TODO: Re-add the remesh tool when we have a way to remesh both 2d and 3d rest space meshes at the same time
-		//ToolbarBuilder.AddToolBarButton(Commands.BeginRemeshTool);
-
-		ToolbarBuilder.AddToolBarButton(Commands.BeginAttributeEditorTool);
 		ToolbarBuilder.AddToolBarButton(Commands.BeginWeightMapPaintTool);
 		ToolbarBuilder.AddToolBarButton(Commands.BeginClothTrainingTool);
 		ToolbarBuilder.AddToolBarButton(Commands.BeginTransferSkinWeightsTool);
diff --git a/Engine/Plugins/Experimental/ChaosClothAssetEditor/Source/ChaosClothAssetEditor/Private/ChaosClothAsset/ClothEditorRestSpaceViewportClient.cpp b/Engine/Plugins/Experimental/ChaosClothAssetEditor/Source/ChaosClothAssetEditor/Private/ChaosClothAsset/ClothEditorRestSpaceViewportClient.cpp
--- a/Engine/Plugins/Experimental/ChaosClothAssetEditor/Source/ChaosClothAssetEditor/Private/ChaosClothAsset/ClothEditorRestSpaceViewportClient.cpp
+++ b/Engine/Plugins/Experimental/ChaosClothAssetEditor/Source/ChaosClothAssetEditor/Private/ChaosClothAsset/ClothEditorRestSpaceViewportClient.cpp
@@ -55,20 +55,18 @@ FChaosClothEditorRestSpaceViewportClient::FChaosClothEditorRestSpaceViewportClie
 	ModeTools->GetInteractiveToolsContext()->InputRouter->RegisterSource(this);
 }
 
-void FChaosClothEditorRestSpaceViewportClient::Set2DMode(bool In2DMode)
+void FChaosClothEditorRestSpaceViewportClient::SetConstructionViewMode(EClothPatternVertexType InViewMode)
 {
-	b2DMode = In2DMode;
+	ConstructionViewMode = InViewMode;
 
 	BehaviorSet->RemoveAll();
-	ModeTools->GetInteractiveToolsContext()->InputRouter->DeregisterSource(this);
 
-	if (b2DMode)
+	if (ConstructionViewMode == EClothPatternVertexType::Sim2D)
 	{
 		for (UInputBehavior* const Behavior : BehaviorsFor2DMode)
 		{
 			BehaviorSet->Add(Behavior);
 		}
-		ModeTools->GetInteractiveToolsContext()->InputRouter->RegisterSource(this);
 
 		const double AbsZ = FMath::Abs(ViewTransformPerspective.GetLocation().Z);
 		constexpr double CameraFarPlaneWorldZ = -10.0;
@@ -82,6 +80,8 @@ void FChaosClothEditorRestSpaceViewportClient::Set2DMode(bool In2DMode)
 		OverrideNearClipPlane(UE_KINDA_SMALL_NUMBER);
 	}
 
+	ModeTools->GetInteractiveToolsContext()->InputRouter->DeregisterSource(this);
+	ModeTools->GetInteractiveToolsContext()->InputRouter->RegisterSource(this);
 }
 
 const UInputBehaviorSet* FChaosClothEditorRestSpaceViewportClient::GetInputBehaviors() const
@@ -99,7 +99,7 @@ void FChaosClothEditorRestSpaceViewportClient::AddReferencedObjects(FReferenceCo
 
 bool FChaosClothEditorRestSpaceViewportClient::ShouldOrbitCamera() const
 {
-	if (b2DMode)
+	if (ConstructionViewMode == EClothPatternVertexType::Sim2D)
 	{
 		return false;
 	}
@@ -122,7 +122,7 @@ bool FChaosClothEditorRestSpaceViewportClient::InputKey(const FInputKeyEventArgs
 		}
 	}
 
-	if (!b2DMode)
+	if (ConstructionViewMode != EClothPatternVertexType::Sim2D)
 	{
 		return FEditorViewportClient::InputKey(EventArgs);
 	}
@@ -197,47 +197,6 @@ bool FChaosClothEditorRestSpaceViewportClient::InputKey(const FInputKeyEventArgs
 	return bHandled;
 }
 
-
-void FChaosClothEditorRestSpaceViewportClient::ProcessClick(FSceneView& View, HHitProxy* HitProxy, FKey Key, EInputEvent Event, uint32 HitX, uint32 HitY)
-{
-	FEditorViewportClient::ProcessClick(View, HitProxy, Key, Event, HitX, HitY);
-
-	// TODO: Add/modify selection if modifier keys are pressed
-	USelection* SelectedComponents = ModeTools->GetSelectedComponents();
-	SelectedComponents->Modify();
-	SelectedComponents->BeginBatchSelectOperation();
-
-	TArray<UDynamicMeshComponent*> PreviouslySelectedComponents;
-	SelectedComponents->GetSelectedObjects<UDynamicMeshComponent>(PreviouslySelectedComponents);
-	SelectedComponents->DeselectAll(UDynamicMeshComponent::StaticClass());
-
-	if (HitProxy && HitProxy->IsA(HActor::StaticGetType()))
-	{
-		const HActor* ActorProxy = static_cast<HActor*>(HitProxy);
-		if (ActorProxy && ActorProxy->Actor)
-		{
-			const AActor* Actor = ActorProxy->Actor;
-			const TSet<UActorComponent*>& OwnedComponents = Actor->GetComponents();
-			for (UActorComponent* Component : OwnedComponents)
-			{
-				if (UDynamicMeshComponent* DynamicMeshComp = Cast<UDynamicMeshComponent>(Component))
-				{
-					SelectedComponents->Select(DynamicMeshComp);
-					DynamicMeshComp->PushSelectionToProxy();
-				}
-			}
-		}
-	}
-
-	SelectedComponents->EndBatchSelectOperation();
-
-	for (UDynamicMeshComponent* Component : PreviouslySelectedComponents)
-	{
-		Component->PushSelectionToProxy();
-	}
-
-}
-
 void FChaosClothEditorRestSpaceViewportClient::SetEditorViewportWidget(TWeakPtr<SEditorViewport> InEditorViewportWidget)
 {
 	EditorViewportWidget = InEditorViewportWidget;
diff --git a/Engine/Plugins/Experimental/ChaosClothAssetEditor/Source/ChaosClothAssetEditor/Private/ChaosClothAsset/ClothEditorRestSpaceViewportClient.h b/Engine/Plugins/Experimental/ChaosClothAssetEditor/Source/ChaosClothAssetEditor/Private/ChaosClothAsset/ClothEditorRestSpaceViewportClient.h
--- a/Engine/Plugins/Experimental/ChaosClothAssetEditor/Source/ChaosClothAssetEditor/Private/ChaosClothAsset/ClothEditorRestSpaceViewportClient.h
+++ b/Engine/Plugins/Experimental/ChaosClothAssetEditor/Source/ChaosClothAssetEditor/Private/ChaosClothAsset/ClothEditorRestSpaceViewportClient.h
@@ -6,6 +6,8 @@
 #include "EditorViewportClient.h"
 #include "Behaviors/2DViewportBehaviorTargets.h" // FEditor2DScrollBehaviorTarget, FEditor2DMouseWheelZoomBehaviorTarget
 #include "InputBehaviorSet.h"
+#include "ChaosClothAsset/ClothEditorMode.h"
+#include "ChaosClothAsset/ClothPatternVertexType.h"
 
 class UInputBehaviorSet;
 
@@ -26,19 +28,18 @@ public:
 	// FGCObject
 	virtual void AddReferencedObjects(FReferenceCollector& Collector) override;
 
-	virtual void ProcessClick(FSceneView& View, HHitProxy* HitProxy, FKey Key, EInputEvent Event, uint32 HitX, uint32 HitY) override;
 	virtual bool InputKey(const FInputKeyEventArgs& EventArgs) override;
 
 	virtual bool ShouldOrbitCamera() const override;
 
-	void Set2DMode(bool In2DMode);
+	void SetConstructionViewMode(EClothPatternVertexType InViewMode);
 
 	void SetEditorViewportWidget(TWeakPtr<SEditorViewport> InEditorViewportWidget);
 	void SetToolCommandList(TWeakPtr<FUICommandList> ToolCommandList);
 
 private:
 
-	bool b2DMode = false;
+	EClothPatternVertexType ConstructionViewMode = EClothPatternVertexType::Sim2D;
 
 	TObjectPtr<UInputBehaviorSet> BehaviorSet;
 
diff --git a/Engine/Plugins/Experimental/ChaosClothAssetEditor/Source/ChaosClothAssetEditor/Private/ChaosClothAsset/ClothEditorStyle.cpp b/Engine/Plugins/Experimental/ChaosClothAssetEditor/Source/ChaosClothAssetEditor/Private/ChaosClothAsset/ClothEditorStyle.cpp
--- a/Engine/Plugins/Experimental/ChaosClothAssetEditor/Source/ChaosClothAssetEditor/Private/ChaosClothAsset/ClothEditorStyle.cpp
+++ b/Engine/Plugins/Experimental/ChaosClothAssetEditor/Source/ChaosClothAssetEditor/Private/ChaosClothAsset/ClothEditorStyle.cpp
@@ -59,16 +59,20 @@ FChaosClothAssetEditorStyle::FChaosClothAssetEditorStyle()
 	PropertyNameString = "ChaosClothAssetEditor." + FChaosClothAssetEditorCommands::BeginTransferSkinWeightsToolIdentifier;
 	Set(*PropertyNameString, new FSlateImageBrush(FChaosClothAssetEditorStyle::InContent("Icons/ModelingAttributePaint_x40", ".png"), ToolbarIconSize));
 
+	//
+	// Construction Viewport
+	// 
+	
 	// TODO: Get a permanent icon to replace this placeholder
-	PropertyNameString = "ChaosClothAssetEditor." + FChaosClothAssetEditorCommands::TogglePatternModeIdentifier;
-	Set(*PropertyNameString, new FSlateImageBrush(FChaosClothAssetEditorStyle::InContent("Icons/TogglePatternMode_40x", ".png"), ToolbarIconSize));
-
-	// TODO: Get a permanent icon to replace this placeholder
-	PropertyNameString = "ChaosClothAssetEditor." + FChaosClothAssetEditorCommands::ToggleSimMeshWireframeIdentifier;
+	PropertyNameString = "ChaosClothAssetEditor." + FChaosClothAssetEditorCommands::ToggleConstructionViewWireframeIdentifier;
 	Set(*PropertyNameString, new CORE_IMAGE_BRUSH_SVG("Starship/EditorViewport/grid", ToolbarIconSize));
 
+	//
+	// Preview Viewport
+	// 
+
 	// TODO: Get a permanent icon to replace this placeholder
-	PropertyNameString = "ChaosClothAssetEditor." + FChaosClothAssetEditorCommands::ToggleRenderMeshWireframeIdentifier;
+	PropertyNameString = "ChaosClothAssetEditor." + FChaosClothAssetEditorCommands::TogglePreviewWireframeIdentifier;
 	Set(*PropertyNameString, new CORE_IMAGE_BRUSH_SVG("Starship/EditorViewport/grid", ToolbarIconSize));
 
 	// TODO: Get a permanent icon to replace this placeholder
diff --git a/Engine/Plugins/Experimental/ChaosClothAssetEditor/Source/ChaosClothAssetEditor/Private/ChaosClothAsset/ClothEditorToolkit.cpp b/Engine/Plugins/Experimental/ChaosClothAssetEditor/Source/ChaosClothAssetEditor/Private/ChaosClothAsset/ClothEditorToolkit.cpp
--- a/Engine/Plugins/Experimental/ChaosClothAssetEditor/Source/ChaosClothAssetEditor/Private/ChaosClothAsset/ClothEditorToolkit.cpp
+++ b/Engine/Plugins/Experimental/ChaosClothAssetEditor/Source/ChaosClothAssetEditor/Private/ChaosClothAsset/ClothEditorToolkit.cpp
@@ -1021,7 +1021,6 @@ void FChaosClothAssetEditorToolkit::OnNodeDeleted(const TSet<UObject*>& DeletedN
 	if (ClothMode)
 	{
 		ClothMode->SetSelectedClothCollection(nullptr);
-		ClothMode->RefocusRestSpaceViewportClient();
 	}
 
 }
@@ -1050,7 +1049,6 @@ void FChaosClothAssetEditorToolkit::OnClothAssetChanged()
 		SetEditingObject(ClothAsset);
 
 		ClothPreviewViewportClient->FocusViewportOnBox(ClothMode->PreviewBoundingBox());
-		ClothMode->RefocusRestSpaceViewportClient();
 
 		PopulateOutliner();
 	}
diff --git a/Engine/Plugins/Experimental/ChaosClothAssetEditor/Source/ChaosClothAssetEditor/Private/ChaosClothAsset/SClothEditor3DViewport.cpp b/Engine/Plugins/Experimental/ChaosClothAssetEditor/Source/ChaosClothAssetEditor/Private/ChaosClothAsset/SClothEditor3DViewport.cpp
--- a/Engine/Plugins/Experimental/ChaosClothAssetEditor/Source/ChaosClothAssetEditor/Private/ChaosClothAsset/SClothEditor3DViewport.cpp
+++ b/Engine/Plugins/Experimental/ChaosClothAssetEditor/Source/ChaosClothAssetEditor/Private/ChaosClothAsset/SClothEditor3DViewport.cpp
@@ -66,17 +66,7 @@ void SChaosClothAssetEditor3DViewport::BindCommands()
 	const FChaosClothAssetEditorCommands& CommandInfos = FChaosClothAssetEditorCommands::Get();
 
 	CommandList->MapAction(
-		CommandInfos.ToggleSimMeshWireframe,
-		FExecuteAction::CreateLambda([this]() 
-		{
-			TSharedPtr<FChaosClothAssetEditor3DViewportClient> ClothViewportClient = StaticCastSharedPtr<FChaosClothAssetEditor3DViewportClient>(Client);
-			ClothViewportClient->EnableSimMeshWireframe(!ClothViewportClient->SimMeshWireframeEnabled());
-		}),
-		FCanExecuteAction::CreateLambda([this]() { return true; }),
-		FIsActionChecked::CreateLambda([this]() { return StaticCastSharedPtr<FChaosClothAssetEditor3DViewportClient>(Client)->SimMeshWireframeEnabled(); }));
-
-	CommandList->MapAction(
-		CommandInfos.ToggleRenderMeshWireframe,
+		CommandInfos.TogglePreviewWireframe,
 		FExecuteAction::CreateLambda([this]()
 		{
 			TSharedPtr<FChaosClothAssetEditor3DViewportClient> ClothViewportClient = StaticCastSharedPtr<FChaosClothAssetEditor3DViewportClient>(Client);
diff --git a/Engine/Plugins/Experimental/ChaosClothAssetEditor/Source/ChaosClothAssetEditor/Private/ChaosClothAsset/SClothEditor3DViewportToolBar.cpp b/Engine/Plugins/Experimental/ChaosClothAssetEditor/Source/ChaosClothAssetEditor/Private/ChaosClothAsset/SClothEditor3DViewportToolBar.cpp
--- a/Engine/Plugins/Experimental/ChaosClothAssetEditor/Source/ChaosClothAssetEditor/Private/ChaosClothAsset/SClothEditor3DViewportToolBar.cpp
+++ b/Engine/Plugins/Experimental/ChaosClothAssetEditor/Source/ChaosClothAssetEditor/Private/ChaosClothAsset/SClothEditor3DViewportToolBar.cpp
@@ -94,15 +94,12 @@ TSharedRef<SWidget> SChaosClothAssetEditor3DViewportToolBar::MakeToolBar(const T
 	ToolbarBuilder.SetStyle(&FAppStyle::Get(), ToolBarStyle);
 	ToolbarBuilder.SetLabelVisibility(EVisibility::Collapsed);
 
-	FString PropertyNameString;
-
 	ToolbarBuilder.BeginSection("Visualization");
 	ToolbarBuilder.BeginBlockGroup();
 	{
 		// TODO: Add button to toggle between sim and render meshes
 
-		PropertyNameString = "ChaosClothAssetEditor." + FChaosClothAssetEditorCommands::ToggleRenderMeshWireframeIdentifier;
-		ToolbarBuilder.AddToolBarButton(FChaosClothAssetEditorCommands::Get().ToggleRenderMeshWireframe,
+		ToolbarBuilder.AddToolBarButton(FChaosClothAssetEditorCommands::Get().TogglePreviewWireframe,
 			NAME_None,
 			TAttribute<FText>(),
 			TAttribute<FText>(),
@@ -115,7 +112,6 @@ TSharedRef<SWidget> SChaosClothAssetEditor3DViewportToolBar::MakeToolBar(const T
 	ToolbarBuilder.BeginSection("Sim Controls");
 	ToolbarBuilder.BeginBlockGroup();
 	{
-		PropertyNameString = "ChaosClothAssetEditor." + FChaosClothAssetEditorCommands::HardResetSimulationIdentifier;
 		ToolbarBuilder.AddToolBarButton(FChaosClothAssetEditorCommands::Get().HardResetSimulation,
 			NAME_None,
 			TAttribute<FText>(),
@@ -123,7 +119,6 @@ TSharedRef<SWidget> SChaosClothAssetEditor3DViewportToolBar::MakeToolBar(const T
 			TAttribute<FSlateIcon>(),
 			FName(*FChaosClothAssetEditorCommands::HardResetSimulationIdentifier));
 
-		PropertyNameString = "ChaosClothAssetEditor." + FChaosClothAssetEditorCommands::SoftResetSimulationIdentifier;
 		ToolbarBuilder.AddToolBarButton(FChaosClothAssetEditorCommands::Get().SoftResetSimulation,
 			NAME_None,
 			TAttribute<FText>(),
@@ -131,7 +126,6 @@ TSharedRef<SWidget> SChaosClothAssetEditor3DViewportToolBar::MakeToolBar(const T
 			TAttribute<FSlateIcon>(),
 			FName(FChaosClothAssetEditorCommands::SoftResetSimulationIdentifier));
 
-		PropertyNameString = "ChaosClothAssetEditor." + FChaosClothAssetEditorCommands::ToggleSimulationSuspendedIdentifier;
 		ToolbarBuilder.AddToolBarButton(FChaosClothAssetEditorCommands::Get().ToggleSimulationSuspended,
 			NAME_None,
 			TAttribute<FText>(),
diff --git a/Engine/Plugins/Experimental/ChaosClothAssetEditor/Source/ChaosClothAssetEditor/Private/ChaosClothAsset/SClothEditorRestSpaceViewport.cpp b/Engine/Plugins/Experimental/ChaosClothAssetEditor/Source/ChaosClothAssetEditor/Private/ChaosClothAsset/SClothEditorRestSpaceViewport.cpp
--- a/Engine/Plugins/Experimental/ChaosClothAssetEditor/Source/ChaosClothAssetEditor/Private/ChaosClothAsset/SClothEditorRestSpaceViewport.cpp
+++ b/Engine/Plugins/Experimental/ChaosClothAssetEditor/Source/ChaosClothAssetEditor/Private/ChaosClothAsset/SClothEditorRestSpaceViewport.cpp
@@ -9,6 +9,18 @@
 
 #define LOCTEXT_NAMESPACE "SChaosClothAssetEditorRestSpaceViewport"
 
+UChaosClothAssetEditorMode* SChaosClothAssetEditorRestSpaceViewport::GetEdMode() const
+{
+	if (const FEditorModeTools* const EditorModeTools = Client->GetModeTools())
+	{
+		if (UChaosClothAssetEditorMode* const ClothEdMode = Cast<UChaosClothAssetEditorMode>(EditorModeTools->GetActiveScriptableMode(UChaosClothAssetEditorMode::EM_ChaosClothAssetEditorModeId)))
+		{
+			return ClothEdMode;
+		}
+	}
+	return nullptr;
+}
+
 void SChaosClothAssetEditorRestSpaceViewport::BindCommands()
 {
 	using namespace UE::Chaos::ClothAsset;
@@ -18,30 +30,112 @@ void SChaosClothAssetEditorRestSpaceViewport::BindCommands()
 	const FChaosClothAssetEditorCommands& CommandInfos = FChaosClothAssetEditorCommands::Get();
 
 	CommandList->MapAction(
-		CommandInfos.TogglePatternMode,
+		CommandInfos.SetConstructionMode2D,
 		FExecuteAction::CreateLambda([this]()
-	{
-		const FEditorModeTools* const EditorModeTools = Client->GetModeTools();
-		UChaosClothAssetEditorMode* const ClothEdMode = Cast<UChaosClothAssetEditorMode>(EditorModeTools->GetActiveScriptableMode(UChaosClothAssetEditorMode::EM_ChaosClothAssetEditorModeId));
+		{
+			if (UChaosClothAssetEditorMode* const ClothEdMode = GetEdMode())
+			{
+				ClothEdMode->SetConstructionViewMode(EClothPatternVertexType::Sim2D);
+			}
+		}),
+		FCanExecuteAction::CreateLambda([this]() 
+		{ 
+			if (const UChaosClothAssetEditorMode* const ClothEdMode = GetEdMode())
+			{
+				return ClothEdMode->CanChangeConstructionViewMode();
+			}
+			return false; 
+		}),
+		FIsActionChecked::CreateLambda([this]() 
+		{
+			if (const UChaosClothAssetEditorMode* const ClothEdMode = GetEdMode())
+			{
+				return ClothEdMode->GetConstructionViewMode() == EClothPatternVertexType::Sim2D;
+			}
+			return false;
+		}));
+
 
-		if (ClothEdMode)
+	CommandList->MapAction(
+		CommandInfos.SetConstructionMode3D,
+		FExecuteAction::CreateLambda([this]()
 		{
-			ClothEdMode->TogglePatternMode();
-		}
-	}),
-	FCanExecuteAction::CreateLambda([this]() 
-	{ 
-		const FEditorModeTools* const EditorModeTools = Client->GetModeTools();
-		const UChaosClothAssetEditorMode* const ClothEdMode = Cast<UChaosClothAssetEditorMode>(EditorModeTools->GetActiveScriptableMode(UChaosClothAssetEditorMode::EM_ChaosClothAssetEditorModeId));
+			if (UChaosClothAssetEditorMode* const ClothEdMode = GetEdMode()) 
+			{
+				ClothEdMode->SetConstructionViewMode(EClothPatternVertexType::Sim3D);
+			}
+		}),
+		FCanExecuteAction::CreateLambda([this]() 
+		{ 
+			if (const UChaosClothAssetEditorMode* const ClothEdMode = GetEdMode())
+			{
+				return ClothEdMode->CanChangeConstructionViewMode();
+			}
+			return false; 
+		}),
+		FIsActionChecked::CreateLambda([this]() 
+		{
+			if (const UChaosClothAssetEditorMode* const ClothEdMode = GetEdMode())
+			{
+				return ClothEdMode->GetConstructionViewMode() == EClothPatternVertexType::Sim3D;
+			}
+			return false;
+		}));
+
 
-		if (ClothEdMode)
+
+	CommandList->MapAction(
+		CommandInfos.SetConstructionModeRender,
+		FExecuteAction::CreateLambda([this]()
 		{
-			return ClothEdMode->CanTogglePatternMode();
-		}
+			if (UChaosClothAssetEditorMode* ClothEdMode = GetEdMode()) 
+			{
+				ClothEdMode->SetConstructionViewMode(EClothPatternVertexType::Render);
+			}
+		}),
+		FCanExecuteAction::CreateLambda([this]() 
+		{ 
+			if (const UChaosClothAssetEditorMode* const ClothEdMode = GetEdMode())
+			{
+				return ClothEdMode->CanChangeConstructionViewMode();
+			}
+			return false; 
+		}),
+		FIsActionChecked::CreateLambda([this]() 
+		{
+			if (const UChaosClothAssetEditorMode* const ClothEdMode = GetEdMode())
+			{
+				return ClothEdMode->GetConstructionViewMode() == EClothPatternVertexType::Render;
+			}
+			return false;
+		}));
+
+	CommandList->MapAction(
+		CommandInfos.ToggleConstructionViewWireframe,
+		FExecuteAction::CreateLambda([this]()
+		{
+			const FEditorModeTools* const EditorModeTools = Client->GetModeTools();
+			UChaosClothAssetEditorMode* const ClothEdMode = Cast<UChaosClothAssetEditorMode>(EditorModeTools->GetActiveScriptableMode(UChaosClothAssetEditorMode::EM_ChaosClothAssetEditorModeId));
+
+			if (ClothEdMode)
+			{
+				ClothEdMode->ToggleConstructionViewWireframe();
+			}
+		}),
+		FCanExecuteAction::CreateLambda([this]() 
+		{ 
+			return true; 
+		}),
+		FIsActionChecked::CreateLambda([this]() 
+		{
+			if (const UChaosClothAssetEditorMode* const ClothEdMode = GetEdMode())
+			{
+				return ClothEdMode->IsConstructionViewWireframeActive();
+			}
+			return false;
+		}));
+
 
-		return false; 
-	}),
-	EUIActionRepeatMode::RepeatDisabled);
 }
 
 TSharedPtr<SWidget> SChaosClothAssetEditorRestSpaceViewport::MakeViewportToolbar()
diff --git a/Engine/Plugins/Experimental/ChaosClothAssetEditor/Source/ChaosClothAssetEditor/Private/ChaosClothAsset/SClothEditorRestSpaceViewport.h b/Engine/Plugins/Experimental/ChaosClothAssetEditor/Source/ChaosClothAssetEditor/Private/ChaosClothAsset/SClothEditorRestSpaceViewport.h
--- a/Engine/Plugins/Experimental/ChaosClothAssetEditor/Source/ChaosClothAssetEditor/Private/ChaosClothAsset/SClothEditorRestSpaceViewport.h
+++ b/Engine/Plugins/Experimental/ChaosClothAssetEditor/Source/ChaosClothAssetEditor/Private/ChaosClothAsset/SClothEditorRestSpaceViewport.h
@@ -6,6 +6,8 @@
 #include "SBaseCharacterFXEditorViewport.h"
 #include "SCommonEditorViewportToolbarBase.h"
 
+class UChaosClothAssetEditorMode;
+
 class CHAOSCLOTHASSETEDITOR_API SChaosClothAssetEditorRestSpaceViewport : public SBaseCharacterFXEditorViewport, public ICommonEditorViewportToolbarInfoProvider
 {
 public:
@@ -20,4 +22,8 @@ public:
 	virtual TSharedPtr<FExtender> GetExtenders() const override;
 	virtual void OnFloatingButtonClicked() override;
 
+private:
+
+	UChaosClothAssetEditorMode* GetEdMode() const;
+
 };
diff --git a/Engine/Plugins/Experimental/ChaosClothAssetEditor/Source/ChaosClothAssetEditor/Private/ChaosClothAsset/SClothEditorRestSpaceViewportToolBar.cpp b/Engine/Plugins/Experimental/ChaosClothAssetEditor/Source/ChaosClothAssetEditor/Private/ChaosClothAsset/SClothEditorRestSpaceViewportToolBar.cpp
--- a/Engine/Plugins/Experimental/ChaosClothAssetEditor/Source/ChaosClothAssetEditor/Private/ChaosClothAsset/SClothEditorRestSpaceViewportToolBar.cpp
+++ b/Engine/Plugins/Experimental/ChaosClothAssetEditor/Source/ChaosClothAssetEditor/Private/ChaosClothAsset/SClothEditorRestSpaceViewportToolBar.cpp
@@ -10,6 +10,9 @@
 #include "Widgets/SBoxPanel.h"
 #include "Editor/UnrealEd/Public/SEditorViewportToolBarMenu.h"
 #include "SEditorViewportViewMenu.h"
+#include "EditorModeManager.h"
+#include "ChaosClothAsset/ClothEditorMode.h"
+#include "ToolMenus.h"
 
 #define LOCTEXT_NAMESPACE "SChaosClothAssetEditorRestSpaceViewportToolBar"
 
@@ -20,6 +23,8 @@ void SChaosClothAssetEditorRestSpaceViewportToolBar::Construct(const FArguments&
 	ChaosClothAssetEditorRestSpaceViewportPtr = InChaosClothAssetEditorViewport;
 	CommandList = InArgs._CommandList;
 
+	RegisterViewModeMenuContent();
+
 	const FMargin ToolbarSlotPadding(4.0f, 1.0f);
 	TSharedPtr<SHorizontalBox> MainBoxPtr;
 
@@ -90,19 +95,82 @@ TSharedRef<SWidget> SChaosClothAssetEditorRestSpaceViewportToolBar::MakeToolBar(
 	ToolbarBuilder.SetStyle(&FAppStyle::Get(), ToolBarStyle);
 	ToolbarBuilder.SetLabelVisibility(EVisibility::Collapsed);
 
-	ToolbarBuilder.BeginSection("SimControls");
+	ToolbarBuilder.BeginSection("View Controls");
+	ToolbarBuilder.BeginBlockGroup();
 	{
-		ToolbarBuilder.BeginBlockGroup();
-	    
-		static FName TogglePatternModeName = FName(TEXT("PatternMode"));
-		ToolbarBuilder.AddToolBarButton(FChaosClothAssetEditorCommands::Get().TogglePatternMode);
-
-		ToolbarBuilder.EndBlockGroup();
+		ToolbarBuilder.AddToolBarButton(FChaosClothAssetEditorCommands::Get().ToggleConstructionViewWireframe);
 	}
+	ToolbarBuilder.EndBlockGroup();
 
-	ToolbarBuilder.EndSection();
+	ToolbarBuilder.BeginBlockGroup();
+	{
+		// View mode selector (2D/3D/Render)
+		ToolbarBuilder.AddWidget(
+			SNew(SEditorViewportToolbarMenu)
+			.ParentToolBar(SharedThis(this))
+			.Cursor(EMouseCursor::Default)
+			.Label(this, &SChaosClothAssetEditorRestSpaceViewportToolBar::GetViewModeMenuLabel)
+			.LabelIcon(this, &SChaosClothAssetEditorRestSpaceViewportToolBar::GetViewModeMenuLabelIcon)
+			.OnGetMenuContent(this, &SChaosClothAssetEditorRestSpaceViewportToolBar::GenerateViewModeMenuContent)
+		);
+	}
+	ToolbarBuilder.EndBlockGroup();
+	ToolbarBuilder.EndSection(); // View Controls
 
 	return ToolbarBuilder.MakeWidget();
 }
 
+
+FText SChaosClothAssetEditorRestSpaceViewportToolBar::GetViewModeMenuLabel() const
+{
+	FText Label = LOCTEXT("ConstructionViewMenuTitle_Default", "View");
+
+	TSharedPtr< SChaosClothAssetEditorRestSpaceViewport > PinnedViewport = ChaosClothAssetEditorRestSpaceViewportPtr.Pin();
+	if (PinnedViewport.IsValid())
+	{
+		const TSharedPtr<FEditorViewportClient> ViewportClient = PinnedViewport->GetViewportClient();
+		check(ViewportClient.IsValid());
+		const FEditorModeTools* const EditorModeTools = ViewportClient->GetModeTools();
+		UChaosClothAssetEditorMode* const ClothEdMode = Cast<UChaosClothAssetEditorMode>(EditorModeTools->GetActiveScriptableMode(UChaosClothAssetEditorMode::EM_ChaosClothAssetEditorModeId));
+		if (ClothEdMode)
+		{
+			switch (ClothEdMode->GetConstructionViewMode())
+			{
+			case UE::Chaos::ClothAsset::EClothPatternVertexType::Sim2D:
+				Label = LOCTEXT("ConstructionViewMenuTitle_Sim2D", "2D Sim");
+				break;
+			case UE::Chaos::ClothAsset::EClothPatternVertexType::Sim3D:
+				Label = LOCTEXT("ConstructionViewMenuTitle_Sim3D", "3D Sim");
+				break;
+			case UE::Chaos::ClothAsset::EClothPatternVertexType::Render:
+				Label = LOCTEXT("ConstructionViewMenuTitle_Render", "Render");
+				break;
+			}
+		}
+	}
+
+	return Label;
+}
+
+const FSlateBrush* SChaosClothAssetEditorRestSpaceViewportToolBar::GetViewModeMenuLabelIcon() const
+{
+	return FStyleDefaults::GetNoBrush();
+}
+
+TSharedRef<SWidget> SChaosClothAssetEditorRestSpaceViewportToolBar::GenerateViewModeMenuContent() const
+{
+	return UToolMenus::Get()->GenerateWidget("ChaosClothAssetEditor.ConstructionViewModeMenu", FToolMenuContext(CommandList));
+}
+
+void SChaosClothAssetEditorRestSpaceViewportToolBar::RegisterViewModeMenuContent()
+{
+	UToolMenu* const Menu = UToolMenus::Get()->RegisterMenu("ChaosClothAssetEditor.ConstructionViewModeMenu");
+
+	FToolMenuSection& Section = Menu->FindOrAddSection("ConstructionViewModeMenuSection");
+	Section.AddMenuEntry(UE::Chaos::ClothAsset::FChaosClothAssetEditorCommands::Get().SetConstructionMode2D);
+	Section.AddMenuEntry(UE::Chaos::ClothAsset::FChaosClothAssetEditorCommands::Get().SetConstructionMode3D);
+	Section.AddMenuEntry(UE::Chaos::ClothAsset::FChaosClothAssetEditorCommands::Get().SetConstructionModeRender);
+}
+
+
 #undef LOCTEXT_NAMESPACE
diff --git a/Engine/Plugins/Experimental/ChaosClothAssetEditor/Source/ChaosClothAssetEditor/Private/ChaosClothAsset/SClothEditorRestSpaceViewportToolBar.h b/Engine/Plugins/Experimental/ChaosClothAssetEditor/Source/ChaosClothAssetEditor/Private/ChaosClothAsset/SClothEditorRestSpaceViewportToolBar.h
--- a/Engine/Plugins/Experimental/ChaosClothAssetEditor/Source/ChaosClothAssetEditor/Private/ChaosClothAsset/SClothEditorRestSpaceViewportToolBar.h
+++ b/Engine/Plugins/Experimental/ChaosClothAssetEditor/Source/ChaosClothAssetEditor/Private/ChaosClothAsset/SClothEditorRestSpaceViewportToolBar.h
@@ -26,6 +26,12 @@ private:
 	TSharedRef<SWidget> MakeDisplayToolBar(const TSharedPtr<FExtender> InExtenders);
 	TSharedRef<SWidget> MakeToolBar(const TSharedPtr<FExtender> InExtenders);
 
+	FText GetViewModeMenuLabel() const;
+	const FSlateBrush* GetViewModeMenuLabelIcon() const;
+	TSharedRef<SWidget> GenerateViewModeMenuContent() const;
+
+	void RegisterViewModeMenuContent();
+
 	/** The viewport that we are in */
 	TWeakPtr<class SChaosClothAssetEditorRestSpaceViewport> ChaosClothAssetEditorRestSpaceViewportPtr;
 
diff --git a/Engine/Plugins/Experimental/ChaosClothAssetEditor/Source/ChaosClothAssetTools/Private/ChaosClothAsset/ClothPatternToDynamicMesh.cpp b/Engine/Plugins/Experimental/ChaosClothAssetEditor/Source/ChaosClothAssetTools/Private/ChaosClothAsset/ClothPatternToDynamicMesh.cpp
--- a/Engine/Plugins/Experimental/ChaosClothAssetEditor/Source/ChaosClothAssetTools/Private/ChaosClothAsset/ClothPatternToDynamicMesh.cpp
+++ b/Engine/Plugins/Experimental/ChaosClothAssetEditor/Source/ChaosClothAssetTools/Private/ChaosClothAsset/ClothPatternToDynamicMesh.cpp
@@ -11,10 +11,13 @@
 #include "SkeletalMeshAttributes.h"
 #include "Animation/Skeleton.h"
 
+namespace UE::Chaos::ClothAsset
+{
+
 //
 // Wrapper for accessing a Cloth Pattern. Implements the interface expected by TToDynamicMesh<>.
 //
-class FClothSimPatternWrapper
+class FClothPatternWrapper
 {
 public:
 
@@ -26,20 +29,20 @@ public:
 	typedef int32 NormalIDType;
 	typedef int32 ColorIDType;
 
-	FClothSimPatternWrapper(const UE::Chaos::ClothAsset::FCollectionClothConstFacade& ClothFacade, int32 LodIndex, int32 PatternIndex, bool bGet2DPatternVertices) :
-		bGet2DPatternVertices(bGet2DPatternVertices),
+	FClothPatternWrapper(const FCollectionClothConstFacade& ClothFacade, int32 LodIndex, int32 PatternIndex, EClothPatternVertexType VertexDataType) :
+		VertexDataType(VertexDataType),
 		Pattern(ClothFacade.GetLod(LodIndex).GetPattern(PatternIndex))
 	{
-		const int32 NumSimVertices = Pattern.GetNumSimVertices();
-		VertIDs.SetNum(NumSimVertices);
-		for (int32 VtxIndex = 0; VtxIndex < NumSimVertices; ++VtxIndex)
+		const int32 NumVertices = (VertexDataType == EClothPatternVertexType::Render) ? Pattern.GetNumRenderVertices() : Pattern.GetNumSimVertices();
+		VertIDs.SetNum(NumVertices);
+		for (int32 VtxIndex = 0; VtxIndex < NumVertices; ++VtxIndex)
 		{
 			VertIDs[VtxIndex] = VtxIndex;
 		}
 
-		const int32 NumSimFaces = Pattern.GetNumSimFaces();
-		TriIDs.SetNum(NumSimFaces);
-		for (int32 TriIndex = 0; TriIndex < NumSimFaces; ++TriIndex)
+		const int32 NumFaces = (VertexDataType == EClothPatternVertexType::Render) ? Pattern.GetNumRenderFaces() : Pattern.GetNumSimFaces();
+		TriIDs.SetNum(NumFaces);
+		for (int32 TriIndex = 0; TriIndex < NumFaces; ++TriIndex)
 		{
 			TriIDs[TriIndex] = TriIndex;
 		}
@@ -70,7 +73,7 @@ public:
 
 	int32 NumUVLayers() const
 	{
-		return 0;	// No UVs for Sim mesh
+		return (VertexDataType == EClothPatternVertexType::Render) ? 1 : 0;		// No UVs for Sim mesh
 	}
 
 	int32 NumWeightMapLayers() const
@@ -96,18 +99,29 @@ public:
 
 	FVector3d GetPosition(VertIDType VtxID) const
 	{
-		if (bGet2DPatternVertices)
+		switch (VertexDataType)
 		{
-			const TConstArrayView<FVector2f> SimPositions = Pattern.GetSimPosition();
-			const FVector2f& Pos = SimPositions[VtxID];
-			return FVector3d(Pos[0], Pos[1], 0.0);
-		} 
-		else
-		{
-			const TConstArrayView<FVector3f> RestPositions = Pattern.GetSimRestPosition();
-			const FVector3f& Pos = RestPositions[VtxID];
-			return FVector3d(Pos[0], Pos[1], Pos[2]);
-		}
+			case EClothPatternVertexType::Render:
+			{
+				const TConstArrayView<FVector3f> RenderPositions = Pattern.GetRenderPosition();
+				return FVector3d(RenderPositions[VtxID]);
+			}
+			case EClothPatternVertexType::Sim2D:
+			{
+				const TConstArrayView<FVector2f> SimPositions = Pattern.GetSimPosition();
+				const FVector2f& Pos = SimPositions[VtxID];
+				return FVector3d(Pos[0], Pos[1], 0.0);
+			}
+			case EClothPatternVertexType::Sim3D:
+			{
+				const TConstArrayView<FVector3f> RestPositions = Pattern.GetSimRestPosition();
+				const FVector3f& Pos = RestPositions[VtxID];
+				return FVector3d(Pos[0], Pos[1], Pos[2]);
+			}
+			default:
+				checkNoEntry();
+				return FVector3d();
+		};
 	}
 
 	// --"Index Buffer" info
@@ -118,11 +132,23 @@ public:
 
 	bool GetTri(TriIDType TriID, VertIDType& VID0, VertIDType& VID1, VertIDType& VID2) const
 	{
-		const FIntVector Face = Pattern.GetSimIndices()[TriID];          // Indices are using an LOD based indexation
-		const int32 SimVerticesOffset = Pattern.GetSimVerticesOffset();  // and need to have the pattern offset removed
-		VID0 = Face[0] - SimVerticesOffset;
-		VID1 = Face[1] - SimVerticesOffset;
-		VID2 = Face[2] - SimVerticesOffset;
+		FIntVector Face;		// Indices are using an LOD based indexation
+		int32 VerticesOffset;	// and need to have the pattern offset removed
+
+		if (VertexDataType == EClothPatternVertexType::Render)
+		{
+			Face = Pattern.GetRenderIndices()[TriID];
+			VerticesOffset = Pattern.GetRenderVerticesOffset();
+		}
+		else
+		{
+			Face = Pattern.GetSimIndices()[TriID];
+			VerticesOffset = Pattern.GetSimVerticesOffset();
+		}
+
+		VID0 = Face[0] - VerticesOffset;
+		VID1 = Face[1] - VerticesOffset;
+		VID2 = Face[2] - VerticesOffset;
 
 		return true;
 	}
@@ -134,7 +160,7 @@ public:
 
 	bool HasTangents() const
 	{
-		return false;
+		return VertexDataType == EClothPatternVertexType::Render;
 	}
 
 	bool HasBiTangents() const
@@ -144,7 +170,7 @@ public:
 
 	bool HasColors() const
 	{
-		return false;
+		return VertexDataType == EClothPatternVertexType::Render;
 	}
 
 	// -- Access to per-wedge attributes -- //
@@ -159,8 +185,9 @@ public:
 
 	FVector2f GetWedgeUV(int32 UVLayerIndex, WedgeIDType WID) const
 	{
-		ensure(false);
-		return FVector2f(0.0f, 0.0f); 
+		checkf(VertexDataType == EClothPatternVertexType::Render, TEXT("Requested Wedge UVs from a Sim mesh"));
+		const TConstArrayView<FVector2f> VertexUVs = Pattern.GetRenderUVs()[WID];
+		return VertexUVs[UVLayerIndex];
 	}
 
 	FVector3f GetWedgeNormal(WedgeIDType WID) const
@@ -203,14 +230,15 @@ public:
 		
 		checkfSlow(SkinWeightAttributeIndex == 0, TEXT("Cloth assets should only have one skin weight profile")); 
 
-		const TConstArrayView<int32> SimNumBoneInfluences = Pattern.GetSimNumBoneInfluences();
-		const TConstArrayView<TArray<int32>> SimBoneIndices = Pattern.GetSimBoneIndices();
-		const TConstArrayView<TArray<float>> BoneWeights = Pattern.GetSimBoneWeights();
+		const bool bGetRenderMeshData = (VertexDataType == EClothPatternVertexType::Render);
+		const TConstArrayView<int32> NumBoneInfluences = bGetRenderMeshData ? Pattern.GetRenderNumBoneInfluences() : Pattern.GetSimNumBoneInfluences();
+		const TConstArrayView<TArray<int32>> BoneIndices = bGetRenderMeshData ? Pattern.GetRenderBoneIndices() : Pattern.GetSimBoneIndices();
+		const TConstArrayView<TArray<float>> BoneWeights = bGetRenderMeshData ? Pattern.GetRenderBoneWeights() : Pattern.GetSimBoneWeights();
 
-		if (ensure(VertexID >= 0 && VertexID < SimNumBoneInfluences.Num()))
+		if (ensure(VertexID >= 0 && VertexID < NumBoneInfluences.Num()))
 		{
-			const int32 NumInfluences = SimNumBoneInfluences[VertexID];
-			const TArray<int32> Indices = SimBoneIndices[VertexID];
+			const int32 NumInfluences = NumBoneInfluences[VertexID];
+			const TArray<int32> Indices = BoneIndices[VertexID];
 			const TArray<float> Weights = BoneWeights[VertexID];
 			
 			TArray<FBoneWeight> BoneWeightArray;
@@ -276,36 +304,96 @@ public:
 		return FVector4f::One();
 	}
 
-	const TArray<int32>& GetNormalIDs() const { return VertIDs; }
-	FVector3f GetNormal(NormalIDType ID) const { return Pattern.GetSimRestNormal()[ID]; }
-	bool GetNormalTri(const TriIDType& TriID, NormalIDType& ID0, NormalIDType& ID1, NormalIDType& ID2) const { return GetTri(TriID, ID0, ID1, ID2); }
+	const TArray<int32>& GetNormalIDs() const 
+	{ 
+		return VertIDs; 
+	}
+
+	FVector3f GetNormal(NormalIDType ID) const 
+	{ 
+		return (VertexDataType == EClothPatternVertexType::Render) ? Pattern.GetRenderNormal()[ID] : Pattern.GetSimRestNormal()[ID];
+	}
+
+	bool GetNormalTri(const TriIDType& TriID, NormalIDType& ID0, NormalIDType& ID1, NormalIDType& ID2) const 
+	{ 
+		return GetTri(TriID, ID0, ID1, ID2); 
+	}
 	
-	const TArray<int32>& GetUVIDs(int32 LayerID) const { return EmptyArray; }
-	FVector2f GetUV(int32 LayerID, UVIDType UVID) const { check(0); return FVector2f(); }
-	bool GetUVTri(int32 LayerID, const TriIDType&, UVIDType& ID0, UVIDType& ID1, UVIDType& ID2) const { ID0 = ID1 = ID2 = UVIDType(-1); return false;}
+	const TArray<int32>& GetUVIDs(int32 LayerID) const 
+	{ 
+		return (VertexDataType == EClothPatternVertexType::Render) ? VertIDs : EmptyArray;
+	}
+
+	FVector2f GetUV(int32 LayerID, UVIDType UVID) const 
+	{ 
+		checkf(VertexDataType == EClothPatternVertexType::Render, TEXT("Requested UVs from a Sim mesh"));
+		const TConstArrayView<FVector2f> VertexUVs = Pattern.GetRenderUVs()[UVID];
+		return VertexUVs[LayerID];
+	}
+
+	bool GetUVTri(int32 LayerID, const TriIDType& TriID, UVIDType& ID0, UVIDType& ID1, UVIDType& ID2) const 
+	{ 
+		return GetTri(TriID, ID0, ID1, ID2); 
+	}
 	
-	const TArray<int32>& GetTangentIDs() const { return EmptyArray; }
-	FVector3f GetTangent(NormalIDType ID) const { check(0); return FVector3f(); }
-	bool GetTangentTri(const TriIDType&, NormalIDType& ID0, NormalIDType& ID1, NormalIDType& ID2) const { ID0 = ID1 = ID2 = NormalIDType(-1); return false; }
+	const TArray<int32>& GetTangentIDs() const 
+	{ 
+		return (VertexDataType == EClothPatternVertexType::Render) ? VertIDs : EmptyArray;
+	}
+
+	FVector3f GetTangent(NormalIDType ID) const 
+	{ 
+		checkf(VertexDataType == EClothPatternVertexType::Render, TEXT("Requested Tangent from a Sim mesh"));
+		return Pattern.GetRenderTangentU()[ID];
+	}
+
+	bool GetTangentTri(const TriIDType& TriID, NormalIDType& ID0, NormalIDType& ID1, NormalIDType& ID2) const 
+	{ 
+		return GetTri(TriID, ID0, ID1, ID2);
+	}
+
+	const TArray<int32>& GetBiTangentIDs() const 
+	{ 
+		return (VertexDataType == EClothPatternVertexType::Render) ? VertIDs : EmptyArray;
+	}
+
+	FVector3f GetBiTangent(NormalIDType ID) const 
+	{
+		checkf(VertexDataType == EClothPatternVertexType::Render, TEXT("Requested Bitangent from a Sim mesh"));
+		return Pattern.GetRenderTangentV()[ID];
+	}
+
+	bool GetBiTangentTri(const TriIDType& TriID, NormalIDType& ID0, NormalIDType& ID1, NormalIDType& ID2) const
+	{
+		return GetTri(TriID, ID0, ID1, ID2);
+	}
+
+	const TArray<int32>& GetColorIDs() const 
+	{ 
+		return (VertexDataType == EClothPatternVertexType::Render) ? VertIDs : EmptyArray;
+	}
 
-	const TArray<int32>& GetBiTangentIDs() const { return EmptyArray; }
-	FVector3f GetBiTangent(NormalIDType ID) const {check(0); return FVector3f(); }
-	bool GetBiTangentTri(const TriIDType&, NormalIDType& ID0, NormalIDType& ID1, NormalIDType& ID2) const { ID0 = ID1 = ID2 = NormalIDType(-1); return false; }
+	FVector4f GetColor(ColorIDType VID) const 
+	{
+		checkf(VertexDataType == EClothPatternVertexType::Render, TEXT("Requested color from a Sim mesh"));
+		return Pattern.GetRenderColor()[VID];
+	}
 
-	const TArray<int32>& GetColorIDs() const { return EmptyArray; }
-	FVector4f GetColor(ColorIDType ID) const { check(0); return FVector4f(); }
-	bool GetColorTri(const TriIDType&, ColorIDType& ID0, ColorIDType& ID1, ColorIDType& ID2) const { ID0 = ID1 = ID2 = ColorIDType(-1); return false; }
+	bool GetColorTri(const TriIDType& TriID, ColorIDType& ID0, ColorIDType& ID1, ColorIDType& ID2) const 
+	{
+		return GetTri(TriID, ID0, ID1, ID2);
+	}
 
 private:
 
-	bool bGet2DPatternVertices = false;
+	const EClothPatternVertexType VertexDataType;
 
 	TArray<TriIDType> TriIDs;		// indices into Pattern.GetSimFaces()
 	TArray<VertIDType> VertIDs;		// indices into Pattern.GetSimVertices()
 	
 	TArray<FName> WeightMapNames;
 
-	const UE::Chaos::ClothAsset::FCollectionClothPatternConstFacade Pattern;
+	const FCollectionClothPatternConstFacade Pattern;
 
 	TArray<int32> EmptyArray;
 
@@ -313,18 +401,16 @@ private:
 };
 
 
-void FClothPatternToDynamicMesh::Convert(const TSharedPtr<const FManagedArrayCollection> ClothCollection, int32 LODIndex, int32 PatternIndex, bool bGet2DPattern, UE::Geometry::FDynamicMesh3& MeshOut)
+void FClothPatternToDynamicMesh::Convert(const TSharedPtr<const FManagedArrayCollection> ClothCollection, int32 LODIndex, int32 PatternIndex, EClothPatternVertexType VertexDataType, UE::Geometry::FDynamicMesh3& MeshOut)
 {
-	using namespace UE::Chaos::ClothAsset;
-
 	const FCollectionClothConstFacade ClothFacade(ClothCollection);
 
 	// Actual conversion
-	UE::Geometry::TToDynamicMesh<FClothSimPatternWrapper> PatternToDynamicMesh; 
-	FClothSimPatternWrapper PatternWrapper(ClothFacade, LODIndex, PatternIndex, bGet2DPattern);
+	UE::Geometry::TToDynamicMesh<FClothPatternWrapper> PatternToDynamicMesh; 
+	FClothPatternWrapper PatternWrapper(ClothFacade, LODIndex, PatternIndex, VertexDataType);
 
 	const bool bDisableAttributes = false;
-	auto TriangleToGroupFunction = [](FClothSimPatternWrapper::TriIDType) { return 0; };
+	auto TriangleToGroupFunction = [](FClothPatternWrapper::TriIDType) { return 0; };
 
 	if (bDisableAttributes)
 	{
@@ -335,26 +421,46 @@ void FClothPatternToDynamicMesh::Convert(const TSharedPtr<const FManagedArrayCol
 	{
 		MeshOut.EnableAttributes();
 
-		auto TriangleToMaterialFunction = [](FClothSimPatternWrapper::TriIDType) { return 0; };
+		const FCollectionClothLodConstFacade LodFacade(ClothFacade.GetLod(LODIndex));
+		const TArrayView<const int32> RenderMaterialIndex = LodFacade.GetRenderMaterialIndex();
+
+		auto TriangleToMaterialFunction = [&RenderMaterialIndex](FClothPatternWrapper::TriIDType TriID)
+		{
+			if (ensure(RenderMaterialIndex.IsValidIndex(TriID)))
+			{
+				return RenderMaterialIndex[TriID];
+			}
+			return 0;
+		};
+
 		constexpr bool bCopyTangents = false;
 		PatternToDynamicMesh.Convert(MeshOut, PatternWrapper, TriangleToGroupFunction, TriangleToMaterialFunction, bCopyTangents);
 	}
+
 }
 
-void FClothPatternToDynamicMesh::Convert(const UChaosClothAsset* ClothAssetMeshIn, int32 LODIndex, int32 PatternIndex, bool bGet2DPattern, UE::Geometry::FDynamicMesh3& MeshOut)
+void FClothPatternToDynamicMesh::Convert(const UChaosClothAsset* ClothAssetMeshIn, int32 LODIndex, int32 PatternIndex, EClothPatternVertexType VertexDataType, UE::Geometry::FDynamicMesh3& MeshOut)
 {
 	const TSharedPtr<const FManagedArrayCollection> ClothCollection = ClothAssetMeshIn->GetClothCollection();
 	check(ClothCollection.IsValid());
 
-	Convert(ClothCollection, LODIndex, PatternIndex, bGet2DPattern, MeshOut);
+	Convert(ClothCollection, LODIndex, PatternIndex, VertexDataType, MeshOut);
 }
 
+}	// namespace UE::Chaos::ClothAsset
+
 #else
 
+namespace UE::Chaos::ClothAsset
+{
+
 void FClothAssetToDynamicMesh::Convert(const UChaosClothAsset* ClothAssetMeshIn, FDynamicMesh3& MeshOut, int32 LODIndex, int32 PatternIndex)
 {
 	// Conversion only supported with editor.
 	check(0);
 }
 
+}	// namespace UE::Chaos::ClothAsset
+
 #endif  // end with editor
+
diff --git a/Engine/Plugins/Experimental/ChaosClothAssetEditor/Source/ChaosClothAssetTools/Public/ChaosClothAsset/ClothPatternToDynamicMesh.h b/Engine/Plugins/Experimental/ChaosClothAssetEditor/Source/ChaosClothAssetTools/Public/ChaosClothAsset/ClothPatternToDynamicMesh.h
--- a/Engine/Plugins/Experimental/ChaosClothAssetEditor/Source/ChaosClothAssetTools/Public/ChaosClothAsset/ClothPatternToDynamicMesh.h
+++ b/Engine/Plugins/Experimental/ChaosClothAssetEditor/Source/ChaosClothAssetTools/Public/ChaosClothAsset/ClothPatternToDynamicMesh.h
@@ -2,22 +2,28 @@
 
 #pragma once
 
-
-#include "CoreTypes.h"
+#include "ClothPatternVertexType.h"
 #include "Templates/SharedPointer.h"
 
+
 namespace UE::Geometry { class FDynamicMesh3; }
 class UChaosClothAsset;
 struct FManagedArrayCollection;
 
+namespace UE::Chaos::ClothAsset
+{
+
 /**
- * Convert a single pattern from a ClothAsset to a FDynamicMesh3
- */
+* Convert a single pattern from a ClothAsset to a FDynamicMesh3
+*/
 class CHAOSCLOTHASSETTOOLS_API FClothPatternToDynamicMesh
 {
 public:
 
-	void Convert(const TSharedPtr<const FManagedArrayCollection> ClothCollection, int32 LODIndex, int32 PatternIndex, bool bGet2DPattern, UE::Geometry::FDynamicMesh3& MeshOut);
+	void Convert(const TSharedPtr<const FManagedArrayCollection> ClothCollection, int32 LODIndex, int32 PatternIndex, EClothPatternVertexType VertexDataType, UE::Geometry::FDynamicMesh3& MeshOut);
 
-	void Convert(const UChaosClothAsset* ClothAssetMeshIn, int32 LODIndex, int32 PatternIndex, bool bGet2DPattern, UE::Geometry::FDynamicMesh3& MeshOut);
+	void Convert(const UChaosClothAsset* ClothAssetMeshIn, int32 LODIndex, int32 PatternIndex, EClothPatternVertexType VertexDataType, UE::Geometry::FDynamicMesh3& MeshOut);
 };
+
+}	// namespace UE::Chaos::ClothAsset
+
diff --git a/Engine/Plugins/Experimental/ChaosClothAssetEditor/Source/ChaosClothAssetTools/Public/ChaosClothAsset/ClothPatternVertexType.h b/Engine/Plugins/Experimental/ChaosClothAssetEditor/Source/ChaosClothAssetTools/Public/ChaosClothAsset/ClothPatternVertexType.h
new file mode 100644
--- /dev/null
+++ b/Engine/Plugins/Experimental/ChaosClothAssetEditor/Source/ChaosClothAssetTools/Public/ChaosClothAsset/ClothPatternVertexType.h
@@ -0,0 +1,16 @@
+// Copyright Epic Games, Inc. All Rights Reserved.
+
+#pragma once
+
+#include "CoreTypes.h"
+
+namespace UE::Chaos::ClothAsset
+{
+	enum class EClothPatternVertexType : uint8
+	{
+		Sim2D = 0,
+		Sim3D = 1,
+		Render = 2
+	};
+}
+
