Chaos Cloth - Fixed LOD merge in terminal node when some of the LODs are invalid.

#rnx

diff --git a/Engine/Plugins/Experimental/ChaosClothAssetEditor/Source/ChaosClothAssetDataflowNodes/Private/ChaosClothAsset/TerminalNode.cpp b/Engine/Plugins/Experimental/ChaosClothAssetEditor/Source/ChaosClothAssetDataflowNodes/Private/ChaosClothAsset/TerminalNode.cpp
--- a/Engine/Plugins/Experimental/ChaosClothAssetEditor/Source/ChaosClothAssetDataflowNodes/Private/ChaosClothAsset/TerminalNode.cpp
+++ b/Engine/Plugins/Experimental/ChaosClothAssetEditor/Source/ChaosClothAssetDataflowNodes/Private/ChaosClothAsset/TerminalNode.cpp
@@ -41,11 +41,13 @@ void FChaosClothAssetTerminalNode::SetAssetValue(TObjectPtr<UObject> Asset, Data
 		FString PhysicsAssetPathName;
 
 		const TArray<const FManagedArrayCollection*> CollectionLods = GetCollectionLods();
+		int32 LastValidLodIndex = INDEX_NONE;
 		for (int32 LodIndex = 0; LodIndex < CollectionLods.Num(); ++LodIndex)
 		{
 			// New LOD
 			const int32 NewLodIndex = ClothFacade.AddLod();
 			check(NewLodIndex == LodIndex);
+			FCollectionClothLodFacade ClothLodFacade = ClothFacade.GetLod(LodIndex);
 
 			// Retrieve input LOD
 			const FManagedArrayCollection& InCollectionLod = GetValue<FManagedArrayCollection>(Context, CollectionLods[LodIndex]);
@@ -57,55 +59,91 @@ void FChaosClothAssetTerminalNode::SetAssetValue(TObjectPtr<UObject> Asset, Data
 			{
 				const FCollectionClothLodConstFacade InClothLodFacade = InClothFacade.GetLod(0);
 
-				// Copy input LOD 0 to current output LOD
-				FCollectionClothLodFacade ClothLodFacade = ClothFacade.GetLod(LodIndex);
-				ClothLodFacade.Initialize(InClothLodFacade);
+				// Check LOD validity
+				if (InClothLodFacade.GetNumPatterns() &&
+					InClothLodFacade.GetNumSimVertices() &&
+					InClothLodFacade.GetNumSimFaces() &&
+					InClothLodFacade.GetNumRenderVertices() &&
+					InClothLodFacade.GetNumRenderFaces())
+				{
+					LastValidLodIndex = LodIndex;
 
-				// Add this LOD's materials to the asset
-				const int32 NumLodMaterials = InClothLodFacade.GetNumMaterials();
+					// Copy input LOD 0 to current output LOD
+					ClothLodFacade.Initialize(InClothLodFacade);
 
-				TArray<FSkeletalMaterial>& Materials = ClothAsset->GetMaterials();
-				Materials.Reserve(Materials.Num() + NumLodMaterials);
+					// Add this LOD's materials to the asset
+					const int32 NumLodMaterials = InClothLodFacade.GetNumMaterials();
 
-				const TConstArrayView<FString> LodRenderMaterialPathName = ClothLodFacade.GetRenderMaterialPathName();
-				for (int32 LodMaterialIndex = 0; LodMaterialIndex < NumLodMaterials; ++LodMaterialIndex)
-				{
-					const FString& RenderMaterialPathName = LodRenderMaterialPathName[LodMaterialIndex];
+					TArray<FSkeletalMaterial>& Materials = ClothAsset->GetMaterials();
+					Materials.Reserve(Materials.Num() + NumLodMaterials);
 
-					auto Predicate = [&RenderMaterialPathName](const FSkeletalMaterial& SkeletalMaterial)
+					const TConstArrayView<FString> LodRenderMaterialPathName = ClothLodFacade.GetRenderMaterialPathName();
+					for (int32 LodMaterialIndex = 0; LodMaterialIndex < NumLodMaterials; ++LodMaterialIndex)
+					{
+						const FString& RenderMaterialPathName = LodRenderMaterialPathName[LodMaterialIndex];
+
+						auto Predicate = [&RenderMaterialPathName](const FSkeletalMaterial& SkeletalMaterial)
 						{
 							return SkeletalMaterial.MaterialInterface && SkeletalMaterial.MaterialInterface->GetPathName() == RenderMaterialPathName;
 						};
 
-					if (!Materials.FindByPredicate(Predicate))
+						if (!Materials.FindByPredicate(Predicate))
+						{
+							if (UMaterial* const Material = LoadObject<UMaterial>(ClothAsset, *RenderMaterialPathName, nullptr, LOAD_None, nullptr))
+							{
+								Materials.Emplace(Material, true, false, Material->GetFName());
+							}
+						}
+					}
+
+					const TArray<FName> MapNames = InClothFacade.GetWeightMapNames();
+					for (const FName& MapName : MapNames)
 					{
-						if (UMaterial* const Material = LoadObject<UMaterial>(ClothAsset, *RenderMaterialPathName, nullptr, LOAD_None, nullptr))
+						ClothFacade.AddWeightMap(MapName);
+						TArrayView<float> WeightMap = ClothLodFacade.GetWeightMap(MapName);
+						TConstArrayView<float> InWeightMap = InClothLodFacade.GetWeightMap(MapName);
+						for (int32 i = 0; i < WeightMap.Num(); ++i)
 						{
-							Materials.Emplace(Material, true, false, Material->GetFName());
+							WeightMap[i] = InWeightMap[i];
 						}
 					}
-				}
 
-				const TArray<FName> MapNames = InClothFacade.GetWeightMapNames();
-				for (const FName& MapName : MapNames)
-				{
-					ClothFacade.AddWeightMap(MapName);
-					TArrayView<float> WeightMap = ClothLodFacade.GetWeightMap(MapName);
-					TConstArrayView<float> InWeightMap = InClothLodFacade.GetWeightMap(MapName);
-					for (int32 i = 0; i < WeightMap.Num(); ++i)
+					// Set properties and physics asset only with LOD 0 at the moment
+					if (LodIndex == 0)
 					{
-						WeightMap[i] = InWeightMap[i];
+						using namespace ::Chaos::Softs;
+
+						PhysicsAssetPathName = InClothLodFacade.GetPhysicsAssetPathName();
+
+						FCollectionPropertyMutableFacade(ClothCollection).Append(*InClothCollection);
 					}
 				}
-
-				// Set properties and physics asset only with LOD 0 at the moment
-				if (LodIndex == 0)
+			}
+			if (LodIndex != LastValidLodIndex)
+			{
+				if (LastValidLodIndex >= 0)
 				{
-					using namespace ::Chaos::Softs;
-
-					PhysicsAssetPathName = InClothLodFacade.GetPhysicsAssetPathName();
-				
-					FCollectionPropertyMutableFacade(ClothCollection).Append(*InClothCollection);
+					ClothLodFacade.Initialize(ClothFacade.GetLod(LastValidLodIndex));
+
+					DataflowNodes::LogAndToastWarning(
+						FText::Format(
+							LOCTEXT(
+								"InvalidLodN",
+								"ClothAssetTerminal: Invalid or empty input LOD for LOD {0}.\n"
+								"Using previous valid LOD {1} instead."),
+						LodIndex,
+						LastValidLodIndex));
+				}
+				else
+				{
+					ClothFacade.SetNumLods(0);
+
+					DataflowNodes::LogAndToastWarning(
+						LOCTEXT(
+							"InvalidLod0",
+							"ClothAssetTerminal: Invalid or empty input LOD for LOD 0.\n"
+							"LOD 0 cannot be empty in order to construct a valid Cloth Asset."));
+					break; // Empty cloth asset
 				}
 			}
 		}
