Cloth Paint Tool: show the current weight value under the brush. Only active when in Brush mode (e.g. not Gradient or PolyLasso)

diff --git a/Engine/Plugins/Experimental/ChaosClothAssetEditor/Source/ChaosClothAssetEditorTools/Private/ChaosClothAsset/ClothWeightMapPaintTool.cpp b/Engine/Plugins/Experimental/ChaosClothAssetEditor/Source/ChaosClothAssetEditorTools/Private/ChaosClothAsset/ClothWeightMapPaintTool.cpp
--- a/Engine/Plugins/Experimental/ChaosClothAssetEditor/Source/ChaosClothAssetEditorTools/Private/ChaosClothAsset/ClothWeightMapPaintTool.cpp
+++ b/Engine/Plugins/Experimental/ChaosClothAssetEditor/Source/ChaosClothAssetEditorTools/Private/ChaosClothAsset/ClothWeightMapPaintTool.cpp
@@ -42,6 +42,7 @@
 #include "Widgets/Notifications/SNotificationList.h"
 #include "Framework/Notifications/NotificationManager.h"
 #include "Selection/PolygonSelectionMechanic.h"
+#include "BaseGizmos/BrushStampIndicator.h"
 
 #include UE_INLINE_GENERATED_CPP_BY_NAME(ClothWeightMapPaintTool)
 
@@ -1317,7 +1318,7 @@ void UClothEditorWeightMapPaintTool::OnTick(float DeltaTime)
 	check(!(bIsLasso && bIsGradient));
 
 	ConfigureIndicator(false);
-	SetIndicatorVisibility(bIsLasso == false);
+	SetIndicatorVisibility(!bIsLasso && !bIsGradient);
 
 	if (bHavePendingAction)
 	{
@@ -1343,18 +1344,25 @@ void UClothEditorWeightMapPaintTool::OnTick(float DeltaTime)
 		return;
 	}
 
-	if (bPendingPickWeight)
+	// Get value at brush location
+	const bool bShouldPickWeight = bPendingPickWeight && IsStampPending() == false;
+	const bool bShouldUpdateValueAtBrush = IsInBrushSubMode();
+
+	if (bShouldPickWeight || bShouldUpdateValueAtBrush)
 	{
-		if (GetBrushTriangleID() >= 0 && IsStampPending() == false )
+		if (GetSculptMesh()->IsTriangle(GetBrushTriangleID()))
 		{
-			if (GetSculptMesh()->IsTriangle(GetBrushTriangleID()))
+			const double HitWeightValue = GetCurrentWeightValueUnderBrush();
+
+			if (bShouldPickWeight)
 			{
-				double HitWeightValue = GetCurrentWeightValueUnderBrush();
-				if (bPendingPickWeight)
-				{
-					FilterProperties->AttributeValue = HitWeightValue;
-					NotifyOfPropertyChangeByTool(FilterProperties);
-				}
+				FilterProperties->AttributeValue = HitWeightValue;
+				NotifyOfPropertyChangeByTool(FilterProperties);
+			}
+
+			if (bShouldUpdateValueAtBrush)
+			{
+				FilterProperties->ValueAtBrush = HitWeightValue;
 			}
 		}
 		bPendingPickWeight = false;
diff --git a/Engine/Plugins/Experimental/ChaosClothAssetEditor/Source/ChaosClothAssetEditorTools/Public/ChaosClothAsset/ClothWeightMapPaintTool.h b/Engine/Plugins/Experimental/ChaosClothAssetEditor/Source/ChaosClothAssetEditorTools/Public/ChaosClothAsset/ClothWeightMapPaintTool.h
--- a/Engine/Plugins/Experimental/ChaosClothAssetEditor/Source/ChaosClothAssetEditorTools/Public/ChaosClothAsset/ClothWeightMapPaintTool.h
+++ b/Engine/Plugins/Experimental/ChaosClothAssetEditor/Source/ChaosClothAssetEditorTools/Public/ChaosClothAsset/ClothWeightMapPaintTool.h
@@ -118,6 +118,7 @@ enum class EClothEditorWeightMapPaintVisibilityType : uint8
 
 
 
+// TODO: Look at EditConditions for all these properties. Which ones make sense for which SubToolType?
 
 UCLASS()
 class CHAOSCLOTHASSETEDITORTOOLS_API UClothEditorWeightMapPaintBrushFilterProperties : public UInteractiveToolPropertySet
@@ -125,8 +126,6 @@ class CHAOSCLOTHASSETEDITORTOOLS_API UClothEditorWeightMapPaintBrushFilterProper
 	GENERATED_BODY()
 
 public:
-	/** Primary Brush Mode */
-	//UPROPERTY(EditAnywhere, Category = Brush2, meta = (DisplayName = "Brush Type"))
 	
 	UPROPERTY(EditAnywhere, Category = ActionType, meta = (DisplayName = "Action"))
 	EClothEditorWeightMapPaintInteractionType SubToolType = EClothEditorWeightMapPaintInteractionType::Brush;
@@ -140,9 +139,11 @@ public:
 		HideEditConditionToggle, EditConditionHides, EditCondition = "SubToolType != EClothEditorWeightMapPaintInteractionType::PolyLasso"))
 	float BrushSize = 0.25f;
 
+	/** The new value to paint on the mesh */
 	UPROPERTY(EditAnywhere, Category = ActionType, meta = (UIMin = 0, ClampMin = 0, UIMax = 1, ClampMax = 1))
 	double AttributeValue = 1;
 
+	/** How quickly each brush stroke will drive mesh values towards the desired value */
 	UPROPERTY(EditAnywhere, Category = ActionType, meta = (UIMin = 0, ClampMin = 0, UIMax = 1, ClampMax = 1))
 	double Strength = 1;
 
@@ -162,6 +163,12 @@ public:
 	/** Control which triangles can be affected by the current operation based on visibility. Applied after all other filters. */
 	UPROPERTY(EditAnywhere, Category = Filters)
 	EClothEditorWeightMapPaintVisibilityType VisibilityFilter = EClothEditorWeightMapPaintVisibilityType::None;
+
+	/** The weight value at the brush indicator */
+	UPROPERTY(VisibleAnywhere, Transient, Category = Query, meta = (NoResetToDefault, 
+		HideEditConditionToggle, EditConditionHides, EditCondition = "SubToolType == EClothEditorWeightMapPaintInteractionType::Brush || SubToolType == EClothEditorWeightMapPaintInteractionType::Fill"))
+	double ValueAtBrush = 0;
+
 };
 
 
