Chaos Cloth - Enabled support for material instances by replacing UMaterial by UMaterialInterface where needed.

#rnx

diff --git a/Engine/Plugins/Experimental/ChaosClothAssetEditor/Source/ChaosClothAssetDataflowNodes/Private/ChaosClothAsset/TerminalNode.cpp b/Engine/Plugins/Experimental/ChaosClothAssetEditor/Source/ChaosClothAssetDataflowNodes/Private/ChaosClothAsset/TerminalNode.cpp
--- a/Engine/Plugins/Experimental/ChaosClothAssetEditor/Source/ChaosClothAssetDataflowNodes/Private/ChaosClothAsset/TerminalNode.cpp
+++ b/Engine/Plugins/Experimental/ChaosClothAssetEditor/Source/ChaosClothAssetDataflowNodes/Private/ChaosClothAsset/TerminalNode.cpp
@@ -89,7 +89,7 @@ void FChaosClothAssetTerminalNode::SetAssetValue(TObjectPtr<UObject> Asset, Data
 
 						if (!Materials.FindByPredicate(Predicate))
 						{
-							if (UMaterial* const Material = LoadObject<UMaterial>(ClothAsset, *RenderMaterialPathName, nullptr, LOAD_None, nullptr))
+							if (UMaterialInterface* const Material = LoadObject<UMaterialInterface>(ClothAsset, *RenderMaterialPathName, nullptr, LOAD_None, nullptr))
 							{
 								Materials.Emplace(Material, true, false, Material->GetFName());
 							}
diff --git a/Engine/Plugins/Experimental/ChaosClothAssetEditor/Source/ChaosClothAssetEditor/Private/ChaosClothAsset/ClothEditorMode.cpp b/Engine/Plugins/Experimental/ChaosClothAssetEditor/Source/ChaosClothAssetEditor/Private/ChaosClothAsset/ClothEditorMode.cpp
--- a/Engine/Plugins/Experimental/ChaosClothAssetEditor/Source/ChaosClothAssetEditor/Private/ChaosClothAsset/ClothEditorMode.cpp
+++ b/Engine/Plugins/Experimental/ChaosClothAssetEditor/Source/ChaosClothAssetEditor/Private/ChaosClothAsset/ClothEditorMode.cpp
@@ -414,7 +414,7 @@ void UChaosClothAssetEditorMode::ReinitializeDynamicMeshComponents()
 				for (int32 MaterialIndex = 0; MaterialIndex < MaterialPaths.Num(); ++MaterialIndex)
 				{
 					const FString& Path = MaterialPaths[MaterialIndex];
-					UMaterialInterface* const Material = LoadObject<UMaterial>(nullptr, *Path);
+					UMaterialInterface* const Material = LoadObject<UMaterialInterface>(nullptr, *Path);
 					MeshComponent.SetMaterial(MaterialIndex, Material);
 				}
 			}
