ClothTrainingTool: remove chaos cache, save to geometry cache

diff --git a/Engine/Plugins/Experimental/ChaosClothAsset/ChaosClothAsset.uplugin b/Engine/Plugins/Experimental/ChaosClothAsset/ChaosClothAsset.uplugin
--- a/Engine/Plugins/Experimental/ChaosClothAsset/ChaosClothAsset.uplugin
+++ b/Engine/Plugins/Experimental/ChaosClothAsset/ChaosClothAsset.uplugin
@@ -30,6 +30,10 @@
 		{
 			"Name": "ChaosCloth",
 			"Enabled": true
+		},
+		{
+			"Name": "GeometryCache",
+			"Enabled": true
 		}
 	]
 }
diff --git a/Engine/Plugins/Experimental/ChaosClothAssetEditor/Source/ChaosClothAssetEditorTools/ChaosClothAssetEditorTools.Build.cs b/Engine/Plugins/Experimental/ChaosClothAssetEditor/Source/ChaosClothAssetEditorTools/ChaosClothAssetEditorTools.Build.cs
--- a/Engine/Plugins/Experimental/ChaosClothAssetEditor/Source/ChaosClothAssetEditorTools/ChaosClothAssetEditorTools.Build.cs
+++ b/Engine/Plugins/Experimental/ChaosClothAssetEditor/Source/ChaosClothAssetEditorTools/ChaosClothAssetEditorTools.Build.cs
@@ -37,6 +37,7 @@ public class ChaosClothAssetEditorTools : ModuleRules
 				"MeshModelingToolsExp",
 				"EditorStyle",
 				"EditorSubsystem",
+				"GeometryCache",
 				"GeometryCore",
 				"GeometryFramework",
 				"LevelEditor",
@@ -55,7 +56,6 @@ public class ChaosClothAssetEditorTools : ModuleRules
 				"MeshDescription",
 				"SkeletalMeshDescription",
 				"ChaosClothAsset",
-				"ChaosCaching",
 				"ChaosClothAssetDataflowNodes",
 				"DataflowCore",
 				"DataflowEngine",
diff --git a/Engine/Plugins/Experimental/ChaosClothAssetEditor/Source/ChaosClothAssetEditorTools/Private/ChaosClothAsset/ClothTrainingTool.cpp b/Engine/Plugins/Experimental/ChaosClothAssetEditor/Source/ChaosClothAssetEditorTools/Private/ChaosClothAsset/ClothTrainingTool.cpp
--- a/Engine/Plugins/Experimental/ChaosClothAssetEditor/Source/ChaosClothAssetEditorTools/Private/ChaosClothAsset/ClothTrainingTool.cpp
+++ b/Engine/Plugins/Experimental/ChaosClothAssetEditor/Source/ChaosClothAssetEditorTools/Private/ChaosClothAsset/ClothTrainingTool.cpp
@@ -4,9 +4,9 @@
 #include "CoreMinimal.h"
 #include "Animation/AnimSequence.h"
 #include "Animation/AttributesRuntime.h"
+#include "AssetRegistry/AssetRegistryModule.h"
 #include "BonePose.h"
-#include "Chaos/ChaosCache.h"
-#include "Chaos/CacheCollection.h"
+#include "Chaos/TriangleMesh.h"
 #include "ChaosCloth/ChaosClothingSimulationConfig.h"
 #include "ChaosClothAsset/ClothAsset.h"
 #include "ChaosClothAsset/ClothDataGenerationComponent.h"
@@ -15,14 +15,22 @@
 #include "ChaosClothAsset/ClothCollection.h"
 #include "ChaosClothAsset/ClothSimulationProxy.h"
 #include "ClothingSystemRuntimeTypes.h"
+#include "ClothTrainingToolProperties.h"
 #include "ComponentReregisterContext.h"
 #include "ContextObjectStore.h"
+#include "Engine/SkinnedAssetCommon.h"
+#include "FileHelpers.h"
 #include "GenericPlatform/GenericPlatformProcess.h"
+#include "GeometryCache.h"
+#include "GeometryCacheCodecV1.h"
+#include "GeometryCacheMeshData.h"
+#include "GeometryCacheTrackStreamable.h"
 #include "InteractiveToolManager.h"
 #include "Internationalization/Regex.h"
 #include "ModelingOperators.h"
 #include "Misc/AsyncTaskNotification.h"
 #include "PhysicsEngine/PhysicsAsset.h"
+#include "Rendering/SkeletalMeshRenderData.h"
 #include "SkeletalRenderPublic.h"
 #include "Tasks/Pipe.h"
 #include "ToolTargetManager.h"
@@ -35,12 +43,6 @@ DEFINE_LOG_CATEGORY_STATIC(LogClothTrainingTool, Log, All);
 
 namespace UE::ClothTrainingTool::Private
 {
-	UChaosCache& GetCache(UChaosCacheCollection& CacheCollection)
-	{
-		static const FName CacheName = FName("SimulatedCache");
-		return *CacheCollection.FindOrAddCache(CacheName);
-	}
-
 	TArray<int32> ParseFrames(const FString& FramesString)
 	{
 		TArray<int32> Result;
@@ -104,6 +106,264 @@ namespace UE::ClothTrainingTool::Private
 		return Result;
 	}
 
+	FBox3f GetBoundingBox(const TArray<FVector3f>& Positions)
+	{
+		FBox3f BoundingBox;
+		for (const FVector3f& Position : Positions)
+		{
+			BoundingBox += Position;
+		}
+		return BoundingBox;
+	}
+
+	void CopyMaterials(const USkinnedAsset& Asset, UGeometryCache& Cache)
+	{
+		const TArray<FSkeletalMaterial>& Materials = Asset.GetMaterials();
+		TArray<TObjectPtr<UMaterialInterface>>& Interfaces = Cache.Materials;
+		Interfaces.SetNumZeroed(Materials.Num());
+		for (int32 MaterialIndex = 0; MaterialIndex < Materials.Num(); ++MaterialIndex)
+		{
+			Interfaces[MaterialIndex] = Materials[MaterialIndex].MaterialInterface;
+		}
+	}
+
+	TArray<uint32> Range(uint32 Start, uint32 End)
+	{
+		TArray<uint32> Result;
+		const uint32 Num = End - Start;
+		Result.SetNumUninitialized(Num);
+		for (uint32 Index = 0; Index < Num; ++Index)
+		{
+			Result.Add(Index + Start);
+		}
+		return Result;
+	}
+
+	int32 GetNumVertices(const FSkeletalMeshLODRenderData& LODData)
+	{
+		int32 NumVertices = 0;
+		for(const FSkelMeshRenderSection& Section : LODData.RenderSections)
+		{
+			NumVertices += Section.NumVertices;
+		}
+		return NumVertices;
+	}
+
+	TArray<FVector2f> GetUV0s(const FSkeletalMeshLODRenderData& LODData)
+	{
+		TArray<FVector2f> UV0s;
+		const FStaticMeshVertexBuffer& StaticMeshVertexBuffer = LODData.StaticVertexBuffers.StaticMeshVertexBuffer;
+		const int32 NumVertices = StaticMeshVertexBuffer.GetNumVertices();
+		UV0s.SetNumZeroed(NumVertices);
+		for (int32 Index = 0; Index < NumVertices; ++Index)
+		{
+			UV0s[Index] = StaticMeshVertexBuffer.GetVertexUV(Index, 0);
+		}
+		return UV0s;
+	}
+
+	TArray<FColor> GetColors(const FSkeletalMeshLODRenderData& LODData, int32 NumVertices)
+	{
+		TArray<FColor> Colors;
+		Colors.SetNum(NumVertices);
+		const FColorVertexBuffer& ColorVertexBuffer = LODData.StaticVertexBuffers.ColorVertexBuffer;
+		if (ColorVertexBuffer.GetNumVertices() == NumVertices)
+		{
+			for (int32 Index = 0; Index < NumVertices; ++Index)
+			{
+				Colors[Index] = ColorVertexBuffer.VertexColor(Index);
+			}
+		}
+		else
+		{
+			for (int32 Index = 0; Index < NumVertices; ++Index)
+			{
+				Colors[Index] = FColor::White;
+			}
+		}
+		return Colors;
+	}
+
+	TArray<FPackedNormal> GetPackedNormals(TConstArrayView<FVector3f> Normals)
+	{
+		TArray<FPackedNormal> PackedNormals;
+		PackedNormals.Reserve(Normals.Num());
+		for (const FVector3f& Normal : Normals)
+		{
+			PackedNormals.Add(Normal);
+		}
+		return PackedNormals;
+	}
+
+	TArray<FVector3f> ComputeNormals(const ::Chaos::FTriangleMesh& TriangleMesh, TConstArrayView<FVector3f> Positions)
+	{
+		using TVec3 = ::Chaos::TVec3<float>;
+		TConstArrayView<TVec3> ChaosPositions(reinterpret_cast<const TVec3*>(Positions.GetData()), Positions.Num());
+		const TArray<TVec3> FaceNormals = TriangleMesh.GetFaceNormals(ChaosPositions, false);
+		TriangleMesh.GetPointToTriangleMap();
+		TArray<FVector3f> Normals;
+		Normals.SetNumUninitialized(Positions.Num());
+		TArrayView<TVec3> ChaosNormals(reinterpret_cast<TVec3*>(Normals.GetData()), Normals.Num());
+		TriangleMesh.GetPointNormals(ChaosNormals, TConstArrayView<TVec3>(FaceNormals), false);
+		return Normals;
+	}
+	
+	TArray<FVector3f> ComputeTangentsX(const ::Chaos::FTriangleMesh& TriangleMesh, TConstArrayView<FVector3f> Positions, TConstArrayView<FVector3f> Normals, TConstArrayView<FVector2f> UVs)
+	{
+		using ::Chaos::TVec3;
+		TArray<FVector3f> Tangents;
+		const int32 NumVertices = Positions.Num();
+		Tangents.SetNumZeroed(NumVertices);
+		TConstArrayView<TVec3<int32>> Elements{TriangleMesh.GetElements()};
+		for (const TVec3<int32>& Elment : Elements)
+		{
+			const FVector3f& P0 = Positions[Elment[0]];
+			const FVector3f& P1 = Positions[Elment[1]];
+			const FVector3f& P2 = Positions[Elment[2]];
+			const FVector2f& UV0 = UVs[Elment[0]];
+			const FVector2f& UV1 = UVs[Elment[1]];
+			const FVector2f& UV2 = UVs[Elment[2]];
+	
+			const FVector3f Edge1 = P1 - P0;
+			const FVector3f Edge2 = P2 - P0;
+			const FVector2f UVEdge1 = UV1 - UV0;
+			const FVector2f UVEdge2 = UV2 - UV0;
+			const float Det = UVEdge1.Y * UVEdge2.X - UVEdge1.X * UVEdge2.Y;
+			if (FMath::Abs(Det) > UE_SMALL_NUMBER)
+			{
+				FVector3f TangentX = (-Edge1 * UVEdge2.Y + Edge2 * UVEdge1.Y) / Det;
+				TangentX.Normalize();
+				TangentX *= Edge1.Cross(Edge2).Size();
+				for (int32 Index = 0; Index < 3; ++Index)
+				{
+					Tangents[Elment[Index]] += TangentX;
+				}
+			}
+		}
+		for (int32 Index = 0; Index < NumVertices; Index++)
+		{
+			FVector3f& TangentX = Tangents[Index];
+			const FVector3f& Normal = Normals[Index];
+			TangentX -= Normal * Normal.Dot(TangentX);
+			TangentX = TangentX.GetSafeNormal(UE_SMALL_NUMBER, FVector3f(1.0f, 0.0f, 0.0f));
+		}
+		return Tangents;
+	}
+
+	void SaveGeometryCache(UGeometryCache& GeometryCache, const USkinnedAsset& Asset, TArrayView<TArray<FVector3f>> PositionsToMoveFrom)
+	{
+		const FSkeletalMeshRenderData* RenderData = Asset.GetResourceForRendering();
+		constexpr int32 LODIndex = 0;
+		if (!RenderData || !RenderData->LODRenderData.IsValidIndex(LODIndex))
+		{
+			return;
+		}
+		const FSkeletalMeshLODRenderData& LODData = RenderData->LODRenderData[LODIndex];
+		TArray<uint32> Indices;
+		LODData.MultiSizeIndexContainer.GetIndexBuffer(Indices);
+
+		using ::Chaos::TVec3;
+		TArray<TVec3<int32>> Elements;
+		const int32 NumTriangles = Indices.Num() / 3;
+		Elements.SetNumUninitialized(NumTriangles);
+		for(int32 TriangleIndex = 0; TriangleIndex < NumTriangles; ++TriangleIndex)
+		{
+			const int32 Index = TriangleIndex * 3;
+			Elements[TriangleIndex] = TVec3<int32>(Indices[Index], Indices[Index + 1], Indices[Index + 2]);
+		}
+		::Chaos::FTriangleMesh TriangleMesh(MoveTemp(Elements));
+
+		const int32 NumFrames = PositionsToMoveFrom.Num();
+		if (NumFrames == 0)
+		{
+			return;
+		}
+		const int32 NumVertices = GetNumVertices(LODData);
+		if (!ensureAlways(NumVertices == PositionsToMoveFrom[0].Num()))
+		{
+			return;
+		}
+		const TArray<FVector2f> UV0 = GetUV0s(LODData);
+		const bool bHasUV0 = UV0.Num() == NumVertices;
+		const TArray<FColor> Colors = GetColors(LODData, NumVertices);
+		const TArray<uint32> ImportedVertexNumbers = Range(0, NumVertices);
+
+		GeometryCache.ClearForReimporting();
+
+		const FName BaseName = FName(*Asset.GetName());
+		const FName CodecName = MakeUniqueObjectName(&GeometryCache, UGeometryCacheCodecV1::StaticClass(), FName(BaseName.ToString() + FString(TEXT("_Codec"))));
+		UGeometryCacheCodecV1* Codec = NewObject<UGeometryCacheCodecV1>(&GeometryCache, CodecName, RF_Public);
+
+		constexpr float CompressedPositionPrecision = 1e-3f;
+		constexpr int32 CompressedTextureCoordinatesNumberOfBits = 4;
+		Codec->InitializeEncoder(CompressedPositionPrecision, CompressedTextureCoordinatesNumberOfBits);
+		
+		const FName TrackName = MakeUniqueObjectName(&GeometryCache, UGeometryCacheTrackStreamable::StaticClass(), BaseName);
+		UGeometryCacheTrackStreamable* Track = NewObject<UGeometryCacheTrackStreamable>(&GeometryCache, TrackName, RF_Public);
+		
+		constexpr bool bApplyConstantTopologyOptimizations = true;
+		constexpr bool bCalculateMotionVectors = false;
+		constexpr bool bOptimizeIndexBuffers = true;
+		Track->BeginCoding(Codec, bApplyConstantTopologyOptimizations, bCalculateMotionVectors, bOptimizeIndexBuffers);
+
+		constexpr float CacheFPS = 30;
+		int32 MaxRecordedFrame = -1;
+		for (int32 Frame = 0; Frame < NumFrames; ++Frame)
+		{
+			const float Time = Frame / CacheFPS;
+			if (PositionsToMoveFrom[Frame].Num() != NumVertices)
+			{
+				break;
+			}
+			FGeometryCacheMeshData MeshData;
+			MeshData.Positions = MoveTemp(PositionsToMoveFrom[Frame]);
+			// Ideally we should not need to save normals to save disk space. Save normals for now for visualization.
+			const TArray<FVector3f> Normals = ComputeNormals(TriangleMesh, MeshData.Positions);
+			MeshData.VertexInfo.bHasTangentZ = true;
+			MeshData.TangentsZ = GetPackedNormals(Normals);
+			if (bHasUV0)
+			{
+				MeshData.VertexInfo.bHasUV0 = true;
+				MeshData.TextureCoordinates = UV0;
+				MeshData.VertexInfo.bHasTangentX = true;
+				MeshData.TangentsX = GetPackedNormals(ComputeTangentsX(TriangleMesh, MeshData.Positions, Normals, UV0));
+			}
+			MeshData.VertexInfo.bHasColor0 = true;
+			MeshData.Colors = Colors;
+			MeshData.VertexInfo.bHasImportedVertexNumbers = true;
+			MeshData.ImportedVertexNumbers = ImportedVertexNumbers;
+			MeshData.Indices = Indices;
+			MeshData.BoundingBox = GetBoundingBox(MeshData.Positions);
+
+			MeshData.BatchesInfo.Reserve(LODData.RenderSections.Num());
+			for (const FSkelMeshRenderSection& Section : LODData.RenderSections)
+			{
+				FGeometryCacheMeshBatchInfo BatchInfo;
+				BatchInfo.StartIndex = Section.BaseIndex;
+				BatchInfo.NumTriangles = Section.NumTriangles;
+				BatchInfo.MaterialIndex = Section.MaterialIndex;
+				MeshData.BatchesInfo.Add(MoveTemp(BatchInfo));
+			}
+
+			constexpr bool bConstTopology = true;
+			Track->AddMeshSample(MeshData, Time, bConstTopology);
+			constexpr bool bVisible = true;
+			Track->AddVisibilitySample(bVisible, Time);
+			MaxRecordedFrame = Frame;
+		}
+
+		TArray<FMatrix> Mats { FMatrix::Identity, FMatrix::Identity };
+		TArray<float> MatTimes { 0.0f, (MaxRecordedFrame - 1) / CacheFPS };
+		Track->SetMatrixSamples(Mats, MatTimes);
+
+		if (ensureAlways(Track->EndCoding()))
+		{
+			GeometryCache.AddTrack(Track);
+			GeometryCache.SetFrameStartEnd(0, NumFrames - 1);
+		}
+		CopyMaterials(Asset, GeometryCache);
+	}
+
 	class FTimeScope
 	{
 	public:
@@ -121,6 +381,39 @@ namespace UE::ClothTrainingTool::Private
 		FString Name;
 		FDateTime StartTime;
 	};
+
+	template<class T>
+	T* CreateOrLoad(const FString& PackageName)
+	{
+		const FName AssetName(FPackageName::GetLongPackageAssetName(PackageName));		
+		if (UPackage* const Package = CreatePackage(*PackageName))
+		{
+			LoadPackage(nullptr, *PackageName, LOAD_Quiet | LOAD_EditorOnly);
+			T* Asset = FindObject<T>(Package, *AssetName.ToString());
+			if (!Asset)
+			{
+				Asset = NewObject<T>(Package, *AssetName.ToString(), RF_Public | RF_Standalone | RF_Transactional);
+				Asset->MarkPackageDirty();
+				FAssetRegistryModule::AssetCreated(Asset);
+			}
+			return Asset;
+		}
+		return nullptr;
+	}
+
+	void SavePackage(UObject& Object)
+	{
+		TArray<UPackage*> PackagesToSave = { Object.GetOutermost() };
+		constexpr bool bCheckDirty = false;
+		constexpr bool bPromptToSave = false;
+		FEditorFileUtils::PromptForCheckoutAndSave(PackagesToSave, bCheckDirty, bPromptToSave);
+	}
+};
+
+// @@@@@@@@@ TODO: Change this to whatever makes sense for output
+struct FSkinnedMeshVertices
+{
+	TArray<FVector3f> Vertices;
 };
 
 struct UClothTrainingTool::FSimResource
@@ -130,14 +423,16 @@ struct UClothTrainingTool::FSimResource
 	TUniquePtr<UE::Tasks::FPipe> Pipe;
 	FEvent* SkinEvent = nullptr;
 	std::atomic<bool> bNeedsSkin = false;
+	TArrayView<TArray<FVector3f>> SimulatedPositions;
 };
 
 class UClothTrainingTool::FLaunchSimsOp : public UE::Geometry::TGenericDataOperator<FSkinnedMeshVertices>
 {
 public: 
-	FLaunchSimsOp(TArray<FSimResource>& InSimResources, FCriticalSection& InSimMutex, TObjectPtr<UClothTrainingToolProperties> InToolProperties)
+	FLaunchSimsOp(TArray<FSimResource>& InSimResources, FCriticalSection& InSimMutex, TArray<TArray<FVector3f>>& InSimulatedPositions,  TObjectPtr<UClothTrainingToolProperties> InToolProperties)
 		: SimResources(InSimResources)
 		, SimMutex(InSimMutex)
+		, SimulatedPositions(InSimulatedPositions)
 		, ToolProperties(InToolProperties)
 	{
 	}
@@ -153,16 +448,15 @@ private:
 		EveryStep,
 	};
 
-	void Simulate(FSimResource &SimResource, int32 AnimFrame, int32 CacheFrame, UChaosCache& Cache, FProgressCancel* Progress, float ProgressStep);
+	void Simulate(FSimResource &SimResource, int32 AnimFrame, int32 CacheFrame, FProgressCancel* Progress, float ProgressStep);
 	void PrepareAnimationSequence();
 	void RestoreAnimationSequence();
 	TArray<FTransform> GetBoneTransforms(UChaosClothComponent& InClothComponent, int32 Frame) const;
-	bool GetSimPositions(FProxy& DataGenerationProxy, TArray<FVector3f>& OutPositions) const;
-	void GetRenderPositions(FSimResource& SimResource, TArray<FVector3f>& OutPositions) const;
-	void AddToCache(FSimResource& SimResource, UChaosCache& Cache, int32 Frame) const;
+	TArray<FVector3f> GetRenderPositions(FSimResource& SimResource) const;
 
 	TArray<FSimResource>& SimResources;
 	FCriticalSection& SimMutex;
+	TArray<TArray<FVector3f>>& SimulatedPositions;
 	TObjectPtr<UClothTrainingToolProperties> ToolProperties = nullptr;
 	EAnimInterpolationType InterpolationTypeBackup = EAnimInterpolationType::Linear;
 
@@ -178,11 +472,12 @@ struct UClothTrainingTool::FTaskResource
 
 	TUniquePtr<FExecuterType> Executer;
 	TUniquePtr<FAsyncTaskNotification> Notification;
-	UChaosCache *Cache = nullptr;
-	TUniquePtr<FCacheUserToken> CacheUserToken;
 	FDateTime StartTime;
 	FDateTime LastUpdateTime;
 
+	TArray<TArray<FVector3f>> SimulatedPositions;
+	UGeometryCache* Cache = nullptr;
+
 	bool AllocateSimResources_GameThread(const UChaosClothComponent& InClothComponent, int32 Num);
 	void FreeSimResources_GameThread();
 	void FlushRendering();
@@ -253,46 +548,10 @@ TArray<FTransform> UClothTrainingTool::FLaunchSimsOp::GetBoneTransforms(UChaosCl
 	return ComponentSpaceTransforms;
 }
 
-void UClothTrainingTool::FLaunchSimsOp::AddToCache(FSimResource& SimResource, UChaosCache& OutCache, int32 Frame) const
-{
-	TArray<FVector3f> Positions;
-	GetRenderPositions(SimResource, Positions);
-
-	constexpr float CacheFPS = 30;
-	const float Time = Frame / CacheFPS;
-	FPendingFrameWrite NewFrame;
-	NewFrame.Time = Time;
-
-	const int32 NumParticles = Positions.Num();
-	TArray<int32>& PendingID = NewFrame.PendingChannelsIndices;
-	TArray<float> PendingPX, PendingPY, PendingPZ;
-	TArray<float> PendingVX, PendingVY, PendingVZ;
-	PendingID.SetNum(NumParticles);
-	PendingPX.SetNum(NumParticles);
-	PendingPY.SetNum(NumParticles);
-	PendingPZ.SetNum(NumParticles);
-
-	for (int32 ParticleIndex = 0; ParticleIndex < NumParticles; ++ParticleIndex)
-	{
-		const FVector3f& Position = Positions[ParticleIndex];
-		PendingID[ParticleIndex] = ParticleIndex;
-		PendingPX[ParticleIndex] = Position.X;
-		PendingPY[ParticleIndex] = Position.Y;
-		PendingPZ[ParticleIndex] = Position.Z;
-	}
-
-	NewFrame.PendingChannelsData.Add(PositionXName, PendingPX);
-	NewFrame.PendingChannelsData.Add(PositionYName, PendingPY);
-	NewFrame.PendingChannelsData.Add(PositionZName, PendingPZ);
-
-	OutCache.AddFrame_Concurrent(MoveTemp(NewFrame));
-}
-
 void UClothTrainingTool::FLaunchSimsOp::CalculateResult(FProgressCancel* Progress)
 {
 	using UE::ClothTrainingTool::Private::ParseFrames;
 	using UE::ClothTrainingTool::Private::Range;
-	using UE::ClothTrainingTool::Private::GetCache;
 
 	const TArray<int32> FramesToSimulate = ToolProperties->FramesToSimulate.Len() > 0
 		? ParseFrames(ToolProperties->FramesToSimulate) 
@@ -303,22 +562,19 @@ void UClothTrainingTool::FLaunchSimsOp::CalculateResult(FProgressCancel* Progres
 	{
 		return;
 	}
+	SimulatedPositions.SetNum(ToolProperties->bDebug ? ToolProperties->NumSteps : NumFrames);
 	const float ProgressStep = 1.f / NumFrames;
-
-	UChaosCache& Cache = ToolProperties->bDebug ? GetCache(*ToolProperties->DebugCacheCollection) : GetCache(*ToolProperties->CacheCollection);
 	PrepareAnimationSequence();
 
 	const int32 NumThreads = ToolProperties->bDebug ? 1 : ToolProperties->NumThreads;
 	FScopeLock Lock(&SimMutex);
 
-	bool bCancelled = false;
 	for (int32 Frame = 0; Frame < NumFrames; Frame++)
 	{
 		if (Progress)
 		{
 			if (Progress->Cancelled())
 			{
-				bCancelled = true;
 				break;
 			}
 
@@ -326,10 +582,10 @@ void UClothTrainingTool::FLaunchSimsOp::CalculateResult(FProgressCancel* Progres
 			const int32 AnimFrame = FramesToSimulate[Frame];
 
 			FSimResource& SimResource = SimResources[ThreadIdx];
-			SimResource.Pipe->Launch(*FString::Printf(TEXT("SimFrame:%d"), AnimFrame), [this, &SimResource, AnimFrame, Frame, &Cache, Progress, ProgressStep]()
+			SimResource.Pipe->Launch(*FString::Printf(TEXT("SimFrame:%d"), AnimFrame), [this, &SimResource, AnimFrame, Frame, Progress, ProgressStep]()
 			{ 
 				FMemMark Mark(FMemStack::Get());
-				Simulate(SimResource, AnimFrame, Frame, Cache, Progress, ProgressStep);
+				Simulate(SimResource, AnimFrame, Frame, Progress, ProgressStep);
 			});
 		}
 	}
@@ -340,10 +596,9 @@ void UClothTrainingTool::FLaunchSimsOp::CalculateResult(FProgressCancel* Progres
 	}
 
 	RestoreAnimationSequence();
-	Cache.FlushPendingFrames();
 }
 
-void UClothTrainingTool::FLaunchSimsOp::Simulate(FSimResource &SimResource, int32 AnimFrame, int32 CacheFrame, UChaosCache& Cache, FProgressCancel* Progress, float ProgressStep)
+void UClothTrainingTool::FLaunchSimsOp::Simulate(FSimResource &SimResource, int32 AnimFrame, int32 CacheFrame, FProgressCancel* Progress, float ProgressStep)
 {
 	UClothDataGenerationComponent& TaskComponent = *SimResource.ClothComponent;
 	FProxy& DataGenerationProxy = *SimResource.Proxy;
@@ -380,47 +635,29 @@ void UClothTrainingTool::FLaunchSimsOp::Simulate(FSimResource &SimResource, int3
 			if (SaveType == ESaveType::EveryStep)
 			{
 				DataGenerationProxy.WriteSimulationData();
-				AddToCache(SimResource, Cache, Step);
+				SimulatedPositions[Step] = GetRenderPositions(SimResource);
 			}
 		}
+		else
+		{
+			bCancelled = true;
+			break;
+		}
 	}
 
 	if (SaveType == ESaveType::LastStep && !bCancelled)
 	{
 		DataGenerationProxy.WriteSimulationData();
-		AddToCache(SimResource, Cache, CacheFrame);
+		SimulatedPositions[CacheFrame] = GetRenderPositions(SimResource);
 	}
 
 	Progress->AdvanceCurrentScopeProgressBy(ProgressStep);
 }
 
-bool UClothTrainingTool::FLaunchSimsOp::GetSimPositions(FProxy& DataGenerationProxy, TArray<FVector3f> &OutPositions) const
-{
-	const TMap<int32, FClothSimulData> &SimulDataMap = DataGenerationProxy.GetCurrentSimulationData_AnyThread();
-	const FClothSimulData* const SimulData = SimulDataMap.Find(0);
-	if (SimulDataMap.Num() > 1)
-	{
-		ensureMsgf(false, TEXT("Multiple cloth is not yet supported."));
-		return false;
-	}
-	if (SimulData == nullptr)
-	{
-		ensureMsgf(false, TEXT("ClothSimulData is nullptr"));
-		return false;
-	}
-
-	const TArray<FVector3f>& SimPositions = SimulData->Positions;
-	OutPositions.SetNum(SimPositions.Num());
-	for (int32 Index = 0; Index < SimPositions.Num(); ++Index)
-	{
-		OutPositions[Index] = FVector3f(SimulData->ComponentRelativeTransform.TransformPosition(FVector(SimPositions[Index])));
-	}
-	return true;
-}
-
-void UClothTrainingTool::FLaunchSimsOp::GetRenderPositions(FSimResource& SimResource, TArray<FVector3f> &OutPositions) const
+TArray<FVector3f> UClothTrainingTool::FLaunchSimsOp::GetRenderPositions(FSimResource& SimResource) const
 {
 	check(SimResource.ClothComponent);
+	TArray<FVector3f> Positions;
 	TArray<FFinalSkinVertex> OutVertices;
 	// This could potentially be slow. 
 	SimResource.ClothComponent->RecreateRenderState_Concurrent();
@@ -428,11 +665,12 @@ void UClothTrainingTool::FLaunchSimsOp::GetRenderPositions(FSimResource& SimReso
 	SimResource.SkinEvent->Wait();
 	
 	SimResource.ClothComponent->GetCPUSkinnedCachedFinalVertices(OutVertices);
-	OutPositions.SetNum(OutVertices.Num());
+	Positions.SetNum(OutVertices.Num());
 	for (int32 Index = 0; Index < OutVertices.Num(); ++Index)
 	{
-		OutPositions[Index] = OutVertices[Index].Position;
+		Positions[Index] = OutVertices[Index].Position;
 	}
+	return Positions;
 }
 
 bool UClothTrainingTool::IsClothComponentValid() const
@@ -452,16 +690,6 @@ bool UClothTrainingTool::IsClothComponentValid() const
 	}
 }
 
-// ------------------- Properties -------------------
-
-void UClothTrainingToolActionProperties::PostAction(EClothTrainingToolActions Action)
-{
-	if (ParentTool.IsValid())
-	{
-		ParentTool->RequestAction(Action);
-	}
-}
-
 // ------------------- Builder -------------------
 
 const FToolTargetTypeRequirements& UClothTrainingToolBuilder::GetTargetRequirements() const
@@ -526,8 +754,8 @@ void UClothTrainingTool::StartTraining()
 		PendingAction = EClothTrainingToolActions::NoAction;
 		return;
 	}
-	UChaosCacheCollection* const CacheCollection = GetCacheCollection();
-	if (CacheCollection == nullptr)
+	UGeometryCache* const Cache = GetCache();
+	if (Cache == nullptr)
 	{
 		PendingAction = EClothTrainingToolActions::NoAction;
 		return;
@@ -538,18 +766,15 @@ void UClothTrainingTool::StartTraining()
 		return;
 	}
 	TaskResource = MakeUnique<FTaskResource>();
+
 	if (!TaskResource->AllocateSimResources_GameThread(*ClothComponent, ToolProperties->NumThreads))
 	{
 		PendingAction = EClothTrainingToolActions::NoAction;
 		return;
 	}
-
-	using UE::ClothTrainingTool::Private::GetCache;
-	UChaosCache* const Cache = &GetCache(*CacheCollection);
 	TaskResource->Cache = Cache;
-	TaskResource->CacheUserToken = MakeUnique<FCacheUserToken>(Cache->BeginRecord(ClothComponent, FGuid(), FTransform::Identity));
 
-	TUniquePtr<FLaunchSimsOp> NewOp = MakeUnique<FLaunchSimsOp>(TaskResource->SimResources, *TaskResource->SimMutex, ToolProperties);
+	TUniquePtr<FLaunchSimsOp> NewOp = MakeUnique<FLaunchSimsOp>(TaskResource->SimResources, *TaskResource->SimMutex, TaskResource->SimulatedPositions, ToolProperties);
 	TaskResource->Executer = MakeUnique<FExecuterType>(MoveTemp(NewOp));
 	TaskResource->Executer->StartBackgroundTask();
 
@@ -647,7 +872,7 @@ void UClothTrainingTool::Shutdown(EToolShutdownType ShutdownType)
 
 bool UClothTrainingTool::FTaskResource::AllocateSimResources_GameThread(const UChaosClothComponent& InClothComponent, int32 Num)
 {
-	SimResources.SetNumUninitialized(Num);
+	SimResources.SetNum(Num);
 	for(int32 Index = 0; Index < Num; ++Index)
 	{
 		UClothDataGenerationComponent* const CopyComponent = NewObject<UClothDataGenerationComponent>(InClothComponent.GetOuter());
@@ -673,6 +898,8 @@ bool UClothTrainingTool::FTaskResource::AllocateSimResources_GameThread(const UC
 		SimResource.SkinEvent = FPlatformProcess::GetSynchEventFromPool();
 		SimResource.bNeedsSkin.store(false);
 
+		SimResource.SimulatedPositions = TArrayView<TArray<FVector3f>>(SimulatedPositions);
+
 		if (CopyComponent == nullptr || SimResource.Proxy == nullptr || SimResource.Pipe == nullptr)
 		{
 			UE_LOG(LogClothTrainingTool, Error, TEXT("Failed to allocate simulation resources"));
@@ -730,51 +957,10 @@ void UClothTrainingTool::FTaskResource::FlushRendering()
 
 }
 
-UChaosCacheCollection* UClothTrainingTool::GetCacheCollection() const
-{
-	UChaosCacheCollection* CacheCollection = nullptr;
-	if (ToolProperties->bDebug)
-	{
-		CacheCollection = ToolProperties->DebugCacheCollection;
-		if (CacheCollection == nullptr)
-		{
-			UE_LOG(LogClothTrainingTool, Error, TEXT("Debug cache is None. Please select a valid cache for output."));
-		}
-	}
-	else
-	{
-		CacheCollection = ToolProperties->CacheCollection;
-		if (CacheCollection == nullptr)
-		{
-			UE_LOG(LogClothTrainingTool, Error, TEXT("Generated Cache is None. Please select a valid cache for output."));
-		}
-	}
-	return CacheCollection;
-}
-
-bool UClothTrainingTool::SaveCacheCollection(UChaosCacheCollection* CacheCollection) const
+UGeometryCache* UClothTrainingTool::GetCache() const
 {
-	if (CacheCollection == nullptr)
-	{
-		return false;
-	}
-	UPackage* const Package = CacheCollection->GetPackage();
-	if (Package == nullptr)
-	{
-		UE_LOG(LogClothTrainingTool, Error, TEXT("Failed to get package for %s"), *(CacheCollection->GetFName().ToString()));
-		return false;
-	}
-	const FString SavePath = Package->GetFName().ToString();
-	UE_LOG(LogClothTrainingTool, Display, TEXT("Save to %s"), *SavePath);
-
-	FSavePackageArgs SaveArgs;
-	SaveArgs.SaveFlags = SAVE_NoError;
-	const bool bSaveSucced = UPackage::SavePackage(Package, CacheCollection, *SavePath, SaveArgs);
-	if (!bSaveSucced)
-	{
-		UE_LOG(LogClothTrainingTool, Error, TEXT("Failed to save cache collection"));
-	}
-	return bSaveSucced;
+	const FString PackageName = ToolProperties->bDebug ? ToolProperties->DebugCacheName : ToolProperties->SimulatedCacheName;
+	return UE::ClothTrainingTool::Private::CreateOrLoad<UGeometryCache>(PackageName);
 }
 
 void UClothTrainingTool::FreeTaskResource(bool bCancelled)
@@ -786,12 +972,11 @@ void UClothTrainingTool::FreeTaskResource(bool bCancelled)
 
 	{
 		UE::ClothTrainingTool::Private::FTimeScope TimeScope(TEXT("Saving"));
-		TaskResource->Cache->bCompressChannels = true;
-		TaskResource->Cache->EndRecord(*TaskResource->CacheUserToken);
 
-		UChaosCacheCollection* const CacheCollection = GetCacheCollection();
-		ensureMsgf(CacheCollection != nullptr, TEXT("CacheCollection should not be nullptr"));
-		SaveCacheCollection(CacheCollection);
+		using UE::ClothTrainingTool::Private::SaveGeometryCache;
+		using UE::ClothTrainingTool::Private::SavePackage;
+		SaveGeometryCache(*TaskResource->Cache, *ClothComponent->GetClothAsset(), TaskResource->SimulatedPositions);
+		SavePackage(*TaskResource->Cache);
 	}
 	if (bCancelled)
 	{
diff --git a/Engine/Plugins/Experimental/ChaosClothAssetEditor/Source/ChaosClothAssetEditorTools/Private/ChaosClothAsset/ClothTrainingToolProperties.cpp b/Engine/Plugins/Experimental/ChaosClothAssetEditor/Source/ChaosClothAssetEditorTools/Private/ChaosClothAsset/ClothTrainingToolProperties.cpp
new file mode 100644
--- /dev/null
+++ b/Engine/Plugins/Experimental/ChaosClothAssetEditor/Source/ChaosClothAssetEditorTools/Private/ChaosClothAsset/ClothTrainingToolProperties.cpp
@@ -0,0 +1,90 @@
+// Copyright Epic Games, Inc. All Rights Reserved.
+
+#include "ClothTrainingToolProperties.h"
+
+#include "ChaosClothAsset/ClothTrainingTool.h"
+#include "ChaosClothAsset/ClothAsset.h"
+#include "ChaosClothAsset/ClothComponent.h"
+#include "ContentBrowserModule.h"
+#include "GeometryCache.h"
+#include "IContentBrowserSingleton.h"
+#include "Misc/FileHelper.h"
+
+#define LOCTEXT_NAMESPACE "ClothTrainingToolProperties"
+
+namespace UE::ClothTrainingTool::Private
+{
+	TOptional<FString> ExportGeometryCacheDialog(const UObject& ClothAsset)
+	{
+		FSaveAssetDialogConfig Config;
+		{
+			const FString PackageName = ClothAsset.GetOutermost()->GetName();
+			Config.DefaultPath = FPackageName::GetLongPackagePath(PackageName);
+			Config.DefaultAssetName = FString::Printf(TEXT("GC_%s"), *ClothAsset.GetName());
+			Config.AssetClassNames.Add(UGeometryCache::StaticClass()->GetClassPathName());
+			Config.ExistingAssetPolicy = ESaveAssetDialogExistingAssetPolicy::AllowButWarn;
+			Config.DialogTitleOverride = LOCTEXT("ExportGeometryCacheDialogTitle", "Export Geometry Cache As");
+		}
+	
+		FContentBrowserModule& ContentBrowserModule = FModuleManager::LoadModuleChecked<FContentBrowserModule>("ContentBrowser");
+	
+		FString NewPackageName;
+		FText OutError;
+		for (bool bFilenameValid = false; !bFilenameValid; bFilenameValid = FFileHelper::IsFilenameValidForSaving(NewPackageName, OutError))
+		{
+			const FString AssetPath = ContentBrowserModule.Get().CreateModalSaveAssetDialog(Config);
+			if (AssetPath.IsEmpty())
+			{
+				return TOptional<FString>();
+			}
+			NewPackageName = FPackageName::ObjectPathToPackageName(AssetPath);
+		}
+		return NewPackageName;
+	}
+}
+
+void UClothTrainingToolProperties::SetSimulatedCacheName()
+{
+	if (const UObject* ClothAsset = GetClothAsset())
+	{
+		using UE::ClothTrainingTool::Private::ExportGeometryCacheDialog;
+		if (TOptional<FString> OptionalName = ExportGeometryCacheDialog(*ClothAsset); OptionalName.IsSet())
+		{
+			SimulatedCacheName = OptionalName.GetValue();
+		}
+	}
+}
+
+void UClothTrainingToolProperties::SetDebugCacheName()
+{
+	if (const UObject* ClothAsset = GetClothAsset())
+	{
+		using UE::ClothTrainingTool::Private::ExportGeometryCacheDialog;
+		if (TOptional<FString> OptionalName = ExportGeometryCacheDialog(*ClothAsset); OptionalName.IsSet())
+		{
+			DebugCacheName = OptionalName.GetValue();
+		}
+	}
+}
+
+UObject* UClothTrainingToolProperties::GetClothAsset()
+{
+	if (const UClothTrainingTool* Tool = Cast<UClothTrainingTool>(GetOuter()))
+	{
+		if (Tool->ClothComponent)
+		{
+			return Tool->ClothComponent->GetClothAsset();
+		}
+	}
+	return nullptr;
+}
+
+void UClothTrainingToolActionProperties::StartGenerating()
+{
+	if (ParentTool.IsValid())
+	{
+		ParentTool->RequestAction(EClothTrainingToolActions::StartTrain);
+	}
+}
+
+#undef LOCTEXT_NAMESPACE
\ No newline at end of file
diff --git a/Engine/Plugins/Experimental/ChaosClothAssetEditor/Source/ChaosClothAssetEditorTools/Private/ChaosClothAsset/ClothTrainingToolProperties.h b/Engine/Plugins/Experimental/ChaosClothAssetEditor/Source/ChaosClothAssetEditorTools/Private/ChaosClothAsset/ClothTrainingToolProperties.h
new file mode 100644
--- /dev/null
+++ b/Engine/Plugins/Experimental/ChaosClothAssetEditor/Source/ChaosClothAssetEditorTools/Private/ChaosClothAsset/ClothTrainingToolProperties.h
@@ -0,0 +1,69 @@
+// Copyright Epic Games, Inc. All Rights Reserved.
+
+#pragma once
+
+#include "InteractiveTool.h"
+
+#include "ClothTrainingToolProperties.generated.h"
+
+class UAnimSequence;
+class USkinnedAsset;
+
+UCLASS()
+class UClothTrainingToolProperties : public UInteractiveToolPropertySet
+{
+	GENERATED_BODY()
+public:
+	UPROPERTY(EditAnywhere, Category = Input)
+	TObjectPtr<UAnimSequence> AnimationSequence;
+
+	/* e.g. "0, 2, 5-10, 12-15". If left empty, all frames will be used */
+	UPROPERTY(EditAnywhere, Category = Input)
+	FString FramesToSimulate;
+
+	UPROPERTY(VisibleAnywhere, Category = Output, meta = (EditCondition = "!bDebug"))
+	FString SimulatedCacheName;
+
+	UFUNCTION(CallInEditor, Category = "Output Actions", meta = (EditCondition = "!bDebug"))
+	void SetSimulatedCacheName();
+
+	UPROPERTY(EditAnywhere, Category = Output)
+	bool bDebug = false;
+
+	UPROPERTY(EditAnywhere, Category = Output, meta = (Min = 0, EditCondition = "bDebug"))
+	uint32 DebugFrame = 0;
+
+	UPROPERTY(VisibleAnywhere, Category = Output, meta = (EditCondition = "bDebug"))
+	FString DebugCacheName;
+
+	UFUNCTION(CallInEditor, Category = "Output Actions", meta = (EditCondition = "bDebug"))
+	void SetDebugCacheName();
+
+	UPROPERTY(EditAnywhere, Category = "Simulation Settings", meta = (Min = 0))
+	float TimeStep = 1.f / 30;
+
+	UPROPERTY(EditAnywhere, Category = "Simulation Settings", meta = (Min = 0))
+	int32 NumSteps = 200;
+
+	UPROPERTY(EditAnywhere, Category = "Simulation Settings", meta = (Min = 1, EditCondition = "!bDebug"))
+	int32 NumThreads = 1;
+
+private:
+	UObject* GetClothAsset();
+};
+
+class UClothTrainingTool;
+
+UCLASS()
+class UClothTrainingToolActionProperties : public UObject
+{
+	GENERATED_BODY()
+public:
+
+	TWeakObjectPtr<UClothTrainingTool> ParentTool;
+
+	void Initialize(UClothTrainingTool* ParentToolIn) { ParentTool = ParentToolIn; }
+
+	UFUNCTION(CallInEditor, Category = Actions, meta = (DisplayName = "Begin Generating", DisplayPriority = 1))
+	void StartGenerating();
+};
\ No newline at end of file
diff --git a/Engine/Plugins/Experimental/ChaosClothAssetEditor/Source/ChaosClothAssetEditorTools/Public/ChaosClothAsset/ClothTrainingTool.h b/Engine/Plugins/Experimental/ChaosClothAssetEditor/Source/ChaosClothAssetEditorTools/Public/ChaosClothAsset/ClothTrainingTool.h
--- a/Engine/Plugins/Experimental/ChaosClothAssetEditor/Source/ChaosClothAssetEditorTools/Public/ChaosClothAsset/ClothTrainingTool.h
+++ b/Engine/Plugins/Experimental/ChaosClothAssetEditor/Source/ChaosClothAssetEditorTools/Public/ChaosClothAsset/ClothTrainingTool.h
@@ -14,57 +14,16 @@
 struct FCacheUserToken;
 class FAsyncTaskNotification;
 class UAnimSequence;
-class UChaosCache;
-class UChaosCacheCollection;
 class UChaosClothComponent;
-class UClothTrainingTool;
+class UClothTrainingToolActionProperties;
+class UClothTrainingToolProperties;
+class UGeometryCache;
 
 namespace UE::Chaos::ClothAsset
 {
 	class FClothSimulationDataGenerationProxy;
 };
 
-// @@@@@@@@@ TODO: Change this to whatever makes sense for output
-struct FSkinnedMeshVertices
-{
-	TArray<FVector3f> Vertices;
-};
-
-UCLASS()
-class CHAOSCLOTHASSETEDITORTOOLS_API UClothTrainingToolProperties : public UInteractiveToolPropertySet
-{
-	GENERATED_BODY()
-public:
-
-	UPROPERTY(EditAnywhere, Category = Input)
-	TObjectPtr<UAnimSequence> AnimationSequence;
-
-	/* e.g. "0, 2, 5-10, 12-15". If left empty, all frames will be used */
-	UPROPERTY(EditAnywhere, Category = Input)
-	FString FramesToSimulate;
-
-	UPROPERTY(EditAnywhere, Category = Output, meta = (DisplayName = "Generated Cache", EditCondition = "!bDebug"))
-	TObjectPtr<UChaosCacheCollection> CacheCollection;
-
-	UPROPERTY(EditAnywhere, Category = Output)
-	bool bDebug = false;
-
-	UPROPERTY(EditAnywhere, Category = Output, meta = (Min = 0, EditCondition = "bDebug"))
-	uint32 DebugFrame = 0;
-
-	UPROPERTY(EditAnywhere, Category = Output, meta = (EditCondition = "bDebug", DisplayName = "Debug Cache"))
-	TObjectPtr<UChaosCacheCollection> DebugCacheCollection;
-
-	UPROPERTY(EditAnywhere, Category = "Simulation Settings", meta = (Min = 0))
-	float TimeStep = 1.f / 30;
-
-	UPROPERTY(EditAnywhere, Category = "Simulation Settings", meta = (Min = 0))
-	int32 NumSteps = 200;
-
-	UPROPERTY(EditAnywhere, Category = "Simulation Settings", meta = (Min = 1, EditCondition = "!bDebug"))
-	int32 NumThreads = 1;
-};
-
 UENUM()
 enum class EClothTrainingToolActions
 {
@@ -73,26 +32,6 @@ enum class EClothTrainingToolActions
 	TickTrain
 };
 
-UCLASS()
-class CHAOSCLOTHASSETEDITORTOOLS_API UClothTrainingToolActionProperties : public UObject
-{
-	GENERATED_BODY()
-public:
-
-	TWeakObjectPtr<UClothTrainingTool> ParentTool;
-
-	void Initialize(UClothTrainingTool* ParentToolIn) { ParentTool = ParentToolIn; }
-
-	void PostAction(EClothTrainingToolActions Action);
-
-	UFUNCTION(CallInEditor, Category = Actions, meta = (DisplayName = "Begin Generating", DisplayPriority = 1))
-		void StartGenerating()
-	{
-		PostAction(EClothTrainingToolActions::StartTrain);
-	}
-
-};
-
 
 UCLASS()
 class CHAOSCLOTHASSETEDITORTOOLS_API UClothTrainingToolBuilder : public UInteractiveToolWithToolTargetsBuilder
@@ -135,6 +74,7 @@ private:
 	using FProxy = UE::Chaos::ClothAsset::FClothSimulationDataGenerationProxy;
 
 	friend class UClothTrainingToolActionProperties;
+	friend class UClothTrainingToolProperties;
 
 	UPROPERTY()
 	TObjectPtr<UClothTrainingToolProperties> ToolProperties;
@@ -155,8 +95,7 @@ private:
 	void RunTraining();
 	void FreeTaskResource(bool bCancelled);
 	bool IsClothComponentValid() const;
-	UChaosCacheCollection* GetCacheCollection() const;
-	bool SaveCacheCollection(UChaosCacheCollection* CacheCollection) const;
+	UGeometryCache* GetCache() const;
 };
 
 
