ClothTrainingTool: add MLDeformerAsset field
use MLDeformer Asset import map for correct vertex ordering
use MLDeformer Asset to check asset inconsistencies
fix missing last frame in sequencer when saving geometry cache
fix floating point issue in AnimSequence

diff --git a/Engine/Plugins/Experimental/ChaosClothAssetEditor/Source/ChaosClothAssetEditorTools/Private/ChaosClothAsset/ClothTrainingTool.cpp b/Engine/Plugins/Experimental/ChaosClothAssetEditor/Source/ChaosClothAssetEditorTools/Private/ChaosClothAsset/ClothTrainingTool.cpp
--- a/Engine/Plugins/Experimental/ChaosClothAssetEditor/Source/ChaosClothAssetEditorTools/Private/ChaosClothAsset/ClothTrainingTool.cpp
+++ b/Engine/Plugins/Experimental/ChaosClothAssetEditor/Source/ChaosClothAssetEditorTools/Private/ChaosClothAsset/ClothTrainingTool.cpp
@@ -29,8 +29,10 @@
 #include "Internationalization/Regex.h"
 #include "ModelingOperators.h"
 #include "Misc/AsyncTaskNotification.h"
+#include "Misc/Optional.h"
 #include "PhysicsEngine/PhysicsAsset.h"
 #include "Rendering/SkeletalMeshRenderData.h"
+#include "Rendering/SkeletalMeshModel.h"
 #include "SkeletalRenderPublic.h"
 #include "Tasks/Pipe.h"
 #include "ToolTargetManager.h"
@@ -250,7 +252,7 @@ namespace UE::ClothTrainingTool::Private
 		return Tangents;
 	}
 
-	void SaveGeometryCache(UGeometryCache& GeometryCache, const USkinnedAsset& Asset, TArrayView<TArray<FVector3f>> PositionsToMoveFrom)
+	void SaveGeometryCache(UGeometryCache& GeometryCache, const USkinnedAsset& Asset, TConstArrayView<uint32> ImportedVertexNumbers, TArrayView<TArray<FVector3f>> PositionsToMoveFrom)
 	{
 		const FSkeletalMeshRenderData* RenderData = Asset.GetResourceForRendering();
 		constexpr int32 LODIndex = 0;
@@ -286,7 +288,6 @@ namespace UE::ClothTrainingTool::Private
 		const TArray<FVector2f> UV0 = GetUV0s(LODData);
 		const bool bHasUV0 = UV0.Num() == NumVertices;
 		const TArray<FColor> Colors = GetColors(LODData, NumVertices);
-		const TArray<uint32> ImportedVertexNumbers = Range(0, NumVertices);
 
 		GeometryCache.ClearForReimporting();
 
@@ -353,7 +354,7 @@ namespace UE::ClothTrainingTool::Private
 		}
 
 		TArray<FMatrix> Mats { FMatrix::Identity, FMatrix::Identity };
-		TArray<float> MatTimes { 0.0f, (MaxRecordedFrame - 1) / CacheFPS };
+		TArray<float> MatTimes { 0.0f, MaxRecordedFrame / CacheFPS };
 		Track->SetMatrixSamples(Mats, MatTimes);
 
 		if (ensureAlways(Track->EndCoding()))
@@ -408,6 +409,60 @@ namespace UE::ClothTrainingTool::Private
 		constexpr bool bPromptToSave = false;
 		FEditorFileUtils::PromptForCheckoutAndSave(PackagesToSave, bCheckDirty, bPromptToSave);
 	}
+	
+	TOptional<TArray<int32>> GetMeshImportVertexMap(const USkinnedAsset& MLDeformerAsset, const UChaosClothAsset& ClothAsset)
+	{
+		constexpr int32 LODIndex = 0;
+		const TOptional<TArray<int32>> None;
+		const FSkeletalMeshModel* const MLDModel = MLDeformerAsset.GetImportedModel();
+		if (!MLDModel || !MLDModel->LODModels.IsValidIndex(LODIndex))
+		{
+			return None;
+		}
+		const FSkeletalMeshLODModel& MLDLOD = MLDModel->LODModels[LODIndex];
+		const TArray<int32>& Map = MLDLOD.MeshToImportVertexMap;
+		if (Map.IsEmpty())
+		{
+			UE_LOG(LogClothTrainingTool, Warning, TEXT("MeshToImportVertexMap is empty. MLDeformer Asset should be an imported SkeletalMesh (e.g. from fbx)."));
+			return None;
+		}
+		const FSkeletalMeshModel* const ClothModel = ClothAsset.GetImportedModel();
+		if (!ClothModel || !ClothModel->LODModels.IsValidIndex(LODIndex))
+		{
+			UE_LOG(LogClothTrainingTool, Warning, TEXT("ClothAsset has no imported model."));
+			return None;
+		}
+		const FSkeletalMeshLODModel& ClothLOD = ClothModel->LODModels[LODIndex];
+
+		if (MLDLOD.NumVertices != ClothLOD.NumVertices || MLDLOD.Sections.Num() != ClothLOD.Sections.Num())
+		{
+			UE_LOG(LogClothTrainingTool, Warning, TEXT("MLDeformerAsset and ClothAsset have different number of vertices or sections. Check if the assets have the same mesh."));
+			return None;
+		}
+		
+		for (int32 SectionIndex = 0; SectionIndex < MLDLOD.Sections.Num(); ++SectionIndex)
+		{
+			const FSkelMeshSection& MLDSection = MLDLOD.Sections[SectionIndex];
+			const FSkelMeshSection& ClothSection = ClothLOD.Sections[SectionIndex];
+			if (MLDSection.NumVertices != ClothSection.NumVertices)
+			{
+				UE_LOG(LogClothTrainingTool, Warning, TEXT("MLDeformerAsset and ClothAsset have different number of vertices in section %d. Check if the assets have the same mesh."), SectionIndex);
+				return None;
+			}
+			for (int32 VertexIndex = 0; VertexIndex < MLDSection.NumVertices; ++VertexIndex)
+			{
+				const FVector3f& MLDPosition = MLDSection.SoftVertices[VertexIndex].Position;
+				const FVector3f& ClothPosition = ClothSection.SoftVertices[VertexIndex].Position;
+				if (!MLDPosition.Equals(ClothPosition, UE_KINDA_SMALL_NUMBER))
+				{
+					UE_LOG(LogClothTrainingTool, Warning, TEXT("MLDeformerAsset and ClothAsset have different vertex positions. Check if the assets have the same vertex order."));
+					return None;
+				}
+			}
+		}
+
+		return Map;
+	}
 };
 
 // @@@@@@@@@ TODO: Change this to whatever makes sense for output
@@ -426,16 +481,34 @@ struct UClothTrainingTool::FSimResource
 	TArrayView<TArray<FVector3f>> SimulatedPositions;
 };
 
+struct UClothTrainingTool::FTaskResource
+{
+	TUniquePtr<FCriticalSection> SimMutex;
+	TArray<FSimResource> SimResources;
+
+	TUniquePtr<FExecuterType> Executer;
+	TUniquePtr<FAsyncTaskNotification> Notification;
+	FDateTime StartTime;
+	FDateTime LastUpdateTime;
+
+	TArray<TArray<FVector3f>> SimulatedPositions;
+	TArray<uint32> ImportedVertexNumbers;
+	UGeometryCache* Cache = nullptr;
+
+	bool AllocateSimResources_GameThread(const UChaosClothComponent& InClothComponent, int32 Num);
+	void FreeSimResources_GameThread();
+	void FlushRendering();
+};
+
 class UClothTrainingTool::FLaunchSimsOp : public UE::Geometry::TGenericDataOperator<FSkinnedMeshVertices>
 {
 public: 
-	FLaunchSimsOp(TArray<FSimResource>& InSimResources, FCriticalSection& InSimMutex, TArray<TArray<FVector3f>>& InSimulatedPositions,  TObjectPtr<UClothTrainingToolProperties> InToolProperties)
-		: SimResources(InSimResources)
-		, SimMutex(InSimMutex)
-		, SimulatedPositions(InSimulatedPositions)
+	FLaunchSimsOp(FTaskResource& InTaskResource, TObjectPtr<UClothTrainingToolProperties> InToolProperties)
+		: SimResources(InTaskResource.SimResources)
+		, SimMutex(*InTaskResource.SimMutex)
+		, SimulatedPositions(InTaskResource.SimulatedPositions)
 		, ToolProperties(InToolProperties)
-	{
-	}
+	{}
 
 	virtual void CalculateResult(FProgressCancel* Progress) override;
 
@@ -465,24 +538,6 @@ private:
 	inline static const FName PositionZName = TEXT("PositionZ");
 };
 
-struct UClothTrainingTool::FTaskResource
-{
-	TUniquePtr<FCriticalSection> SimMutex;
-	TArray<FSimResource> SimResources;
-
-	TUniquePtr<FExecuterType> Executer;
-	TUniquePtr<FAsyncTaskNotification> Notification;
-	FDateTime StartTime;
-	FDateTime LastUpdateTime;
-
-	TArray<TArray<FVector3f>> SimulatedPositions;
-	UGeometryCache* Cache = nullptr;
-
-	bool AllocateSimResources_GameThread(const UChaosClothComponent& InClothComponent, int32 Num);
-	void FreeSimResources_GameThread();
-	void FlushRendering();
-};
-
 void UClothTrainingTool::FLaunchSimsOp::PrepareAnimationSequence()
 {
 	TObjectPtr<UAnimSequence> AnimationSequence = ToolProperties->AnimationSequence;
@@ -505,7 +560,7 @@ void UClothTrainingTool::FLaunchSimsOp::RestoreAnimationSequence()
 TArray<FTransform> UClothTrainingTool::FLaunchSimsOp::GetBoneTransforms(UChaosClothComponent& InClothComponent, int32 Frame) const
 {
 	const UAnimSequence* AnimationSequence = ToolProperties->AnimationSequence;
-	const double Time = AnimationSequence->GetTimeAtFrame(Frame);
+	const double Time = FMath::Clamp(AnimationSequence->GetSamplingFrameRate().AsSeconds(Frame), 0., (double)AnimationSequence->GetPlayLength());
 	FAnimExtractContext ExtractionContext(Time);
 
 	UChaosClothAsset* const ClothAsset = InClothComponent.GetClothAsset();
@@ -765,6 +820,18 @@ void UClothTrainingTool::StartTraining()
 		PendingAction = EClothTrainingToolActions::NoAction;
 		return;
 	}
+	if (!ToolProperties->MLDeformerAsset)
+	{
+		PendingAction = EClothTrainingToolActions::NoAction;
+		return;
+	}
+	using UE::ClothTrainingTool::Private::GetMeshImportVertexMap;
+	TOptional<TArray<int32>> OptionalMap = GetMeshImportVertexMap(*ToolProperties->MLDeformerAsset, *ClothComponent->GetClothAsset());
+	if (!OptionalMap)
+	{
+		PendingAction = EClothTrainingToolActions::NoAction;
+		return;
+	}
 	TaskResource = MakeUnique<FTaskResource>();
 
 	if (!TaskResource->AllocateSimResources_GameThread(*ClothComponent, ToolProperties->NumThreads))
@@ -774,7 +841,7 @@ void UClothTrainingTool::StartTraining()
 	}
 	TaskResource->Cache = Cache;
 
-	TUniquePtr<FLaunchSimsOp> NewOp = MakeUnique<FLaunchSimsOp>(TaskResource->SimResources, *TaskResource->SimMutex, TaskResource->SimulatedPositions, ToolProperties);
+	TUniquePtr<FLaunchSimsOp> NewOp = MakeUnique<FLaunchSimsOp>(*TaskResource, ToolProperties);
 	TaskResource->Executer = MakeUnique<FExecuterType>(MoveTemp(NewOp));
 	TaskResource->Executer->StartBackgroundTask();
 
@@ -788,6 +855,9 @@ void UClothTrainingTool::StartTraining()
 	TaskResource->StartTime = FDateTime::UtcNow();
 	TaskResource->LastUpdateTime = TaskResource->StartTime;
 
+	const TArray<int32>& Map = OptionalMap.GetValue();
+	TaskResource->ImportedVertexNumbers = TArray<uint32>(reinterpret_cast<const uint32*>(Map.GetData()), Map.Num());
+
 	PendingAction = EClothTrainingToolActions::TickTrain;
 }
 
@@ -975,7 +1045,7 @@ void UClothTrainingTool::FreeTaskResource(bool bCancelled)
 
 		using UE::ClothTrainingTool::Private::SaveGeometryCache;
 		using UE::ClothTrainingTool::Private::SavePackage;
-		SaveGeometryCache(*TaskResource->Cache, *ClothComponent->GetClothAsset(), TaskResource->SimulatedPositions);
+		SaveGeometryCache(*TaskResource->Cache, *ClothComponent->GetClothAsset(), TaskResource->ImportedVertexNumbers, TaskResource->SimulatedPositions);
 		SavePackage(*TaskResource->Cache);
 	}
 	if (bCancelled)
diff --git a/Engine/Plugins/Experimental/ChaosClothAssetEditor/Source/ChaosClothAssetEditorTools/Private/ChaosClothAsset/ClothTrainingToolProperties.h b/Engine/Plugins/Experimental/ChaosClothAssetEditor/Source/ChaosClothAssetEditorTools/Private/ChaosClothAsset/ClothTrainingToolProperties.h
--- a/Engine/Plugins/Experimental/ChaosClothAssetEditor/Source/ChaosClothAssetEditorTools/Private/ChaosClothAsset/ClothTrainingToolProperties.h
+++ b/Engine/Plugins/Experimental/ChaosClothAssetEditor/Source/ChaosClothAssetEditorTools/Private/ChaosClothAsset/ClothTrainingToolProperties.h
@@ -14,6 +14,10 @@ class UClothTrainingToolProperties : public UInteractiveToolPropertySet
 {
 	GENERATED_BODY()
 public:
+	/* Skeletal mesh that will be used in MLDeformer */
+	UPROPERTY(EditAnywhere, Category = Input)
+	TObjectPtr<USkinnedAsset> MLDeformerAsset;
+
 	UPROPERTY(EditAnywhere, Category = Input)
 	TObjectPtr<UAnimSequence> AnimationSequence;
 
