Chaos Cloth - Fixed Datasmith Clo Adapter plugin inclusion/dependencies to the Cloth Editor.

#rnx

diff --git a/Engine/Plugins/Experimental/ChaosClothAssetEditor/ChaosClothAssetEditor.uplugin b/Engine/Plugins/Experimental/ChaosClothAssetEditor/ChaosClothAssetEditor.uplugin
--- a/Engine/Plugins/Experimental/ChaosClothAssetEditor/ChaosClothAssetEditor.uplugin
+++ b/Engine/Plugins/Experimental/ChaosClothAssetEditor/ChaosClothAssetEditor.uplugin
@@ -70,6 +70,13 @@
 			"Name": "DatasmithImporter",
 			"Enabled": true
 		},
+		{
+			"Name": "DatasmithCloTranslator",
+			"Enabled": true,
+			"SupportedTargetPlatforms": [
+				"Win64"
+			]
+		},
 		{
 			"Name": "ChaosCaching",
 			"Enabled": true
