FClothPatternToDynamicMesh: assert if we try to access wedge data

diff --git a/Engine/Plugins/Experimental/ChaosClothAssetEditor/Source/ChaosClothAssetTools/Private/ChaosClothAsset/ClothPatternToDynamicMesh.cpp b/Engine/Plugins/Experimental/ChaosClothAssetEditor/Source/ChaosClothAssetTools/Private/ChaosClothAsset/ClothPatternToDynamicMesh.cpp
--- a/Engine/Plugins/Experimental/ChaosClothAssetEditor/Source/ChaosClothAssetTools/Private/ChaosClothAsset/ClothPatternToDynamicMesh.cpp
+++ b/Engine/Plugins/Experimental/ChaosClothAssetEditor/Source/ChaosClothAssetTools/Private/ChaosClothAsset/ClothPatternToDynamicMesh.cpp
@@ -176,46 +176,43 @@ public:
 	// -- Access to per-wedge attributes -- //
 	void GetWedgeIDs(const TriIDType& TriID, WedgeIDType& WID0, WedgeIDType& WID1, WedgeIDType& WID2) const
 	{
-		// TODO: Review this
-		int32 Offset = 3 * TriID;
-		WID0 = Offset;
-		WID1 = Offset + 1;
-		WID2 = Offset + 2;
+		checkf(false, TEXT("FClothPatternWrapper: ClothPatterns are not expected to use Wedges"));
 	}
 
 	FVector2f GetWedgeUV(int32 UVLayerIndex, WedgeIDType WID) const
 	{
-		checkf(VertexDataType == EClothPatternVertexType::Render, TEXT("Requested Wedge UVs from a Sim mesh"));
-		const TConstArrayView<FVector2f> VertexUVs = Pattern.GetRenderUVs()[WID];
-		return VertexUVs[UVLayerIndex];
+		checkf(false, TEXT("FClothPatternWrapper: ClothPatterns are not expected to use Wedges"));
+		return FVector2f();
 	}
 
 	FVector3f GetWedgeNormal(WedgeIDType WID) const
 	{
-		ensure(false);
-		return FVector3f(0.0f, 0.0f, 0.0f);
+		checkf(false, TEXT("FClothPatternWrapper: ClothPatterns are not expected to use Wedges"));
+		return FVector3f();
 	}
 
 	FVector3f GetWedgeTangent(WedgeIDType WID) const
 	{
-		ensure(false);
-		return FVector3f(0.0f, 0.0f, 0.0f);
+		checkf(false, TEXT("FClothPatternWrapper: ClothPatterns are not expected to use Wedges"));
+		return FVector3f();
 	}
 
 	FVector3f GetWedgeBiTangent(WedgeIDType WID) const
 	{
-		ensure(false);
-		return FVector3f(0.0f, 0.0f, 0.0f);
+		checkf(false, TEXT("FClothPatternWrapper: ClothPatterns are not expected to use Wedges"));
+		return FVector3f();
 	}
 	
 	FVector4f GetWedgeColor(WedgeIDType WID) const
 	{
-		ensure(false);
-		return FVector4f(0.0f, 0.0f, 0.0f, 0.0f);
+		checkf(false, TEXT("FClothPatternWrapper: ClothPatterns are not expected to use Wedges"));
+		return FVector4f();
 	}
+	// -- End of per-wedge attribute access -- //
 
 	int32 GetMaterialIndex(TriIDType TriID) const
 	{
+		checkf(false, TEXT("FClothPatternWrapper: Material indexing should be accomplished by passing a function into Convert"));
 		return 0;
 	}
 
