Chaos Cloth Editor - Made the Outliner follow the Dataflow selection.

#rnx

diff --git a/Engine/Plugins/Experimental/ChaosClothAssetEditor/Source/ChaosClothAssetEditor/Private/ChaosClothAsset/ClothEditorToolkit.cpp b/Engine/Plugins/Experimental/ChaosClothAssetEditor/Source/ChaosClothAssetEditor/Private/ChaosClothAsset/ClothEditorToolkit.cpp
--- a/Engine/Plugins/Experimental/ChaosClothAssetEditor/Source/ChaosClothAssetEditor/Private/ChaosClothAsset/ClothEditorToolkit.cpp
+++ b/Engine/Plugins/Experimental/ChaosClothAssetEditor/Source/ChaosClothAssetEditor/Private/ChaosClothAsset/ClothEditorToolkit.cpp
@@ -1003,6 +1003,11 @@ void FChaosClothAssetEditorToolkit::OnNodeSelectionChanged(const TSet<UObject*>&
 	{
 		ClothMode->SetSelectedClothCollection(Collection);
 	}
+
+	if (Outliner)
+	{
+		Outliner->SetClothCollection(Collection);
+	}
 }
 
 
diff --git a/Engine/Plugins/Experimental/ChaosClothAssetEditor/Source/ChaosClothAssetEditor/Private/ChaosClothAsset/SClothCollectionOutliner.cpp b/Engine/Plugins/Experimental/ChaosClothAssetEditor/Source/ChaosClothAssetEditor/Private/ChaosClothAsset/SClothCollectionOutliner.cpp
--- a/Engine/Plugins/Experimental/ChaosClothAssetEditor/Source/ChaosClothAssetEditor/Private/ChaosClothAsset/SClothCollectionOutliner.cpp
+++ b/Engine/Plugins/Experimental/ChaosClothAssetEditor/Source/ChaosClothAssetEditor/Private/ChaosClothAsset/SClothCollectionOutliner.cpp
@@ -79,7 +79,13 @@ void SClothCollectionOutliner::Construct(const FArguments& InArgs)
 
 void SClothCollectionOutliner::SetClothCollection(TWeakPtr<FManagedArrayCollection> InClothCollection)
 {
-	ClothCollection = InClothCollection;
+	if (ClothCollection != InClothCollection)
+	{
+		ClothCollection = InClothCollection;
+
+		RegenerateHeader();
+		RepopulateListView();
+	}
 }
 
 void SClothCollectionOutliner::SetSelectedGroupName(const FName& InSelectedGroupName)
