Chaos Cloth Editor - Fixed overzealous ensure on invalid material ids from the Cloth Asset Editor.

#rnx

diff --git a/Engine/Plugins/Experimental/ChaosClothAssetEditor/Source/ChaosClothAssetTools/Private/ChaosClothAsset/ClothPatternToDynamicMesh.cpp b/Engine/Plugins/Experimental/ChaosClothAssetEditor/Source/ChaosClothAssetTools/Private/ChaosClothAsset/ClothPatternToDynamicMesh.cpp
--- a/Engine/Plugins/Experimental/ChaosClothAssetEditor/Source/ChaosClothAssetTools/Private/ChaosClothAsset/ClothPatternToDynamicMesh.cpp
+++ b/Engine/Plugins/Experimental/ChaosClothAssetEditor/Source/ChaosClothAssetTools/Private/ChaosClothAsset/ClothPatternToDynamicMesh.cpp
@@ -421,13 +421,13 @@ void FClothPatternToDynamicMesh::Convert(const TSharedPtr<const FManagedArrayCol
 		const FCollectionClothLodConstFacade LodFacade(ClothFacade.GetLod(LODIndex));
 		const TArrayView<const int32> RenderMaterialIndex = LodFacade.GetRenderMaterialIndex();
 
-		auto TriangleToMaterialFunction = [&RenderMaterialIndex](FClothPatternWrapper::TriIDType TriID)
+		auto TriangleToMaterialFunction = [&RenderMaterialIndex, VertexDataType](FClothPatternWrapper::TriIDType TriID)->int32
 		{
-			if (ensure(RenderMaterialIndex.IsValidIndex(TriID)))
+			if (VertexDataType == EClothPatternVertexType::Render && ensure(RenderMaterialIndex.IsValidIndex(TriID)))
 			{
 				return RenderMaterialIndex[TriID];
 			}
-			return 0;
+			return INDEX_NONE;
 		};
 
 		constexpr bool bCopyTangents = false;
