Fix "focus to selection" behavior in the 2D Construction viewport

diff --git a/Engine/Plugins/Experimental/ChaosClothAssetEditor/Source/ChaosClothAssetEditor/Private/ChaosClothAsset/SClothEditorRestSpaceViewport.cpp b/Engine/Plugins/Experimental/ChaosClothAssetEditor/Source/ChaosClothAssetEditor/Private/ChaosClothAsset/SClothEditorRestSpaceViewport.cpp
--- a/Engine/Plugins/Experimental/ChaosClothAssetEditor/Source/ChaosClothAssetEditor/Private/ChaosClothAsset/SClothEditorRestSpaceViewport.cpp
+++ b/Engine/Plugins/Experimental/ChaosClothAssetEditor/Source/ChaosClothAssetEditor/Private/ChaosClothAsset/SClothEditorRestSpaceViewport.cpp
@@ -147,12 +147,15 @@ TSharedPtr<SWidget> SChaosClothAssetEditorRestSpaceViewport::MakeViewportToolbar
 
 void SChaosClothAssetEditorRestSpaceViewport::OnFocusViewportToSelection()
 {
-	const FEditorModeTools* const EditorModeTools = Client->GetModeTools();
-	const UChaosClothAssetEditorMode* const ClothEdMode = Cast<UChaosClothAssetEditorMode>(EditorModeTools->GetActiveScriptableMode(UChaosClothAssetEditorMode::EM_ChaosClothAssetEditorModeId));
+	const UChaosClothAssetEditorMode* const ClothEdMode = GetEdMode();
 
 	if (ClothEdMode)
 	{
 		Client->FocusViewportOnBox(ClothEdMode->SelectionBoundingBox());
+
+		// Reset any changes to the clip planes by the scroll zoom behavior
+		Client->OverrideNearClipPlane(UE_KINDA_SMALL_NUMBER);
+		Client->OverrideFarClipPlane(0);
 	}
 }
 
