Share a command list between Editor Toolkit and Preview Viewport so that command input chords can be used editor-wide

diff --git a/Engine/Plugins/Experimental/ChaosClothAssetEditor/Source/ChaosClothAssetEditor/Private/ChaosClothAsset/ClothEditorToolkit.cpp b/Engine/Plugins/Experimental/ChaosClothAssetEditor/Source/ChaosClothAssetEditor/Private/ChaosClothAsset/ClothEditorToolkit.cpp
--- a/Engine/Plugins/Experimental/ChaosClothAssetEditor/Source/ChaosClothAssetEditor/Private/ChaosClothAsset/ClothEditorToolkit.cpp
+++ b/Engine/Plugins/Experimental/ChaosClothAssetEditor/Source/ChaosClothAssetEditor/Private/ChaosClothAsset/ClothEditorToolkit.cpp
@@ -192,7 +192,8 @@ FChaosClothAssetEditorToolkit::FChaosClothAssetEditorToolkit(UAssetEditor* InOwn
 	ClothPreviewViewportDelegate = [this](FAssetEditorViewportConstructionArgs InArgs)
 	{
 		return SAssignNew(PreviewViewportWidget, SChaosClothAssetEditor3DViewport, InArgs)
-			.EditorViewportClient(ClothPreviewViewportClient);
+			.EditorViewportClient(ClothPreviewViewportClient)
+			.ToolkitCommandList(GetToolkitCommands().ToSharedPtr());
 	};
 
 	FPreviewScene::ConstructionValues SceneArgs;
diff --git a/Engine/Plugins/Experimental/ChaosClothAssetEditor/Source/ChaosClothAssetEditor/Private/ChaosClothAsset/SClothEditor3DViewport.cpp b/Engine/Plugins/Experimental/ChaosClothAssetEditor/Source/ChaosClothAssetEditor/Private/ChaosClothAsset/SClothEditor3DViewport.cpp
--- a/Engine/Plugins/Experimental/ChaosClothAssetEditor/Source/ChaosClothAssetEditor/Private/ChaosClothAsset/SClothEditor3DViewport.cpp
+++ b/Engine/Plugins/Experimental/ChaosClothAssetEditor/Source/ChaosClothAssetEditor/Private/ChaosClothAsset/SClothEditor3DViewport.cpp
@@ -21,6 +21,7 @@ void SChaosClothAssetEditor3DViewport::Construct(const FArguments& InArgs, const
 	{
 		ParentArgs._ViewportSize = InArgs._ViewportSize;
 	}
+	ToolkitCommandList = InArgs._ToolkitCommandList;
 	SAssetEditorViewport::Construct(ParentArgs, InViewportConstructionArgs);
 
 	ViewportOverlay->AddSlot()
@@ -65,7 +66,7 @@ void SChaosClothAssetEditor3DViewport::BindCommands()
 	SAssetEditorViewport::BindCommands();
 	const FChaosClothAssetEditorCommands& CommandInfos = FChaosClothAssetEditorCommands::Get();
 
-	CommandList->MapAction(
+	ToolkitCommandList->MapAction(
 		CommandInfos.TogglePreviewWireframe,
 		FExecuteAction::CreateLambda([this]()
 		{
@@ -75,7 +76,7 @@ void SChaosClothAssetEditor3DViewport::BindCommands()
 		FCanExecuteAction::CreateLambda([this]() { return true; }),
 		FIsActionChecked::CreateLambda([this]() { return StaticCastSharedPtr<FChaosClothAssetEditor3DViewportClient>(Client)->RenderMeshWireframeEnabled(); }));
 
-	CommandList->MapAction(
+	ToolkitCommandList->MapAction(
 		CommandInfos.SoftResetSimulation,
 		FExecuteAction::CreateLambda([this]()
 		{
@@ -85,7 +86,7 @@ void SChaosClothAssetEditor3DViewport::BindCommands()
 		FCanExecuteAction::CreateLambda([this]() { return true; }),
 		FIsActionChecked::CreateLambda([this]() { return false; }));
 
-	CommandList->MapAction(
+	ToolkitCommandList->MapAction(
 		CommandInfos.HardResetSimulation,
 		FExecuteAction::CreateLambda([this]()
 		{
@@ -96,7 +97,7 @@ void SChaosClothAssetEditor3DViewport::BindCommands()
 		FIsActionChecked::CreateLambda([this]() { return false; }));
 
 
-	CommandList->MapAction(
+	ToolkitCommandList->MapAction(
 		CommandInfos.ToggleSimulationSuspended,
 		FExecuteAction::CreateLambda([this]()
 		{
@@ -120,7 +121,7 @@ void SChaosClothAssetEditor3DViewport::BindCommands()
 TSharedPtr<SWidget> SChaosClothAssetEditor3DViewport::MakeViewportToolbar()
 {
 	return SNew(SChaosClothAssetEditor3DViewportToolBar, SharedThis(this))
-		.CommandList(CommandList);
+		.CommandList(ToolkitCommandList);
 }
 
 void SChaosClothAssetEditor3DViewport::OnFocusViewportToSelection()
diff --git a/Engine/Plugins/Experimental/ChaosClothAssetEditor/Source/ChaosClothAssetEditor/Private/ChaosClothAsset/SClothEditor3DViewport.h b/Engine/Plugins/Experimental/ChaosClothAssetEditor/Source/ChaosClothAssetEditor/Private/ChaosClothAsset/SClothEditor3DViewport.h
--- a/Engine/Plugins/Experimental/ChaosClothAssetEditor/Source/ChaosClothAssetEditor/Private/ChaosClothAsset/SClothEditor3DViewport.h
+++ b/Engine/Plugins/Experimental/ChaosClothAssetEditor/Source/ChaosClothAssetEditor/Private/ChaosClothAsset/SClothEditor3DViewport.h
@@ -24,6 +24,7 @@ public:
 	SLATE_BEGIN_ARGS(SChaosClothAssetEditor3DViewport) {}
 		SLATE_ATTRIBUTE(FVector2D, ViewportSize);
 		SLATE_ARGUMENT(TSharedPtr<FEditorViewportClient>, EditorViewportClient)
+		SLATE_ARGUMENT(TSharedPtr<FUICommandList>, ToolkitCommandList)
 	SLATE_END_ARGS()
 
 	void Construct(const FArguments& InArgs, const FAssetEditorViewportConstructionArgs& InViewportConstructionArgs);
@@ -41,6 +42,10 @@ public:
 
 private:
 
+	// Use this command list if we want to enable editor-wide command chords/hotkeys. 
+	// Use SEditorViewport::CommandList if we want command hotkeys to only be active when the mouse is in this viewport.
+	TSharedPtr<FUICommandList> ToolkitCommandList;
+
 	TWeakPtr<UE::Chaos::ClothAsset::FChaosClothPreviewScene> GetPreviewScene();
 	TWeakPtr<const UE::Chaos::ClothAsset::FChaosClothPreviewScene> GetPreviewScene() const;
 
