ClothTrainingTool: fix non unity compiler error

diff --git a/Engine/Plugins/Experimental/ChaosClothAssetEditor/Source/ChaosClothAssetEditorTools/Private/ChaosClothAsset/ClothTrainingToolProperties.cpp b/Engine/Plugins/Experimental/ChaosClothAssetEditor/Source/ChaosClothAssetEditorTools/Private/ChaosClothAsset/ClothTrainingToolProperties.cpp
--- a/Engine/Plugins/Experimental/ChaosClothAssetEditor/Source/ChaosClothAssetEditorTools/Private/ChaosClothAsset/ClothTrainingToolProperties.cpp
+++ b/Engine/Plugins/Experimental/ChaosClothAssetEditor/Source/ChaosClothAssetEditorTools/Private/ChaosClothAsset/ClothTrainingToolProperties.cpp
@@ -79,6 +79,11 @@ UObject* UClothTrainingToolProperties::GetClothAsset()
 	return nullptr;
 }
 
+void UClothTrainingToolActionProperties::Initialize(UClothTrainingTool* ParentToolIn)
+{
+	ParentTool = ParentToolIn;
+}
+
 void UClothTrainingToolActionProperties::StartGenerating()
 {
 	if (ParentTool.IsValid())
diff --git a/Engine/Plugins/Experimental/ChaosClothAssetEditor/Source/ChaosClothAssetEditorTools/Private/ChaosClothAsset/ClothTrainingToolProperties.h b/Engine/Plugins/Experimental/ChaosClothAssetEditor/Source/ChaosClothAssetEditorTools/Private/ChaosClothAsset/ClothTrainingToolProperties.h
--- a/Engine/Plugins/Experimental/ChaosClothAssetEditor/Source/ChaosClothAssetEditorTools/Private/ChaosClothAsset/ClothTrainingToolProperties.h
+++ b/Engine/Plugins/Experimental/ChaosClothAssetEditor/Source/ChaosClothAssetEditorTools/Private/ChaosClothAsset/ClothTrainingToolProperties.h
@@ -66,7 +66,7 @@ public:
 
 	TWeakObjectPtr<UClothTrainingTool> ParentTool;
 
-	void Initialize(UClothTrainingTool* ParentToolIn) { ParentTool = ParentToolIn; }
+	void Initialize(UClothTrainingTool* ParentToolIn);
 
 	UFUNCTION(CallInEditor, Category = Actions, meta = (DisplayName = "Begin Generating", DisplayPriority = 1))
 	void StartGenerating();
