Chaos Cloth Editor - Fixed a crash in the cloth editor toolkit when the toolkit context gets clobbered in the middle of a Dataflow evaluation.

#rnx

diff --git a/Engine/Plugins/Experimental/ChaosClothAssetEditor/Source/ChaosClothAssetEditor/Private/ChaosClothAsset/ClothEditorToolkit.cpp b/Engine/Plugins/Experimental/ChaosClothAssetEditor/Source/ChaosClothAssetEditor/Private/ChaosClothAsset/ClothEditorToolkit.cpp
--- a/Engine/Plugins/Experimental/ChaosClothAssetEditor/Source/ChaosClothAssetEditor/Private/ChaosClothAsset/ClothEditorToolkit.cpp
+++ b/Engine/Plugins/Experimental/ChaosClothAssetEditor/Source/ChaosClothAssetEditor/Private/ChaosClothAsset/ClothEditorToolkit.cpp
@@ -242,7 +242,9 @@ void FChaosClothAssetEditorToolkit::Tick(float DeltaTime)
 		DataflowTerminalPath = Private::GetDataflowTerminalFrom(ClothAsset);
 
 		Dataflow::FTimestamp OldTimestamp = LastDataflowNodeTimestamp;
-		FDataflowEditorCommands::EvaluateTerminalNode(*DataflowContext.Get(), LastDataflowNodeTimestamp, Dataflow, nullptr, nullptr, ClothAsset, DataflowTerminalPath);
+
+		const TSharedPtr<Dataflow::FEngineContext> EvaluationContext(DataflowContext);  // Copy the context pointer as to not lose the reference during an evaluation (some UI operations can reset the toolkit context mid evaluation)
+		FDataflowEditorCommands::EvaluateTerminalNode(*EvaluationContext, LastDataflowNodeTimestamp, Dataflow, nullptr, nullptr, ClothAsset, DataflowTerminalPath);
 
 		if (OldTimestamp.Value < LastDataflowNodeTimestamp.Value)
 		{
