ClothTrainingTool: fix debug frame

diff --git a/Engine/Plugins/Experimental/ChaosClothAssetEditor/Source/ChaosClothAssetEditorTools/Private/ChaosClothAsset/ClothTrainingTool.cpp b/Engine/Plugins/Experimental/ChaosClothAssetEditor/Source/ChaosClothAssetEditorTools/Private/ChaosClothAsset/ClothTrainingTool.cpp
--- a/Engine/Plugins/Experimental/ChaosClothAssetEditor/Source/ChaosClothAssetEditorTools/Private/ChaosClothAsset/ClothTrainingTool.cpp
+++ b/Engine/Plugins/Experimental/ChaosClothAssetEditor/Source/ChaosClothAssetEditorTools/Private/ChaosClothAsset/ClothTrainingTool.cpp
@@ -608,11 +608,13 @@ void UClothTrainingTool::FLaunchSimsOp::CalculateResult(FProgressCancel* Progres
 	using UE::ClothTrainingTool::Private::ParseFrames;
 	using UE::ClothTrainingTool::Private::Range;
 
-	const TArray<int32> FramesToSimulate = ToolProperties->FramesToSimulate.Len() > 0
-		? ParseFrames(ToolProperties->FramesToSimulate) 
-		: Range(ToolProperties->AnimationSequence->GetNumberOfSampledKeys());
+	const TArray<int32> FramesToSimulate = ToolProperties->bDebug 
+		? TArray<int32>{ (int32)ToolProperties->DebugFrame } 
+		: (ToolProperties->FramesToSimulate.Len() > 0
+			? ParseFrames(ToolProperties->FramesToSimulate) 
+			: Range(ToolProperties->AnimationSequence->GetNumberOfSampledKeys()));
 
-	const int32 NumFrames = ToolProperties->bDebug ? 1 : FramesToSimulate.Num();
+	const int32 NumFrames = FramesToSimulate.Num();
 	if (NumFrames == 0)
 	{
 		return;
@@ -626,13 +628,8 @@ void UClothTrainingTool::FLaunchSimsOp::CalculateResult(FProgressCancel* Progres
 
 	for (int32 Frame = 0; Frame < NumFrames; Frame++)
 	{
-		if (Progress)
+		if (Progress && !Progress->Cancelled())
 		{
-			if (Progress->Cancelled())
-			{
-				break;
-			}
-
 			const int32 ThreadIdx = Frame % NumThreads;
 			const int32 AnimFrame = FramesToSimulate[Frame];
 
@@ -643,6 +640,10 @@ void UClothTrainingTool::FLaunchSimsOp::CalculateResult(FProgressCancel* Progres
 				Simulate(SimResource, AnimFrame, Frame, Progress, ProgressStep);
 			});
 		}
+		else
+		{
+			break;
+		}
 	}
 
 	for (FSimResource& SimResource : SimResources)
@@ -670,14 +671,8 @@ void UClothTrainingTool::FLaunchSimsOp::Simulate(FSimResource &SimResource, int3
 	bool bCancelled = false;
 	for (int32 Step = 0; Step < NumSteps; ++Step)
 	{
-		if (Progress)
+		if (Progress && !Progress->Cancelled())
 		{
-			if (Progress->Cancelled())
-			{
-				bCancelled = true;
-				break;
-			}
-
 			DataGenerationProxy.Tick();
 
 			// Clear any reset flags at the end of the first step
