ClothCollection Schema rework. Existing cloth assets will be emptied, but their associated Dataflow graphs and nodes are still valid (with the exception of AddWeightMapNodes which need to be repainted

diff --git a/Engine/Plugins/Experimental/ChaosClothAsset/Source/ChaosClothAssetEngine/Private/ChaosClothAsset/ClothSimulationModel.cpp b/Engine/Plugins/Experimental/ChaosClothAsset/Source/ChaosClothAssetEngine/Private/ChaosClothAsset/ClothSimulationModel.cpp
--- a/Engine/Plugins/Experimental/ChaosClothAsset/Source/ChaosClothAssetEngine/Private/ChaosClothAsset/ClothSimulationModel.cpp
+++ b/Engine/Plugins/Experimental/ChaosClothAsset/Source/ChaosClothAssetEngine/Private/ChaosClothAsset/ClothSimulationModel.cpp
@@ -122,76 +122,66 @@ bool FChaosClothSimulationLodModel::Serialize(FArchive& Ar)
 	return true;
 }
 
-FChaosClothSimulationModel::FChaosClothSimulationModel(const TSharedPtr<const FManagedArrayCollection>& ClothCollection, const FReferenceSkeleton& ReferenceSkeleton)
+FChaosClothSimulationModel::FChaosClothSimulationModel(const TArray<TSharedPtr<const FManagedArrayCollection>>& ClothCollections, const FReferenceSkeleton& ReferenceSkeleton)
 {
 	using namespace UE::Chaos::ClothAsset;
 
-	const FCollectionClothConstFacade Cloth(ClothCollection);
-
-	// Retrieve weigh map names
-	const TArray<FName> WeightMapNames = Cloth.GetWeightMapNames();
 
 	// Initialize LOD models
-	const int32 NumLods = Cloth.GetNumLods();
+	const int32 NumLods = ClothCollections.Num();
 	ClothSimulationLodModels.SetNum(NumLods);
 
 	for (int32 LodIndex = 0; LodIndex < NumLods; ++LodIndex)
 	{
+		const FCollectionClothConstFacade Cloth(ClothCollections[LodIndex]);
+
+		// Retrieve weigh map names
+		const TArray<FName> WeightMapNames = Cloth.GetWeightMapNames();
+
 		FChaosClothSimulationLodModel& LodModel = ClothSimulationLodModels[LodIndex];
-		const FCollectionClothLodConstFacade ClothLod = Cloth.GetLod(LodIndex);
 
 		TArray<TArray<int32>> WeldedToPatternIndices;
-		ClothLod.BuildSimulationMesh(LodModel.Positions, LodModel.Normals, LodModel.Indices, LodModel.PatternPositions, LodModel.PatternIndices, LodModel.PatternToWeldedIndices, &WeldedToPatternIndices);
+		Cloth.BuildSimulationMesh(LodModel.Positions, LodModel.Normals, LodModel.Indices, LodModel.PatternPositions, LodModel.PatternIndices, LodModel.PatternToWeldedIndices, &WeldedToPatternIndices);
 
-		const int32 LodModelNumSimVertices = LodModel.Positions.Num();
-		const int32 ClothLodNumSimVertices = ClothLod.GetNumSimVertices();
-
-		// Copy and weld (average) weight maps
+		// Copy weight maps
 		LodModel.WeightMaps.Reserve(WeightMapNames.Num());
 
 		for (const FName& WeightMapName : WeightMapNames)
 		{
-			const TConstArrayView<float> ClothLodWeightMap = ClothLod.GetWeightMap(WeightMapName);
-			TArray<float>& LodModelWeightMap = LodModel.WeightMaps.Add(WeightMapName);
-
-			LodModelWeightMap = FClothGeometryTools::BuildWeldedWeightMapForLod(ClothCollection, LodIndex, WeightMapName, WeldedToPatternIndices);
+			LodModel.WeightMaps.Add(WeightMapName) = Cloth.GetWeightMap(WeightMapName);
 		}
 
-		// Weld bone influences
-		TConstArrayView<int32> NumBoneInfluences = ClothLod.GetSimNumBoneInfluences();
-		TConstArrayView<TArray<int32>> SimBoneIndices = ClothLod.GetSimBoneIndices();
-		TConstArrayView<TArray<float>> SimBoneWeights = ClothLod.GetSimBoneWeights();
-		LodModel.BoneData.SetNum(LodModelNumSimVertices);
-
-		for (int32 WeldedIndex = 0; WeldedIndex < LodModelNumSimVertices; ++WeldedIndex)
+		// Copy bone influences and gather tether data
+		const int32 NumSimVertices3D = Cloth.GetNumSimVertices3D();
+		TConstArrayView<TArray<int32>> SimBoneIndices = Cloth.GetSimBoneIndices();
+		TConstArrayView<TArray<float>> SimBoneWeights = Cloth.GetSimBoneWeights();
+		TConstArrayView<TArray<int32>> TetherKinematicIndex = Cloth.GetTetherKinematicIndex();
+		TConstArrayView<TArray<float>> TetherReferenceLength = Cloth.GetTetherReferenceLength();
+		TArray<TArray<TPair<float, int32>>> MergedTetherData;
+		LodModel.BoneData.SetNum(NumSimVertices3D);
+		MergedTetherData.SetNum(NumSimVertices3D);
+
+		for (int32 VertexIndex = 0; VertexIndex < NumSimVertices3D; ++VertexIndex)
 		{
-			const TArray<int32>& SourceVertices = WeldedToPatternIndices[WeldedIndex];
-			check(SourceVertices.Num());
-			const int32 SourceVertexIndex = SourceVertices[0]; // Just using first bone's data.
-
-			LodModel.BoneData[WeldedIndex].NumInfluences = NumBoneInfluences[SourceVertexIndex];
-			for (int32 BoneIndex = 0; BoneIndex < LodModel.BoneData[WeldedIndex].NumInfluences; ++BoneIndex)
+			check(SimBoneIndices[VertexIndex].Num() == SimBoneWeights[VertexIndex].Num());
+			check(SimBoneIndices[VertexIndex].Num() <= FClothVertBoneData::MaxTotalInfluences);
+			LodModel.BoneData[VertexIndex].NumInfluences = SimBoneIndices[VertexIndex].Num();
+			for (int32 BoneIndex = 0; BoneIndex < LodModel.BoneData[VertexIndex].NumInfluences; ++BoneIndex)
 			{
-				LodModel.BoneData[WeldedIndex].BoneIndices[BoneIndex] = SimBoneIndices[SourceVertexIndex][BoneIndex];
-				LodModel.BoneData[WeldedIndex].BoneWeights[BoneIndex] = SimBoneWeights[SourceVertexIndex][BoneIndex];
+				LodModel.BoneData[VertexIndex].BoneIndices[BoneIndex] = SimBoneIndices[VertexIndex][BoneIndex];
+				LodModel.BoneData[VertexIndex].BoneWeights[BoneIndex] = SimBoneWeights[VertexIndex][BoneIndex];
 			}
-		}
 
-		// Weld tethers. This will not deduplicate any tethers that may be duplicated by welding. 
-		// This would currently only occur if seams were added after tethers were created in the dataflow graph.
-		const int32 NumTetherBatches = ClothLod.GetNumTetherBatches();
-		LodModel.TetherData.Tethers.Reserve(NumTetherBatches);
-		for (int32 TetherBatchIndex = 0; TetherBatchIndex < NumTetherBatches; ++TetherBatchIndex)
-		{
-			FCollectionClothTetherBatchConstFacade TetherBatchFacade = ClothLod.GetTetherBatch(TetherBatchIndex);
-			TArray<TTuple<int32, int32, float>>& Tethers = LodModel.TetherData.Tethers.Emplace_GetRef(TetherBatchFacade.GetZippedTetherData());
-
-			for (TTuple<int32, int32, float>& Tether : Tethers)
+			check(TetherKinematicIndex[VertexIndex].Num() == TetherReferenceLength[VertexIndex].Num());
+			MergedTetherData[VertexIndex].Reserve(TetherKinematicIndex[VertexIndex].Num());
+			for (int32 TetherIdx = 0; TetherIdx < TetherKinematicIndex[VertexIndex].Num(); ++TetherIdx)
 			{
-				Tether.Get<0>() = LodModel.PatternToWeldedIndices[Tether.Get<0>()];
-				Tether.Get<1>() = LodModel.PatternToWeldedIndices[Tether.Get<1>()];
+				MergedTetherData[VertexIndex].Emplace(TetherReferenceLength[VertexIndex][TetherIdx], TetherKinematicIndex[VertexIndex][TetherIdx]);
 			}
 		}
+
+		// Batch tethers
+		LodModel.TetherData.GenerateTethers(MoveTemp(MergedTetherData));
 	}
 
 	// Populate used bone names and indices
diff --git a/Engine/Plugins/Experimental/ChaosClothAsset/Source/ChaosClothAssetEngine/Private/ChaosClothAsset/ClothSimulationModel.h b/Engine/Plugins/Experimental/ChaosClothAsset/Source/ChaosClothAssetEngine/Private/ChaosClothAsset/ClothSimulationModel.h
--- a/Engine/Plugins/Experimental/ChaosClothAsset/Source/ChaosClothAssetEngine/Private/ChaosClothAsset/ClothSimulationModel.h
+++ b/Engine/Plugins/Experimental/ChaosClothAsset/Source/ChaosClothAssetEngine/Private/ChaosClothAsset/ClothSimulationModel.h
@@ -96,7 +96,7 @@ struct FChaosClothSimulationModel
 	int32 ReferenceBoneIndex = INDEX_NONE;
 
 	FChaosClothSimulationModel() = default;
-	FChaosClothSimulationModel(const TSharedPtr<const FManagedArrayCollection>& ClothCollection, const FReferenceSkeleton& ReferenceSkeleton);
+	FChaosClothSimulationModel(const TArray<TSharedPtr<const FManagedArrayCollection>>& ClothCollections, const FReferenceSkeleton& ReferenceSkeleton);
 
 	int32 GetNumLods() const { return ClothSimulationLodModels.Num(); }
 
diff --git a/Engine/Plugins/Experimental/ChaosClothAsset/Source/ChaosClothAssetEngine/Public/ChaosClothAsset/ClothAsset.h b/Engine/Plugins/Experimental/ChaosClothAsset/Source/ChaosClothAssetEngine/Public/ChaosClothAsset/ClothAsset.h
--- a/Engine/Plugins/Experimental/ChaosClothAsset/Source/ChaosClothAssetEngine/Public/ChaosClothAsset/ClothAsset.h
+++ b/Engine/Plugins/Experimental/ChaosClothAsset/Source/ChaosClothAssetEngine/Public/ChaosClothAsset/ClothAsset.h
@@ -98,10 +98,10 @@ public:
 	//~ End USkinnedAsset interface
 
 	/** Return the enclosed Cloth Collection object. */
-	TSharedPtr<FManagedArrayCollection>& GetClothCollection() { return ClothCollection; }
+	TArray<TSharedPtr<FManagedArrayCollection>>& GetClothCollections() { return ClothCollections; }
 
 	/** Return the enclosed Cloth Collection object, const version. */
-	TSharedPtr<const FManagedArrayCollection> GetClothCollection() const { return ClothCollection; }
+	const TArray<TSharedPtr<const FManagedArrayCollection>>& GetClothCollections() const { return reinterpret_cast<const TArray<TSharedPtr<const FManagedArrayCollection>>&>(ClothCollections); }
 
 	/** Return the cloth simulation ready LOD model data. */
 	TSharedPtr<const FChaosClothSimulationModel> GetClothSimulationModel() const { return ClothSimulationModel; }
@@ -274,8 +274,8 @@ private:
 	/** Mesh-space ref pose, where parent matrices are applied to ref pose matrices. */
 	TArray<FMatrix> CachedComposedRefPoseMatrices;
 
-	/** Cloth Collection containing this asset data. */
-	TSharedPtr<FManagedArrayCollection> ClothCollection;
+	/** Cloth Collection containing this asset data. One per LOD. */
+	TArray<TSharedPtr<FManagedArrayCollection>> ClothCollections;
 
 	/** Reference skeleton created from the provided skeleton asset. */
 	FReferenceSkeleton RefSkeleton;
diff --git a/Engine/Plugins/Experimental/ChaosClothAsset/Source/ChaosClothAssetEngine/Public/ChaosClothAsset/ClothEngineTools.h b/Engine/Plugins/Experimental/ChaosClothAsset/Source/ChaosClothAssetEngine/Public/ChaosClothAsset/ClothEngineTools.h
--- a/Engine/Plugins/Experimental/ChaosClothAsset/Source/ChaosClothAssetEngine/Public/ChaosClothAsset/ClothEngineTools.h
+++ b/Engine/Plugins/Experimental/ChaosClothAsset/Source/ChaosClothAssetEngine/Public/ChaosClothAsset/ClothEngineTools.h
@@ -6,6 +6,7 @@
 
 struct FManagedArrayCollection;
 class FName;
+class FSkeletalMeshLODModel;
 
 namespace UE::Chaos::ClothAsset
 {
diff --git a/Engine/Plugins/Experimental/ChaosClothAssetEditor/Source/ChaosClothAssetDataflowNodes/ChaosClothAssetDataflowNodes.Build.cs b/Engine/Plugins/Experimental/ChaosClothAssetEditor/Source/ChaosClothAssetDataflowNodes/ChaosClothAssetDataflowNodes.Build.cs
--- a/Engine/Plugins/Experimental/ChaosClothAssetEditor/Source/ChaosClothAssetDataflowNodes/ChaosClothAssetDataflowNodes.Build.cs
+++ b/Engine/Plugins/Experimental/ChaosClothAssetEditor/Source/ChaosClothAssetDataflowNodes/ChaosClothAssetDataflowNodes.Build.cs
@@ -35,6 +35,8 @@ public class ChaosClothAssetDataflowNodes : ModuleRules
 				"InputCore",
 				"MeshConversion",
 				"MeshDescription",
+				"MeshUtilitiesCommon",
+				"RenderCore",
 				"SkeletalMeshDescription",
 				"Slate",
 				"SlateCore",
diff --git a/Engine/Plugins/Experimental/ChaosClothAssetEditor/Source/ChaosClothAssetDataflowNodes/Private/ChaosClothAsset/AddWeightMapNode.cpp b/Engine/Plugins/Experimental/ChaosClothAssetEditor/Source/ChaosClothAssetDataflowNodes/Private/ChaosClothAsset/AddWeightMapNode.cpp
--- a/Engine/Plugins/Experimental/ChaosClothAssetEditor/Source/ChaosClothAssetDataflowNodes/Private/ChaosClothAsset/AddWeightMapNode.cpp
+++ b/Engine/Plugins/Experimental/ChaosClothAssetEditor/Source/ChaosClothAssetDataflowNodes/Private/ChaosClothAsset/AddWeightMapNode.cpp
@@ -3,6 +3,7 @@
 #include "ChaosClothAsset/AddWeightMapNode.h"
 #include "ChaosClothAsset/DataflowNodes.h"
 #include "ChaosClothAsset/CollectionClothFacade.h"
+#include "ChaosClothAsset/ClothGeometryTools.h"
 
 #include UE_INLINE_GENERATED_CPP_BY_NAME(AddWeightMapNode)
 
@@ -29,31 +30,22 @@ void FChaosClothAssetAddWeightMapNode::Evaluate(Dataflow::FContext& Context, con
 		FCollectionClothFacade ClothFacade(ClothCollection);
 
 		const FString InNameString = GetValue<FString>(Context, &Name);
+		const FName InName(InNameString);
+		ClothFacade.AddWeightMap(InName);		// Does nothing if weight map already exists
 
-		// Note: If the input collection has no valid LODs, just pass it to the output as is (this can happen if the node is asked to Evaluate before the inputs are connected)
-
-		if (ClothFacade.GetNumLods() > 0)
+		TArrayView<float> ClothWeights = ClothFacade.GetWeightMap(InName);
+		const int32 MaxWeightIndex = FMath::Min(VertexWeights.Num(), ClothWeights.Num());
+		if (VertexWeights.Num() != ClothWeights.Num())
 		{
-			const FName InName(InNameString);
-			ClothFacade.AddWeightMap(InName);		// Does nothing if weight map already exists
-
-			FCollectionClothLodFacade ClothLod = ClothFacade.GetLod(0);
-			TArrayView<float> LodWeights = ClothLod.GetWeightMap(InName);
-
-			const int32 MaxWeightIndex = FMath::Min(VertexWeights.Num(), LodWeights.Num());
+			DataflowNodes::LogAndToastWarning(
+				FText::Format(LOCTEXT("WeightMapSize", "FChaosClothAssetAddWeightMapNode: Vertex count mismatch: vertex weights in the node: {0}; 3D vertices in cloth: {1}"),
+					VertexWeights.Num(),
+					ClothWeights.Num()));
+		}
 
-			if (VertexWeights.Num() != LodWeights.Num())
-			{
-				DataflowNodes::LogAndToastWarning(
-					FText::Format(LOCTEXT("WeightMapSize", "FChaosClothAssetAddWeightMapNode: Vertex count mismatch: vertex weights in the node: {0}; vertices in the specified LOD: {1}"),
-						VertexWeights.Num(),
-						LodWeights.Num()));
-			}
-				
-			for (int32 VertexID = 0; VertexID < MaxWeightIndex; ++VertexID)
-			{
-				LodWeights[VertexID] = VertexWeights[VertexID];
-			}
+		for (int32 VertexID = 0; VertexID < MaxWeightIndex; ++VertexID)
+		{
+			ClothWeights[VertexID] = VertexWeights[VertexID];
 		}
 
 		SetValue<FManagedArrayCollection>(Context, *ClothCollection, &Collection);
diff --git a/Engine/Plugins/Experimental/ChaosClothAssetEditor/Source/ChaosClothAssetDataflowNodes/Private/ChaosClothAsset/ClothDataflowTools.h b/Engine/Plugins/Experimental/ChaosClothAssetEditor/Source/ChaosClothAssetDataflowNodes/Private/ChaosClothAsset/ClothDataflowTools.h
new file mode 100644
--- /dev/null
+++ b/Engine/Plugins/Experimental/ChaosClothAssetEditor/Source/ChaosClothAssetDataflowNodes/Private/ChaosClothAsset/ClothDataflowTools.h
@@ -0,0 +1,20 @@
+// Copyright Epic Games, Inc. All Rights Reserved.
+
+#pragma once
+
+#include "Templates/SharedPointer.h"
+
+struct FManagedArrayCollection;
+class FSkeletalMeshLODModel;
+class FString;
+
+namespace UE::Chaos::ClothAsset
+{
+	/**
+	* Tools shared by cloth data flow nodes
+	*/
+	struct FClothDataflowTools
+	{
+		static void AddRenderPatternFromSkeletalMeshSection(const TSharedPtr<FManagedArrayCollection>& ClothCollection, const FSkeletalMeshLODModel& SkeletalMeshModel, const int32 SectionIndex, const FString& RenderMaterialPathName);
+	};
+}  // End namespace UE::Chaos::ClothAsset
diff --git a/Engine/Plugins/Experimental/ChaosClothAssetEditor/Source/ChaosClothAssetDataflowNodes/Private/ChaosClothAsset/CopySimulationToRenderMeshNode.cpp b/Engine/Plugins/Experimental/ChaosClothAssetEditor/Source/ChaosClothAssetDataflowNodes/Private/ChaosClothAsset/CopySimulationToRenderMeshNode.cpp
--- a/Engine/Plugins/Experimental/ChaosClothAssetEditor/Source/ChaosClothAssetDataflowNodes/Private/ChaosClothAsset/CopySimulationToRenderMeshNode.cpp
+++ b/Engine/Plugins/Experimental/ChaosClothAssetEditor/Source/ChaosClothAssetDataflowNodes/Private/ChaosClothAsset/CopySimulationToRenderMeshNode.cpp
@@ -36,7 +36,7 @@ void FChaosClothAssetCopySimulationToRenderMeshNode::Evaluate(Dataflow::FContext
 			FString(TEXT("/Engine/EditorMaterials/Cloth/CameraLitDoubleSided.CameraLitDoubleSided"));
 
 		// Copy the mesh data
-		FClothGeometryTools::CopySimMeshToRenderMesh(ClothCollection, MaterialPathName);
+		FClothGeometryTools::CopySimMeshToRenderMesh(ClothCollection, MaterialPathName, bGenerateSingleRenderPattern);
 
 		SetValue<FManagedArrayCollection>(Context, *ClothCollection, &Collection);
 	}
diff --git a/Engine/Plugins/Experimental/ChaosClothAssetEditor/Source/ChaosClothAssetDataflowNodes/Private/ChaosClothAsset/DatasmithImportNode.cpp b/Engine/Plugins/Experimental/ChaosClothAssetEditor/Source/ChaosClothAssetDataflowNodes/Private/ChaosClothAsset/DatasmithImportNode.cpp
--- a/Engine/Plugins/Experimental/ChaosClothAssetEditor/Source/ChaosClothAssetDataflowNodes/Private/ChaosClothAsset/DatasmithImportNode.cpp
+++ b/Engine/Plugins/Experimental/ChaosClothAssetEditor/Source/ChaosClothAssetDataflowNodes/Private/ChaosClothAsset/DatasmithImportNode.cpp
@@ -78,8 +78,11 @@ bool FChaosClothAssetDatasmithImportNode::EvaluateImpl(Dataflow::FContext& Conte
 		UChaosClothAsset* const DatasmithClothAsset = Cast<UChaosClothAsset>(ClothObject);
 		if (ensure(DatasmithClothAsset))
 		{
-			DatasmithClothAsset->GetClothCollection()->CopyTo(&OutCollection);
-			return true;
+			if (DatasmithClothAsset->GetClothCollections().Num() > 0 && DatasmithClothAsset->GetClothCollections()[0].IsValid())
+			{
+				DatasmithClothAsset->GetClothCollections()[0]->CopyTo(&OutCollection);
+				return true;
+			}
 		}
 		DatasmithClothAsset->ClearFlags(RF_Standalone);
 	}
@@ -104,7 +107,6 @@ void FChaosClothAssetDatasmithImportNode::Evaluate(Dataflow::FContext& Context,
 		const TSharedRef<FManagedArrayCollection> ClothCollection = MakeShared<FManagedArrayCollection>();
 		FCollectionClothFacade CollectionClothFacade(ClothCollection);
 		CollectionClothFacade.DefineSchema();
-		CollectionClothFacade.AddLod();
 
 		SetValue<FManagedArrayCollection>(Context, *ClothCollection, &Collection);
 	}
diff --git a/Engine/Plugins/Experimental/ChaosClothAssetEditor/Source/ChaosClothAssetDataflowNodes/Private/ChaosClothAsset/ImportNode.cpp b/Engine/Plugins/Experimental/ChaosClothAssetEditor/Source/ChaosClothAssetDataflowNodes/Private/ChaosClothAsset/ImportNode.cpp
--- a/Engine/Plugins/Experimental/ChaosClothAssetEditor/Source/ChaosClothAssetDataflowNodes/Private/ChaosClothAsset/ImportNode.cpp
+++ b/Engine/Plugins/Experimental/ChaosClothAssetEditor/Source/ChaosClothAssetDataflowNodes/Private/ChaosClothAsset/ImportNode.cpp
@@ -58,23 +58,17 @@ void FChaosClothAssetImportNode::Evaluate(Dataflow::FContext& Context, const FDa
 		const TSharedRef<FManagedArrayCollection> ClothCollection = MakeShared<FManagedArrayCollection>();
 		FCollectionClothFacade ClothFacade(ClothCollection);
 		ClothFacade.DefineSchema();
-		FCollectionClothLodFacade ClothLodFacade = ClothFacade.AddGetLod();  // Always output a single LOD at LOD 0, even if empty
 
 		// Copy the cloth asset to this node's output collection
 		if (bSetValue)
 		{
 			const TObjectPtr<const UChaosClothAsset>& InClothAsset = GetValue<TObjectPtr<const UChaosClothAsset>>(Context, &ClothAsset);
-			if (InClothAsset->GetClothCollection())
+			const int32& InImportLod = GetValue<int32>(Context, &ImportLod);
+			const TArray<TSharedPtr<const FManagedArrayCollection>>& InClothCollections = InClothAsset->GetClothCollections();
+			if (InImportLod >= 0 && InClothCollections.Num() > InImportLod && InClothCollections[InImportLod].IsValid())
 			{
-				const int32& InImportLod = GetValue<int32>(Context, &ImportLod);
-				const FCollectionClothConstFacade InClothFacade(InClothAsset->GetClothCollection());
-
-				if (InImportLod >= 0 && InImportLod < InClothFacade.GetNumLods())
-				{
-					// Copy specified LOD from the input asset to the LOD 0 of the output collection
-					const FCollectionClothLodConstFacade InClothLodFacade = InClothFacade.GetLod(InImportLod);
-					ClothLodFacade.Initialize(InClothLodFacade);
-				}
+				const FCollectionClothConstFacade InClothFacade(InClothCollections[InImportLod]);
+				ClothFacade.Initialize(InClothFacade);
 			}
 		}
 		SetValue<FManagedArrayCollection>(Context, *ClothCollection, &Collection);
diff --git a/Engine/Plugins/Experimental/ChaosClothAssetEditor/Source/ChaosClothAssetDataflowNodes/Private/ChaosClothAsset/MergeClothCollectionsNode.cpp b/Engine/Plugins/Experimental/ChaosClothAssetEditor/Source/ChaosClothAssetDataflowNodes/Private/ChaosClothAsset/MergeClothCollectionsNode.cpp
--- a/Engine/Plugins/Experimental/ChaosClothAssetEditor/Source/ChaosClothAssetDataflowNodes/Private/ChaosClothAsset/MergeClothCollectionsNode.cpp
+++ b/Engine/Plugins/Experimental/ChaosClothAssetEditor/Source/ChaosClothAssetDataflowNodes/Private/ChaosClothAsset/MergeClothCollectionsNode.cpp
@@ -28,14 +28,6 @@ void FChaosClothAssetMergeClothCollectionsNode::Evaluate(Dataflow::FContext& Con
 		// Make it a valid cloth collection if needed
 		FCollectionClothFacade ClothFacade(ClothCollection);
 		ClothFacade.DefineSchema();
-
-		if (ClothFacade.GetNumLods() < 1)
-		{
-			// Make sure that whatever happens there is always at least one empty LOD to avoid crashing the render data
-			ClothFacade.AddLod();
-		}
-		FCollectionClothLodFacade ClothLodFacade = ClothFacade.GetLod(0);
-
 		// Iterate through the inputs and append them to LOD 0
 		const TArray<const FManagedArrayCollection*> Collections = GetCollections();
 		for (int32 InputIndex = 1; InputIndex < Collections.Num(); ++InputIndex)
@@ -43,11 +35,7 @@ void FChaosClothAssetMergeClothCollectionsNode::Evaluate(Dataflow::FContext& Con
 			const FManagedArrayCollection& OtherCollection = GetValue<FManagedArrayCollection>(Context, Collections[InputIndex]);
 			const TSharedRef<const FManagedArrayCollection> OtherClothCollection = MakeShared<const FManagedArrayCollection>(OtherCollection);
 			const FCollectionClothConstFacade OtherClothFacade(OtherClothCollection);  // TODO: Remove Shared pointer from facade and skip duplication
-
-			if (OtherClothFacade.GetNumLods() > 0)
-			{
-				ClothLodFacade.Append(OtherClothFacade.GetLod(0));
-			}
+			ClothFacade.Append(OtherClothFacade);
 		}
 
 		SetValue<FManagedArrayCollection>(Context, *ClothCollection, &Collection);
diff --git a/Engine/Plugins/Experimental/ChaosClothAssetEditor/Source/ChaosClothAssetDataflowNodes/Private/ChaosClothAsset/ReverseNormalsNode.cpp b/Engine/Plugins/Experimental/ChaosClothAssetEditor/Source/ChaosClothAssetDataflowNodes/Private/ChaosClothAsset/ReverseNormalsNode.cpp
--- a/Engine/Plugins/Experimental/ChaosClothAssetEditor/Source/ChaosClothAssetDataflowNodes/Private/ChaosClothAsset/ReverseNormalsNode.cpp
+++ b/Engine/Plugins/Experimental/ChaosClothAssetEditor/Source/ChaosClothAssetDataflowNodes/Private/ChaosClothAsset/ReverseNormalsNode.cpp
@@ -12,7 +12,8 @@ FChaosClothAssetReverseNormalsNode::FChaosClothAssetReverseNormalsNode(const Dat
 	: FDataflowNode(InParam, InGuid)
 {
 	RegisterInputConnection(&Collection);
-	RegisterInputConnection(&Patterns);
+	RegisterInputConnection(&SimPatterns);
+	RegisterInputConnection(&RenderPatterns);
 	RegisterOutputConnection(&Collection, &Collection);
 }
 
@@ -32,7 +33,8 @@ void FChaosClothAssetReverseNormalsNode::Evaluate(Dataflow::FContext& Context, c
 			bReverseSimMeshWindingOrder,
 			bReverseRenderMeshNormals,
 			bReverseRenderMeshWindingOrder,
-			Patterns);
+			GetValue<TArray<int32>>(Context, &SimPatterns),
+			GetValue<TArray<int32>>(Context, &RenderPatterns));
 
 		SetValue<FManagedArrayCollection>(Context, *ClothCollection, &Collection);
 	}
diff --git a/Engine/Plugins/Experimental/ChaosClothAssetEditor/Source/ChaosClothAssetDataflowNodes/Private/ChaosClothAsset/SetPhysicsAssetNode.cpp b/Engine/Plugins/Experimental/ChaosClothAssetEditor/Source/ChaosClothAssetDataflowNodes/Private/ChaosClothAsset/SetPhysicsAssetNode.cpp
--- a/Engine/Plugins/Experimental/ChaosClothAssetEditor/Source/ChaosClothAssetDataflowNodes/Private/ChaosClothAsset/SetPhysicsAssetNode.cpp
+++ b/Engine/Plugins/Experimental/ChaosClothAssetEditor/Source/ChaosClothAssetDataflowNodes/Private/ChaosClothAsset/SetPhysicsAssetNode.cpp
@@ -28,11 +28,7 @@ void FChaosClothAssetSetPhysicsAssetNode::Evaluate(Dataflow::FContext& Context,
 		const TSharedRef<FManagedArrayCollection> ClothCollection = MakeShared<FManagedArrayCollection>(MoveTemp(InCollection));
 
 		FCollectionClothFacade CollectionClothFacade(ClothCollection);
-		if (CollectionClothFacade.GetNumLods() > 0)
-		{
-			FCollectionClothLodFacade CollectionClothLodFacade = CollectionClothFacade.GetLod(0);
-			CollectionClothLodFacade.SetPhysicsAssetPathName(PhysicsAsset.GetPathName());
-		}
+		CollectionClothFacade.SetPhysicsAssetPathName(PhysicsAsset.GetPathName());
 
 		SetValue<FManagedArrayCollection>(Context, *ClothCollection, &Collection);
 	}
diff --git a/Engine/Plugins/Experimental/ChaosClothAssetEditor/Source/ChaosClothAssetDataflowNodes/Private/ChaosClothAsset/SkeletalMeshImportNode.cpp b/Engine/Plugins/Experimental/ChaosClothAssetEditor/Source/ChaosClothAssetDataflowNodes/Private/ChaosClothAsset/SkeletalMeshImportNode.cpp
--- a/Engine/Plugins/Experimental/ChaosClothAssetEditor/Source/ChaosClothAssetDataflowNodes/Private/ChaosClothAsset/SkeletalMeshImportNode.cpp
+++ b/Engine/Plugins/Experimental/ChaosClothAssetEditor/Source/ChaosClothAssetDataflowNodes/Private/ChaosClothAsset/SkeletalMeshImportNode.cpp
@@ -3,6 +3,7 @@
 #include "ChaosClothAsset/SkeletalMeshImportNode.h"
 #include "ChaosClothAsset/DataflowNodes.h"
 #include "ChaosClothAsset/CollectionClothFacade.h"
+#include "ChaosClothAsset/ClothDataflowTools.h"
 #include "Animation/Skeleton.h"
 #include "BoneWeights.h"
 #include "Engine/SkeletalMesh.h"
@@ -41,7 +42,7 @@ void FChaosClothAssetSkeletalMeshImportNode::Evaluate(Dataflow::FContext& Contex
 			}
 			const FSkeletalMeshLODModel &LODModel = ImportedModel->LODModels[LODIndex];
 
-			const bool bIsValidSection = SectionIndex < LODModel.Sections.Num();
+			const bool bIsValidSection = LODModel.Sections.IsValidIndex(SectionIndex);;
 			if (!ensureAlways(bIsValidSection))
 			{
 				SetValue<FManagedArrayCollection>(Context, *ClothCollection, &Collection);	
@@ -51,142 +52,27 @@ void FChaosClothAssetSkeletalMeshImportNode::Evaluate(Dataflow::FContext& Contex
 			const FSkelMeshSection &Section = LODModel.Sections[SectionIndex];
 			FCollectionClothFacade ClothFacade(ClothCollection);
 			ClothFacade.DefineSchema();
-			FCollectionClothLodFacade ClothLodFacade = ClothFacade.AddGetLod();
-			FCollectionClothPatternFacade ClothPatternFacade = ClothLodFacade.AddGetPattern();
-
-			using FBoneData = TTuple<int32, TArray<int32>, TArray<float>>;
-			auto ExtractBones = [](const FSkelMeshSection& Section, const FSoftSkinVertex& SoftVertex)
-			{
-				const int32 NumBones = Section.MaxBoneInfluences;
-				TArray<int32> BoneIndices;
-				TArray<float> BoneWeights;
-				BoneIndices.SetNum(NumBones);
-				BoneWeights.SetNum(NumBones);
-				for (int32 BoneIndex = 0; BoneIndex < NumBones; ++BoneIndex)
-				{
-					BoneIndices[BoneIndex] = (int32)Section.BoneMap[(int32)SoftVertex.InfluenceBones[BoneIndex]];
-					BoneWeights[BoneIndex] = (float)SoftVertex.InfluenceWeights[BoneIndex] * UE::AnimationCore::InvMaxRawBoneWeightFloat;
-				}
-				return FBoneData(NumBones, BoneIndices, BoneWeights);
-			};
-
-			auto GetSectionFace = [](const FSkelMeshSection& Section, const TArray<uint32>& IndexBuffer, int32 FaceIndex)
-			{
-				const int32 IndexOffset = Section.BaseIndex + FaceIndex * 3;
-				FIntVector3 TriangleIndices;
-				TriangleIndices.X = IndexBuffer[IndexOffset + 0] - Section.BaseVertexIndex;
-				TriangleIndices.Y = IndexBuffer[IndexOffset + 1] - Section.BaseVertexIndex;
-				TriangleIndices.Z = IndexBuffer[IndexOffset + 2] - Section.BaseVertexIndex;
-				return TriangleIndices;
-			};
 
 			if (bImportSimMesh)
 			{
-				ClothPatternFacade.SetNumSimVertices(Section.NumVertices);
-				ClothPatternFacade.SetNumSimFaces(Section.NumTriangles);
-
-				TArrayView<FVector2f> SimPosition = ClothPatternFacade.GetSimPosition();
-				TArrayView<FVector3f> SimRestPosition = ClothPatternFacade.GetSimRestPosition();
-				TArrayView<FVector3f> SimRestNormal = ClothPatternFacade.GetSimRestNormal();
-				TArrayView<int32> SimBoneInfluences = ClothPatternFacade.GetSimNumBoneInfluences();
-				TArrayView<TArray<int32>> SimBoneIndices = ClothPatternFacade.GetSimBoneIndices();
-				TArrayView<TArray<float>> SimBoneWeights = ClothPatternFacade.GetSimBoneWeights();
-				TArrayView<FIntVector3> SimIndices = ClothPatternFacade.GetSimIndices();
-
-				const bool bIsValidUV = UVChannel >= 0 && UVChannel < MAX_TEXCOORDS;
-				if (!bIsValidUV || UVChannel == INDEX_NONE)
-				{
-					UE_LOG(LogChaos, Error, TEXT("Invalid UV channel %d"), UVChannel);
-				}
-
-				for (int32 VertexIndex = 0; VertexIndex < Section.NumVertices; ++VertexIndex)
-				{
-					const FSoftSkinVertex& SoftVertex = Section.SoftVertices[VertexIndex];
-					SimPosition[VertexIndex] = bIsValidUV ? SoftVertex.UVs[UVChannel] : FVector2f(SoftVertex.Position[1], SoftVertex.Position[2]);
-					SimRestPosition[VertexIndex] = SoftVertex.Position;
-					SimRestNormal[VertexIndex] = -SoftVertex.TangentZ;
-
-					FBoneData BoneData = ExtractBones(Section, SoftVertex);
-					SimBoneInfluences[VertexIndex] = BoneData.Get<0>();
-					SimBoneIndices[VertexIndex] = MoveTemp(BoneData.Get<1>());
-					SimBoneWeights[VertexIndex] = MoveTemp(BoneData.Get<2>());
-				}
-
-				for (uint32 FaceIndex = 0; FaceIndex < Section.NumTriangles; ++FaceIndex)
-				{
-					SimIndices[FaceIndex] = GetSectionFace(Section, LODModel.IndexBuffer, FaceIndex);
-				}
-
-				ensureMsgf(false, TEXT("TODO: BuildSeams needs to be implemented"));
-				ClothLodFacade.SetNumTetherBatches(0);
-				ClothLodFacade.SetNumSeams(0);
+				ensureMsgf(false, TEXT("TODO: Import as sim mesh not implemented"));
 			}
 
 			if (bImportRenderMesh)
 			{
-				ClothPatternFacade.SetNumRenderVertices(Section.NumVertices);
-				ClothPatternFacade.SetNumRenderFaces(Section.NumTriangles);
-
-				TArrayView<FVector3f> RenderPosition = ClothPatternFacade.GetRenderPosition();
-				TArrayView<FVector3f> RenderNormal = ClothPatternFacade.GetRenderNormal();
-				TArrayView<FVector3f> RenderTangentU = ClothPatternFacade.GetRenderTangentU();
-				TArrayView<FVector3f> RenderTangentV = ClothPatternFacade.GetRenderTangentV();
-				TArrayView<TArray<FVector2f>> RenderUVs = ClothPatternFacade.GetRenderUVs();
-				TArrayView<FLinearColor> RenderColor = ClothPatternFacade.GetRenderColor();
-				TArrayView<int32> RenderNumBoneInfluences = ClothPatternFacade.GetRenderNumBoneInfluences();
-				TArrayView<TArray<int32>> RenderBoneIndices = ClothPatternFacade.GetRenderBoneIndices();
-				TArrayView<TArray<float>> RenderBoneWeights = ClothPatternFacade.GetRenderBoneWeights();
-				TArrayView<FIntVector3> RenderIndices = ClothPatternFacade.GetRenderIndices();
-				TArrayView<int32> RenderMaterialIndex = ClothPatternFacade.GetRenderMaterialIndex();
-
-				for (int32 VertexIndex = 0; VertexIndex < Section.NumVertices; ++VertexIndex)
-				{
-					const FSoftSkinVertex& SoftVertex = Section.SoftVertices[VertexIndex];
-
-					RenderPosition[VertexIndex] = SoftVertex.Position;
-					RenderNormal[VertexIndex] = SoftVertex.TangentZ;
-					RenderTangentU[VertexIndex] = SoftVertex.TangentX;
-					RenderTangentV[VertexIndex] = SoftVertex.TangentY;
-
-					TArray<FVector2f> UVs;
-					UVs.SetNum(MAX_TEXCOORDS);
-					for (int32 TexCoordIndex = 0; TexCoordIndex < MAX_TEXCOORDS; ++TexCoordIndex)
-					{
-						UVs[TexCoordIndex] = SoftVertex.UVs[TexCoordIndex];
-					}
-					RenderUVs[VertexIndex] = UVs;
-
-					RenderColor[VertexIndex] = FLinearColor(SoftVertex.Color);
-
-					FBoneData BoneData = ExtractBones(Section, SoftVertex);
-					RenderNumBoneInfluences[VertexIndex] = BoneData.Get<0>();
-					RenderBoneIndices[VertexIndex] = MoveTemp(BoneData.Get<1>());
-					RenderBoneWeights[VertexIndex] = MoveTemp(BoneData.Get<2>());
-				}
-
-				for (uint32 FaceIndex = 0; FaceIndex < Section.NumTriangles; ++FaceIndex)
-				{
-					RenderIndices[FaceIndex] = GetSectionFace(Section, LODModel.IndexBuffer, FaceIndex);
-					RenderMaterialIndex[FaceIndex] = 0;
-				}
-			}
-
-			ClothLodFacade.SetNumMaterials(1);
-			TArrayView<FString> RenderMaterialPathNames = ClothLodFacade.GetRenderMaterialPathName();
-			const TArray<FSkeletalMaterial>& Materials = SkeletalMesh->GetMaterials();
-			check(SectionIndex < Materials.Num());
-			if (const UMaterialInterface* Interface = Materials[SectionIndex].MaterialInterface)
-			{
-				RenderMaterialPathNames[0] = Interface->GetPathName();
+				const TArray<FSkeletalMaterial>& Materials = SkeletalMesh->GetMaterials();
+				check(SectionIndex < Materials.Num());
+				const FString RenderMaterialPathName = Materials[SectionIndex].MaterialInterface ? Materials[SectionIndex].MaterialInterface->GetPathName() : "";
+				FClothDataflowTools::AddRenderPatternFromSkeletalMeshSection(ClothCollection, LODModel, SectionIndex, RenderMaterialPathName);
 			}
 
 			if (const UPhysicsAsset* PhysicsAsset = SkeletalMesh->GetPhysicsAsset())
 			{
-				ClothLodFacade.SetPhysicsAssetPathName(PhysicsAsset->GetPathName());
+				ClothFacade.SetPhysicsAssetPathName(PhysicsAsset->GetPathName());
 			}
 			if (const USkeleton* Skeleton = SkeletalMesh->GetSkeleton())
 			{
-				ClothLodFacade.SetSkeletonAssetPathName(Skeleton->GetPathName());
+				ClothFacade.SetSkeletonAssetPathName(Skeleton->GetPathName());
 			}
 		}
 		SetValue<FManagedArrayCollection>(Context, *ClothCollection, &Collection);	
diff --git a/Engine/Plugins/Experimental/ChaosClothAssetEditor/Source/ChaosClothAssetDataflowNodes/Private/ChaosClothAsset/StaticMeshImportNode.cpp b/Engine/Plugins/Experimental/ChaosClothAssetEditor/Source/ChaosClothAssetDataflowNodes/Private/ChaosClothAsset/StaticMeshImportNode.cpp
--- a/Engine/Plugins/Experimental/ChaosClothAssetEditor/Source/ChaosClothAssetDataflowNodes/Private/ChaosClothAsset/StaticMeshImportNode.cpp
+++ b/Engine/Plugins/Experimental/ChaosClothAssetEditor/Source/ChaosClothAssetDataflowNodes/Private/ChaosClothAsset/StaticMeshImportNode.cpp
@@ -5,11 +5,17 @@
 #include "ChaosClothAsset/CollectionClothFacade.h"
 #include "ChaosClothAsset/ClothAsset.h"
 #include "ChaosClothAsset/ClothCollection.h"
+#include "ChaosClothAsset/ClothGeometryTools.h"
+#include "ChaosClothAsset/ClothDataflowTools.h"
 #include "Engine/StaticMesh.h"
+#include "Interfaces/ITargetPlatformManagerModule.h"
 #include "MeshDescriptionToDynamicMesh.h"
-#include "StaticMeshAttributes.h"
+#include "SkeletalMeshAttributes.h"
+#include "Rendering/SkeletalMeshLODImporterData.h"
+#include "Rendering/SkeletalMeshLODModel.h"
 #include "Materials/MaterialInterface.h"
-#include "StaticMeshOperations.h"
+#include "MeshUtilities.h"
+#include "Modules/ModuleManager.h"
 
 #include UE_INLINE_GENERATED_CPP_BY_NAME(StaticMeshImportNode)
 
@@ -17,202 +23,89 @@
 
 namespace UE::Chaos::ClothAsset::Private
 {
-struct FBuildVertex
-{
-	FVertexID OrigVertID;
-	FVertexInstanceID OrigVertInstanceID;
-
-	FVector3f Position;
-	TArray<FVector2f> UVs;
-	FVector3f Normal;
-	FVector3f Tangent;
-	float BiNormalSign;
-	FVector4f Color;
-
-	bool operator==(const FBuildVertex& Other) const
-	{
-		if (OrigVertID != Other.OrigVertID)
-		{
-			return false;
-		}
-
-		// No need to check position since it's from OrigVertID
-
-		if (!ensure(UVs.Num() == Other.UVs.Num()))
-		{
-			return false;
-		}
-		for (int32 UVChannelIndex = 0; UVChannelIndex < UVs.Num(); ++UVChannelIndex)
-		{
-			if (!UVs[UVChannelIndex].Equals(Other.UVs[UVChannelIndex], UE_THRESH_UVS_ARE_SAME))
-			{
-				return false;
-			}
-		}
-		if (!Normal.Equals(Other.Normal, UE_THRESH_NORMALS_ARE_SAME))
-		{
-			return false;
-		}
-		if (!Tangent.Equals(Other.Tangent, UE_THRESH_NORMALS_ARE_SAME))
-		{
-			return false;
-		}
-		if (BiNormalSign != Other.BiNormalSign) // I think these are just -1 or 1, so just comparing them
-		{
-			return false;
-		}
-		// It looks like we just quantize the 0-1 color values to uint8s. (Call FLinearColor::ToFColor(bSRGB = false)) when we consume these.
-		// Going to be a little more strict in case we ever decide to switch to the gamma conversion.
-		if (!Color.Equals(Other.Color, UE_THRESH_NORMALS_ARE_SAME))
-		{
-			return false;
-		}
-		return true;
-	}
-};
 
-static void MergeVertexInstances(const FMeshDescription* const MeshDescription, TArray<FBuildVertex>& MergedVertices, TArray<int32>& VertexInstanceToMerged)
+void CopyBuildSettings(
+	const FMeshBuildSettings& InStaticMeshBuildSettings,
+	FSkeletalMeshBuildSettings& OutSkeletalMeshBuildSettings
+)
 {
-	FStaticMeshConstAttributes Attributes(*MeshDescription);
-
-	TVertexAttributesConstRef<FVector3f> VertexPositions = Attributes.GetVertexPositions();
-
-	TVertexInstanceAttributesConstRef<FVector2f> VertexInstanceUVs = Attributes.GetVertexInstanceUVs();
-	TVertexInstanceAttributesConstRef<FVector3f> VertexInstanceNormals = Attributes.GetVertexInstanceNormals();
-	TVertexInstanceAttributesConstRef<FVector3f> VertexInstanceTangents = Attributes.GetVertexInstanceTangents();
-	TVertexInstanceAttributesConstRef<float> VertexInstanceBiNormalSigns = Attributes.GetVertexInstanceBinormalSigns();
-	TVertexInstanceAttributesConstRef<FVector4f> VertexInstanceColors = Attributes.GetVertexInstanceColors();
+	OutSkeletalMeshBuildSettings.bRecomputeNormals = InStaticMeshBuildSettings.bRecomputeNormals;
+	OutSkeletalMeshBuildSettings.bRecomputeTangents = InStaticMeshBuildSettings.bRecomputeTangents;
+	OutSkeletalMeshBuildSettings.bUseMikkTSpace = InStaticMeshBuildSettings.bUseMikkTSpace;
+	OutSkeletalMeshBuildSettings.bComputeWeightedNormals = InStaticMeshBuildSettings.bComputeWeightedNormals;
+	OutSkeletalMeshBuildSettings.bRemoveDegenerates = InStaticMeshBuildSettings.bRemoveDegenerates;
+	OutSkeletalMeshBuildSettings.bUseHighPrecisionTangentBasis = InStaticMeshBuildSettings.bUseHighPrecisionTangentBasis;
+	OutSkeletalMeshBuildSettings.bUseFullPrecisionUVs = InStaticMeshBuildSettings.bUseFullPrecisionUVs;
+	OutSkeletalMeshBuildSettings.bUseBackwardsCompatibleF16TruncUVs = InStaticMeshBuildSettings.bUseBackwardsCompatibleF16TruncUVs;
+	// The rest we leave at defaults.
+}
 
-	// Merge identical vertex instances that correspond with the same vertex.
-	auto MakeBuildVertex = [&VertexPositions, &VertexInstanceUVs, &VertexInstanceNormals, &VertexInstanceTangents, &VertexInstanceBiNormalSigns, &VertexInstanceColors](FVertexID VertID, FVertexInstanceID VertexInstanceID)
+bool BuildSkeletalMeshModelFromMeshDescription(const FMeshDescription* const InMeshDescription, const FMeshBuildSettings& InBuildSettings, FSkeletalMeshLODModel& SkeletalMeshModel)
+{
+	// This is following StaticToSkeletalMeshConverter.cpp::AddLODFromStaticMeshSourceModel
+	FSkeletalMeshBuildSettings BuildSettings;
+	CopyBuildSettings(InBuildSettings, BuildSettings);
+	FMeshDescription SkeletalMeshGeometry = *InMeshDescription;
+	FSkeletalMeshAttributes SkeletalMeshAttributes(SkeletalMeshGeometry);
+	SkeletalMeshAttributes.Register();
+
+	// Full binding to the root bone.
+	constexpr int32 RootBoneIndex = 0;
+	FSkinWeightsVertexAttributesRef SkinWeights = SkeletalMeshAttributes.GetVertexSkinWeights();
+	UE::AnimationCore::FBoneWeight RootInfluence(RootBoneIndex, 1.0f);
+	UE::AnimationCore::FBoneWeights RootBinding = UE::AnimationCore::FBoneWeights::Create({ RootInfluence });
+
+	for (const FVertexID VertexID : SkeletalMeshGeometry.Vertices().GetElementIDs())
 	{
-		FBuildVertex BuildVertex;
-		BuildVertex.OrigVertID = VertID;
-		BuildVertex.OrigVertInstanceID = VertexInstanceID;
-		BuildVertex.Position = VertexPositions[VertID];
-		BuildVertex.UVs.SetNumUninitialized(VertexInstanceUVs.GetNumChannels());
-		for (int32 UVChannelIndex = 0; UVChannelIndex < VertexInstanceUVs.GetNumChannels(); ++UVChannelIndex)
-		{
-			BuildVertex.UVs[UVChannelIndex] = VertexInstanceUVs.Get(VertexInstanceID, UVChannelIndex);
-		}
-		BuildVertex.Normal = VertexInstanceNormals[VertexInstanceID];
-		BuildVertex.Tangent = VertexInstanceTangents[VertexInstanceID];
-		BuildVertex.BiNormalSign = VertexInstanceBiNormalSigns[VertexInstanceID];
-		BuildVertex.Color = VertexInstanceColors[VertexInstanceID];
-		return BuildVertex;
-	};
-	MergedVertices.Reset(MeshDescription->VertexInstances().GetArraySize());
-	VertexInstanceToMerged.Init(INDEX_NONE, MeshDescription->VertexInstances().GetArraySize());
+		SkinWeights.Set(VertexID, RootBinding);
+	}
 
-	for (const FVertexID& VertID : MeshDescription->Vertices().GetElementIDs())
+	FSkeletalMeshImportData SkeletalMeshImportGeometry = FSkeletalMeshImportData::CreateFromMeshDescription(SkeletalMeshGeometry);
+	// Data needed by BuildSkeletalMesh
+	TArray<FVector3f> LODPoints;
+	TArray<SkeletalMeshImportData::FMeshWedge> LODWedges;
+	TArray<SkeletalMeshImportData::FMeshFace> LODFaces;
+	TArray<SkeletalMeshImportData::FVertInfluence> LODInfluences;
+	TArray<int32> LODPointToRawMap;
+	SkeletalMeshImportGeometry.CopyLODImportData(LODPoints, LODWedges, LODFaces, LODInfluences, LODPointToRawMap);
+	IMeshUtilities::MeshBuildOptions BuildOptions;
+	BuildOptions.TargetPlatform = GetTargetPlatformManagerRef().GetRunningTargetPlatform();
+	BuildOptions.FillOptions(BuildSettings);
+
+	static const FString SkeletalMeshName("ClothAssetStaticMeshImportConvert"); // This is only used by warning messages in the mesh builder.
+	// Build a RefSkeleton with just a root bone. The BuildSkeletalMesh code expects you have a reference skeleton with at least one bone to work.
+	FReferenceSkeleton RootBoneRefSkeleton;
+	FReferenceSkeletonModifier SkeletonModifier(RootBoneRefSkeleton, nullptr);
+	FMeshBoneInfo RootBoneInfo;
+	RootBoneInfo.Name = FName("Root");
+	SkeletonModifier.Add(RootBoneInfo, FTransform());
+	RootBoneRefSkeleton.RebuildRefSkeleton(nullptr, true);
+
+	IMeshUtilities& MeshUtilities = FModuleManager::Get().LoadModuleChecked<IMeshUtilities>("MeshUtilities");
+	TArray<FText> WarningMessages;
+	if (!MeshUtilities.BuildSkeletalMesh(SkeletalMeshModel, SkeletalMeshName, RootBoneRefSkeleton, LODInfluences, LODWedges, LODFaces, LODPoints, LODPointToRawMap, BuildOptions, &WarningMessages))
 	{
-		TConstArrayView<FVertexInstanceID> VertexInstances = MeshDescription->GetVertexVertexInstanceIDs(VertID);
-		if (VertexInstances.IsEmpty())
+		for (const FText& Message : WarningMessages)
 		{
-			continue;
-		}
-		const int32 FirstMergedVert = MergedVertices.Add(MakeBuildVertex(VertID, VertexInstances[0]));
-		VertexInstanceToMerged[VertexInstances[0].GetValue()] = FirstMergedVert;
-		for (int32 LocalInstance = 1; LocalInstance < VertexInstances.Num(); ++LocalInstance)
-		{
-			FBuildVertex BuildVert = MakeBuildVertex(VertID, VertexInstances[LocalInstance]);
-			bool bFoundDuplicate = false;
-			for (int32 CompareIndex = FirstMergedVert; CompareIndex < MergedVertices.Num(); ++CompareIndex)
-			{
-				if (BuildVert == MergedVertices[CompareIndex])
-				{
-					bFoundDuplicate = true;
-					VertexInstanceToMerged[VertexInstances[LocalInstance].GetValue()] = CompareIndex;
-					break;
-				}
-			}
-
-			if (!bFoundDuplicate)
-			{
-				const int32 AddedMergedVert = MergedVertices.Add(MoveTemp(BuildVert));
-				VertexInstanceToMerged[VertexInstances[LocalInstance].GetValue()] = AddedMergedVert;
-			}
+			DataflowNodes::LogAndToastWarning(FText::Format(LOCTEXT("SkelMeshConvertWarningFmt", "{0}"), Message));
 		}
+		return false;
 	}
+	return true;
 }
 
-static void InitializeRenderPatternDataFromMeshDescription(const FMeshDescription* const InMeshDescription, const FMeshBuildSettings& BuildSettings, FCollectionClothPatternFacade& RenderPattern)
+void InitializeDataFromMeshDescription(const FMeshDescription* const InMeshDescription, const FMeshBuildSettings& InBuildSettings, const TArray<FStaticMaterial>& StaticMaterials, const TSharedPtr<FManagedArrayCollection>& ClothCollection)
 {
-	const FMeshDescription* MeshDescription = InMeshDescription;
-	FMeshDescription WritableMeshDescription;
-	if (BuildSettings.bRecomputeNormals || BuildSettings.bRecomputeTangents)
+	FSkeletalMeshLODModel SkeletalMeshModel;
+	if (BuildSkeletalMeshModelFromMeshDescription(InMeshDescription, InBuildSettings, SkeletalMeshModel))
 	{
-		// check if any are invalid
-		bool bHasInvalidNormals, bHasInvalidTangents;
-		FStaticMeshOperations::AreNormalsAndTangentsValid(*InMeshDescription, bHasInvalidNormals, bHasInvalidTangents);
-
-		// if neither are invalid we are not going to recompute
-		if (bHasInvalidNormals || bHasInvalidTangents)
+		check(SkeletalMeshModel.Sections.Num() == StaticMaterials.Num());
+		for (int32 SectionIndex = 0; SectionIndex < SkeletalMeshModel.Sections.Num(); ++SectionIndex)
 		{
-			WritableMeshDescription = *InMeshDescription;
-			MeshDescription = &WritableMeshDescription;
-			FStaticMeshAttributes Attributes(WritableMeshDescription);
-			if (!Attributes.GetTriangleNormals().IsValid() || !Attributes.GetTriangleTangents().IsValid())
-			{
-				// If these attributes don't exist, create them and compute their values for each triangle
-				FStaticMeshOperations::ComputeTriangleTangentsAndNormals(WritableMeshDescription);
-			}
-
-			EComputeNTBsFlags ComputeNTBsOptions = EComputeNTBsFlags::BlendOverlappingNormals;
-			ComputeNTBsOptions |= BuildSettings.bRecomputeNormals ? EComputeNTBsFlags::Normals : EComputeNTBsFlags::None;
-			ComputeNTBsOptions |= BuildSettings.bRecomputeTangents ? EComputeNTBsFlags::Tangents : EComputeNTBsFlags::None;
-			ComputeNTBsOptions |= BuildSettings.bUseMikkTSpace ? EComputeNTBsFlags::UseMikkTSpace : EComputeNTBsFlags::None;
-			ComputeNTBsOptions |= BuildSettings.bComputeWeightedNormals ? EComputeNTBsFlags::WeightedNTBs : EComputeNTBsFlags::None;
-			ComputeNTBsOptions |= BuildSettings.bRemoveDegenerates ? EComputeNTBsFlags::IgnoreDegenerateTriangles : EComputeNTBsFlags::None;
-
-			FStaticMeshOperations::ComputeTangentsAndNormals(WritableMeshDescription, ComputeNTBsOptions);
+			const FString RenderMaterialPathName = StaticMaterials[SectionIndex].MaterialInterface ? StaticMaterials[SectionIndex].MaterialInterface->GetPathName() : "";
+			FClothDataflowTools::AddRenderPatternFromSkeletalMeshSection(ClothCollection, SkeletalMeshModel, SectionIndex, RenderMaterialPathName);
 		}
 	}
-
-	// Merge vertex instances that share the same vertex. These will become the pattern vertices.
-	TArray<FBuildVertex> MergedVertices;
-	TArray<int32> VertexInstanceToMerged;
-	MergeVertexInstances(MeshDescription, MergedVertices, VertexInstanceToMerged);
-
-	RenderPattern.SetNumRenderVertices(MergedVertices.Num());
-	RenderPattern.SetNumRenderFaces(MeshDescription->Triangles().Num());
-
-	// Vertex data (this will MoveTemp stuff out of MergedVertices!!)
-	TArrayView<FVector3f> RenderPosition = RenderPattern.GetRenderPosition();
-	TArrayView<FVector3f> RenderNormal = RenderPattern.GetRenderNormal();
-	TArrayView<FVector3f> RenderTangentU = RenderPattern.GetRenderTangentU();
-	TArrayView<FVector3f> RenderTangentV = RenderPattern.GetRenderTangentV();
-	TArrayView<TArray<FVector2f>> RenderUVs = RenderPattern.GetRenderUVs();
-	TArrayView<FLinearColor> RenderColor = RenderPattern.GetRenderColor();
-	for (int32 VertexIndex = 0; VertexIndex < MergedVertices.Num(); ++VertexIndex)
-	{
-		FBuildVertex& BuildVertex = MergedVertices[VertexIndex];
-		RenderPosition[VertexIndex] = BuildVertex.Position;
-		RenderNormal[VertexIndex] = BuildVertex.Normal;
-		RenderTangentU[VertexIndex] = BuildVertex.Tangent;
-		RenderTangentV[VertexIndex] = FVector3f::CrossProduct(BuildVertex.Normal, BuildVertex.Tangent).GetSafeNormal() * BuildVertex.BiNormalSign;
-		RenderUVs[VertexIndex] = MoveTemp(BuildVertex.UVs);
-		RenderColor[VertexIndex] = FLinearColor(BuildVertex.Color);
-	}
-
-	// Face data
-	TArrayView<FIntVector3> RenderIndices = RenderPattern.GetRenderIndices();
-	TArrayView<int32> RenderMaterialIndex = RenderPattern.GetRenderMaterialIndex();
-	int32 FaceIndex = 0;
-	for (const FTriangleID TriangleID : MeshDescription->Triangles().GetElementIDs())
-	{
-		const FPolygonGroupID PolygonGroupID = MeshDescription->GetTrianglePolygonGroup(TriangleID);
-		RenderMaterialIndex[FaceIndex] = PolygonGroupID.GetValue();
-
-		const TConstArrayView<FVertexInstanceID> VertexInstances = MeshDescription->GetTriangleVertexInstances(TriangleID);
-		check(VertexInstances.Num() == 3);
-		check(MergedVertices.IsValidIndex(VertexInstanceToMerged[VertexInstances[0].GetValue()]));
-		check(MergedVertices.IsValidIndex(VertexInstanceToMerged[VertexInstances[1].GetValue()]));
-		check(MergedVertices.IsValidIndex(VertexInstanceToMerged[VertexInstances[2].GetValue()]));
-		RenderIndices[FaceIndex] = FIntVector3(VertexInstanceToMerged[VertexInstances[0].GetValue()],
-			VertexInstanceToMerged[VertexInstances[1].GetValue()], VertexInstanceToMerged[VertexInstances[2].GetValue()]);
-		++FaceIndex;
-	}
 }
 
 } // namespace UE::Chaos::ClothAsset::Private
@@ -240,14 +133,11 @@ void FChaosClothAssetStaticMeshImportNode::Evaluate(Dataflow::FContext& Context,
 		if (StaticMesh && (bImportAsSimMesh || bImportAsRenderMesh))
 		{
 			const int32 NumLods = StaticMesh->GetNumSourceModels();
-			for (int32 LodIndex = 0; LodIndex < NumLods; ++LodIndex)
+			if(LodIndex < NumLods)
 			{
 				const FMeshDescription* const MeshDescription = StaticMesh->GetMeshDescription(LodIndex);
 				check(MeshDescription);
 
-				FCollectionClothLodFacade ClothLodFacade = ClothFacade.AddGetLod();
-				ClothLodFacade.Reset();
-
 				if (bImportAsSimMesh)
 				{
 					FMeshDescriptionToDynamicMesh Converter;
@@ -257,37 +147,26 @@ void FChaosClothAssetStaticMeshImportNode::Evaluate(Dataflow::FContext& Context,
 					UE::Geometry::FDynamicMesh3 DynamicMesh;
 
 					Converter.Convert(MeshDescription, DynamicMesh);
-					ClothLodFacade.Initialize(DynamicMesh, UVChannel, UVScale, bUseUVIslands);
+					UE::Chaos::ClothAsset::FClothGeometryTools::BuildSimMeshFromDynamicMesh(ClothCollection, DynamicMesh, UVChannel, UVScale);
 				}
 				if (bImportAsRenderMesh)
 				{
 					// Add render data into a single pattern for now
-					FCollectionClothPatternFacade RenderPattern = ClothLodFacade.AddGetPattern();
-					UE::Chaos::ClothAsset::Private::InitializeRenderPatternDataFromMeshDescription(MeshDescription, StaticMesh->GetSourceModel(LodIndex).BuildSettings, RenderPattern);
-
-					// Set material path names
-					const TArray<FStaticMaterial>& StaticMaterials = StaticMesh->GetStaticMaterials();
-					ClothLodFacade.SetNumMaterials(StaticMaterials.Num());
-					TArrayView<FString> RenderMaterialPathNames = ClothLodFacade.GetRenderMaterialPathName();
-					for(int32 MaterialIndex = 0; MaterialIndex < StaticMaterials.Num(); ++MaterialIndex)
-					{
-						if (ensure(StaticMaterials[MaterialIndex].MaterialInterface))
-						{
-							RenderMaterialPathNames[MaterialIndex] = StaticMaterials[MaterialIndex].MaterialInterface->GetPathName();
-						}
-					}
+					UE::Chaos::ClothAsset::Private::InitializeDataFromMeshDescription(MeshDescription, StaticMesh->GetSourceModel(LodIndex).BuildSettings, StaticMesh->GetStaticMaterials(), ClothCollection);
 				}
 
 				// Set a default skeleton
-				ClothLodFacade.SetSkeletonAssetPathName(DefaultSkeletonPathName);
+				ClothFacade.SetSkeletonAssetPathName(DefaultSkeletonPathName);
+			}
+			else
+			{
+				DataflowNodes::LogAndToastWarning(
+					FText::Format(LOCTEXT("FChaosClothAssetStaticMeshImportNode::InvalidLod", "FChaosClothAssetStaticMeshImportNode: Invalid LodIndex {0} >= Num Lods in static mesh {1}"),
+						LodIndex,
+						NumLods));
 			}
 		}
 
-		// Make sure that whatever happens there is always at least one empty LOD to avoid crashing the render data
-		if (ClothFacade.GetNumLods() < 1)
-		{
-			ClothFacade.AddGetLod().SetSkeletonAssetPathName(DefaultSkeletonPathName);
-		}
 		SetValue<FManagedArrayCollection>(Context, *ClothCollection, &Collection);
 	}
 }
diff --git a/Engine/Plugins/Experimental/ChaosClothAssetEditor/Source/ChaosClothAssetDataflowNodes/Private/ChaosClothAsset/TerminalNode.cpp b/Engine/Plugins/Experimental/ChaosClothAssetEditor/Source/ChaosClothAssetDataflowNodes/Private/ChaosClothAsset/TerminalNode.cpp
--- a/Engine/Plugins/Experimental/ChaosClothAssetEditor/Source/ChaosClothAssetDataflowNodes/Private/ChaosClothAsset/TerminalNode.cpp
+++ b/Engine/Plugins/Experimental/ChaosClothAssetEditor/Source/ChaosClothAssetDataflowNodes/Private/ChaosClothAsset/TerminalNode.cpp
@@ -25,105 +25,85 @@ void FChaosClothAssetTerminalNode::SetAssetValue(TObjectPtr<UObject> Asset, Data
 	if (UChaosClothAsset* ClothAsset = Cast<UChaosClothAsset>(Asset.Get()))
 	{
 		using namespace UE::Chaos::ClothAsset;
+		const TArray<const FManagedArrayCollection*> CollectionLods = GetCollectionLods();
 
 		// Reset the asset's collection
-		const TSharedPtr<FManagedArrayCollection> ClothCollection = ClothAsset->GetClothCollection();
-		ClothCollection->Reset();
+		TArray<TSharedPtr<FManagedArrayCollection>>& ClothCollections = ClothAsset->GetClothCollections();
+		ClothCollections.Reset(CollectionLods.Num());
 
 		// Reset the asset's material list
 		ClothAsset->GetMaterials().Reset();
 
-		// Define the final container's collection schema
-		FCollectionClothFacade ClothFacade(ClothCollection);
-		ClothFacade.DefineSchema();
-
 		// Iterate through the LODs
 		FString PhysicsAssetPathName;
 
-		const TArray<const FManagedArrayCollection*> CollectionLods = GetCollectionLods();
 		int32 LastValidLodIndex = INDEX_NONE;
 		for (int32 LodIndex = 0; LodIndex < CollectionLods.Num(); ++LodIndex)
 		{
 			// New LOD
-			const int32 NewLodIndex = ClothFacade.AddLod();
-			check(NewLodIndex == LodIndex);
-			FCollectionClothLodFacade ClothLodFacade = ClothFacade.GetLod(LodIndex);
+			TSharedPtr<FManagedArrayCollection>& ClothCollection = ClothCollections.Emplace_GetRef(MakeShared<FManagedArrayCollection>());
+			UE::Chaos::ClothAsset::FCollectionClothFacade ClothFacade(ClothCollection);
+			ClothFacade.DefineSchema();
 
 			// Retrieve input LOD
 			const FManagedArrayCollection& InCollectionLod = GetValue<FManagedArrayCollection>(Context, CollectionLods[LodIndex]);
 			const TSharedRef<const FManagedArrayCollection> InClothCollection = MakeShared<const FManagedArrayCollection>(InCollectionLod);
 
-			// Only use LOD 0 from input
 			const FCollectionClothConstFacade InClothFacade(InClothCollection);
-			if (InClothFacade.GetNumLods() >= 1)
+				// Check LOD validity
+			if (InClothFacade.GetNumSimPatterns() &&
+				InClothFacade.GetNumRenderPatterns() &&
+				InClothFacade.GetNumSimVertices2D() &&
+				InClothFacade.GetNumSimVertices3D() &&
+				InClothFacade.GetNumSimFaces() &&
+				InClothFacade.GetNumRenderVertices() &&
+				InClothFacade.GetNumRenderFaces())
 			{
-				const FCollectionClothLodConstFacade InClothLodFacade = InClothFacade.GetLod(0);
+				LastValidLodIndex = LodIndex;
 
-				// Check LOD validity
-				if (InClothLodFacade.GetNumPatterns() &&
-					InClothLodFacade.GetNumSimVertices() &&
-					InClothLodFacade.GetNumSimFaces() &&
-					InClothLodFacade.GetNumRenderVertices() &&
-					InClothLodFacade.GetNumRenderFaces())
-				{
-					LastValidLodIndex = LodIndex;
+				// Copy input LOD to current output LOD
+				ClothFacade.Initialize(InClothFacade);
 
-					// Copy input LOD 0 to current output LOD
-					ClothLodFacade.Initialize(InClothLodFacade);
+				// Add this LOD's materials to the asset
+				const int32 NumLodMaterials = ClothFacade.GetNumRenderPatterns();
 
-					// Add this LOD's materials to the asset
-					const int32 NumLodMaterials = InClothLodFacade.GetNumMaterials();
+				TArray<FSkeletalMaterial>& Materials = ClothAsset->GetMaterials();
+				Materials.Reserve(Materials.Num() + NumLodMaterials);
 
-					TArray<FSkeletalMaterial>& Materials = ClothAsset->GetMaterials();
-					Materials.Reserve(Materials.Num() + NumLodMaterials);
+				const TConstArrayView<FString> LodRenderMaterialPathName = ClothFacade.GetRenderMaterialPathName();
+				for (int32 LodMaterialIndex = 0; LodMaterialIndex < NumLodMaterials; ++LodMaterialIndex)
+				{
+					const FString& RenderMaterialPathName = LodRenderMaterialPathName[LodMaterialIndex];
 
-					const TConstArrayView<FString> LodRenderMaterialPathName = ClothLodFacade.GetRenderMaterialPathName();
-					for (int32 LodMaterialIndex = 0; LodMaterialIndex < NumLodMaterials; ++LodMaterialIndex)
+					auto Predicate = [&RenderMaterialPathName](const FSkeletalMaterial& SkeletalMaterial)
 					{
-						const FString& RenderMaterialPathName = LodRenderMaterialPathName[LodMaterialIndex];
+						return SkeletalMaterial.MaterialInterface && SkeletalMaterial.MaterialInterface->GetPathName() == RenderMaterialPathName;
+					};
 
-						auto Predicate = [&RenderMaterialPathName](const FSkeletalMaterial& SkeletalMaterial)
-						{
-							return SkeletalMaterial.MaterialInterface && SkeletalMaterial.MaterialInterface->GetPathName() == RenderMaterialPathName;
-						};
-
-						if (!Materials.FindByPredicate(Predicate))
-						{
-							if (UMaterialInterface* const Material = LoadObject<UMaterialInterface>(ClothAsset, *RenderMaterialPathName, nullptr, LOAD_None, nullptr))
-							{
-								Materials.Emplace(Material, true, false, Material->GetFName());
-							}
-						}
-					}
-
-					const TArray<FName> MapNames = InClothFacade.GetWeightMapNames();
-					for (const FName& MapName : MapNames)
+					if (!Materials.FindByPredicate(Predicate))
 					{
-						ClothFacade.AddWeightMap(MapName);
-						TArrayView<float> WeightMap = ClothLodFacade.GetWeightMap(MapName);
-						TConstArrayView<float> InWeightMap = InClothLodFacade.GetWeightMap(MapName);
-						for (int32 i = 0; i < WeightMap.Num(); ++i)
+						if (UMaterialInterface* const Material = LoadObject<UMaterialInterface>(ClothAsset, *RenderMaterialPathName, nullptr, LOAD_None, nullptr))
 						{
-							WeightMap[i] = InWeightMap[i];
+							Materials.Emplace(Material, true, false, Material->GetFName());
 						}
 					}
+				}
 
-					// Set properties and physics asset only with LOD 0 at the moment
-					if (LodIndex == 0)
-					{
-						using namespace ::Chaos::Softs;
-
-						PhysicsAssetPathName = InClothLodFacade.GetPhysicsAssetPathName();
+				// Set properties
+				Chaos::Softs::FCollectionPropertyMutableFacade(ClothCollection).Append(*InClothCollection);
 
-						FCollectionPropertyMutableFacade(ClothCollection).Append(*InClothCollection);
-					}
+				// Set physics asset only with LOD 0 at the moment
+				if (LodIndex == 0)
+				{
+					using namespace ::Chaos::Softs;
+					PhysicsAssetPathName = InClothFacade.GetPhysicsAssetPathName();
 				}
 			}
 			if (LodIndex != LastValidLodIndex)
 			{
 				if (LastValidLodIndex >= 0)
 				{
-					ClothLodFacade.Initialize(ClothFacade.GetLod(LastValidLodIndex));
+					ClothFacade.Initialize(FCollectionClothConstFacade(ClothCollections[LastValidLodIndex]));
 
 					DataflowNodes::LogAndToastWarning(
 						FText::Format(
@@ -136,8 +116,6 @@ void FChaosClothAssetTerminalNode::SetAssetValue(TObjectPtr<UObject> Asset, Data
 				}
 				else
 				{
-					ClothFacade.SetNumLods(0);
-
 					DataflowNodes::LogAndToastWarning(
 						LOCTEXT(
 							"InvalidLod0",
@@ -149,9 +127,11 @@ void FChaosClothAssetTerminalNode::SetAssetValue(TObjectPtr<UObject> Asset, Data
 		}
 
 		// Make sure that whatever happens there is always at least one empty LOD to avoid crashing the render data
-		if (ClothFacade.GetNumLods() < 1)
+		if (ClothCollections.Num() < 1)
 		{
-			ClothFacade.AddLod();
+			TSharedPtr<FManagedArrayCollection>& ClothCollection = ClothCollections.Emplace_GetRef(MakeShared<FManagedArrayCollection>());
+			UE::Chaos::ClothAsset::FCollectionClothFacade ClothFacade(ClothCollection);
+			ClothFacade.DefineSchema();
 		}
 
 		// Set reference skeleton
diff --git a/Engine/Plugins/Experimental/ChaosClothAssetEditor/Source/ChaosClothAssetDataflowNodes/Private/ChaosClothAsset/TransferSkinWeightsNode.cpp b/Engine/Plugins/Experimental/ChaosClothAssetEditor/Source/ChaosClothAssetDataflowNodes/Private/ChaosClothAsset/TransferSkinWeightsNode.cpp
--- a/Engine/Plugins/Experimental/ChaosClothAssetEditor/Source/ChaosClothAssetDataflowNodes/Private/ChaosClothAsset/TransferSkinWeightsNode.cpp
+++ b/Engine/Plugins/Experimental/ChaosClothAssetEditor/Source/ChaosClothAssetDataflowNodes/Private/ChaosClothAsset/TransferSkinWeightsNode.cpp
@@ -34,31 +34,23 @@ namespace UE::ChaosClothAsset::Private
 		Converter.Convert(&SourceMesh, ToDynamicMesh);
 	}
 
-	// Convert the cloth LOD into DynamicMesh by welding the sim mesh patterns
+	// Convert the cloth into DynamicMesh
 	// @todo: This should instead be handled by a cloth lod to dynamic mesh converter similar to functions in the ClothPatternToDynamicMesh.h
-	bool SimClothLodToDynamicMesh(const FCollectionClothLodFacade& ClothLodFacade, 
-								  const FReferenceSkeleton& TargetRefSkeleton, // the reference skeleton to add to the dynamic mesh
-								  FDynamicMesh3& WeldedSimMesh,
-								  TArray<TArray<int32>>& WeldedToPatternIndices)
+	bool SimClothToDynamicMesh(const FCollectionClothConstFacade& ClothFacade,
+								const FReferenceSkeleton& TargetRefSkeleton, // the reference skeleton to add to the dynamic mesh
+								FDynamicMesh3& WeldedSimMesh)
 	{
-		// Compute the welded simulation mesh for this LOD.
-		TArray<FVector3f> Positions;
-		TArray<FVector3f> Normals;
-		TArray<uint32> Indices, PatternsIndices, PatternToWeldedIndices;
-		TArray<FVector2f> PatternsPositions;
-		ClothLodFacade.BuildSimulationMesh(Positions, Normals, Indices, PatternsPositions, PatternsIndices, PatternToWeldedIndices, &WeldedToPatternIndices);
-
 		// Convert the sim mesh to DynamicMesh. 
 		// @todo: FTransferBoneWeights should accept raw data arrays (vertices/triangles) to avoid this conversion
 		WeldedSimMesh.Clear();
-		for (const FVector3f& Pos : Positions)
+		for (const FVector3f& Pos : ClothFacade.GetSimPosition3D())
 		{
 			WeldedSimMesh.AppendVertex(FVector3d(Pos));
 		}
 
-		for (int Idx = 0; Idx < Indices.Num(); Idx+=3)
+		for (const FIntVector& Indices : ClothFacade.GetSimIndices3D())
 		{
-			const int TID = WeldedSimMesh.AppendTriangle(Indices[Idx], Indices[Idx+1], Indices[Idx+2]);
+			const int TID = WeldedSimMesh.AppendTriangle(Indices[0], Indices[1], Indices[2]);
 			if (TID < 0) // Failed to add the triangle (non-manifold/duplicate)
 			{
 				return false;
@@ -80,6 +72,112 @@ namespace UE::ChaosClothAsset::Private
 
 		return true;
 	}
+
+	void TransferInpaintWeights(const FReferenceSkeleton& TargetRefSkeleton, const double NormalThreshold, const double RadiusPercentage, bool bUseParallel, FCollectionClothFacade& ClothFacade, FTransferBoneWeights& TransferBoneWeights)
+	{
+
+		//
+		// Convert cloth sim mesh LOD to the welded dynamic sim mesh.
+		//
+		FDynamicMesh3 WeldedSimMesh;
+		if (!ensure(SimClothToDynamicMesh(ClothFacade, TargetRefSkeleton, WeldedSimMesh)))
+		{
+			UE::Chaos::ClothAsset::DataflowNodes::LogAndToastWarning(LOCTEXT("Warning_TransferWeightsFailedLodToDynamicMesh", "TransferSkinWeightsNode: Failed to weld the simulation mesh for LOD."));
+			return;
+		}
+
+
+		//
+		// Transfer the weights from the body to the welded sim mesh.
+		//
+		TransferBoneWeights.NormalThreshold = FMathd::DegToRad * NormalThreshold;
+		TransferBoneWeights.SearchRadius = RadiusPercentage * WeldedSimMesh.GetBounds().DiagonalLength();
+		if (TransferBoneWeights.Validate() != EOperationValidationResult::Ok)
+		{
+			UE::Chaos::ClothAsset::DataflowNodes::LogAndToastWarning(LOCTEXT("Warning_TransferWeightsInpaintWeightsInvalidParameters", "TransferSkinWeightsNode: Transfer method parameters are invalid."));
+			return;
+		}
+		if (!TransferBoneWeights.TransferWeightsToMesh(WeldedSimMesh, FSkeletalMeshAttributes::DefaultSkinWeightProfileName))
+		{
+			UE::Chaos::ClothAsset::DataflowNodes::LogAndToastWarning(LOCTEXT("Warning_TransferWeightsFailed", "TransferSkinWeightsNode: Transferring skin weights failed"));
+			return;
+		}
+
+
+		//
+		// Copy the new bone weight data from the welded sim mesh back to the cloth patterns.
+		//
+		ParallelFor(WeldedSimMesh.MaxVertexID(), [&ClothFacade, &WeldedSimMesh](int32 WeldedID)
+		{
+			const FDynamicMeshVertexSkinWeightsAttribute* OutAttribute = WeldedSimMesh.Attributes()->GetSkinWeightsAttribute(FSkeletalMeshAttributes::DefaultSkinWeightProfileName);
+
+			checkSlow(OutAttribute);
+			checkSlow(WeldedSimMesh.IsVertex(WeldedID));
+			checkSlow(WeldedID < ClothFacade.GetNumSimVertices3D());
+			OutAttribute->GetValue(WeldedID, ClothFacade.GetSimBoneIndices()[WeldedID],
+				ClothFacade.GetSimBoneWeights()[WeldedID]);
+		}, bUseParallel ? EParallelForFlags::None : EParallelForFlags::ForceSingleThread);
+
+
+		//
+		// Compute the bone weights for the render mesh by transferring weights from the sim mesh
+		// @todo If render mesh eventually supports welding, we should be transferring weights from the 
+		// body instead, same as we do for the sim mesh.
+		//
+		FTransferBoneWeights SimToRenderMeshTransfer(&WeldedSimMesh, FSkeletalMeshAttributes::DefaultSkinWeightProfileName);
+		SimToRenderMeshTransfer.bUseParallel = bUseParallel;
+		SimToRenderMeshTransfer.TransferMethod = FTransferBoneWeights::ETransferBoneWeightsMethod::ClosestPointOnSurface;
+		ParallelFor(ClothFacade.GetNumRenderVertices(), [&SimToRenderMeshTransfer, &ClothFacade](int32 VertexID)
+		{
+			SimToRenderMeshTransfer.TransferWeightsToPoint(ClothFacade.GetRenderBoneIndices()[VertexID],
+				ClothFacade.GetRenderBoneWeights()[VertexID],
+				ClothFacade.GetRenderPosition()[VertexID]);
+
+		}, bUseParallel ? EParallelForFlags::None : EParallelForFlags::ForceSingleThread);
+	}
+
+	void TransferClosestPointOnSurface(const FReferenceSkeleton& TargetRefSkeleton, const bool bUseParallel, FCollectionClothFacade& ClothFacade, FTransferBoneWeights& TransferBoneWeights)
+	{
+		//
+		// Compute the bone index mappings. This allows the transfer operator to retarget weights to the correct skeleton.
+		//
+		TMap<FName, FBoneIndexType> TargetBoneToIndex;
+		TargetBoneToIndex.Reserve(TargetRefSkeleton.GetRawBoneNum());
+		for (int32 BoneIdx = 0; BoneIdx < TargetRefSkeleton.GetRawBoneNum(); ++BoneIdx)
+		{
+			TargetBoneToIndex.Add(TargetRefSkeleton.GetRawRefBoneInfo()[BoneIdx].Name, BoneIdx);
+		}
+
+		if (TransferBoneWeights.Validate() != EOperationValidationResult::Ok)
+		{
+			UE::Chaos::ClothAsset::DataflowNodes::LogAndToastWarning(LOCTEXT("Warning_TransferWeightsClosestPointOnSurfaceInvalidParameters", "TransferSkinWeightsNode: Transfer method parameters are invalid."));
+			return;
+		}
+
+		//
+		// Transfer weights to the sim mesh.
+		//
+		ParallelFor(ClothFacade.GetNumSimVertices3D(), [&TransferBoneWeights, &TargetBoneToIndex, &ClothFacade](int32 VertexID)
+		{
+			TransferBoneWeights.TransferWeightsToPoint(ClothFacade.GetSimBoneIndices()[VertexID],
+				ClothFacade.GetSimBoneWeights()[VertexID],
+				ClothFacade.GetSimPosition3D()[VertexID],
+				&TargetBoneToIndex);
+
+		}, bUseParallel ? EParallelForFlags::None : EParallelForFlags::ForceSingleThread);
+
+		//
+		// Transfer weights to the render mesh.
+		//
+		ParallelFor(ClothFacade.GetNumRenderVertices(), [&TransferBoneWeights, &TargetBoneToIndex, &ClothFacade](int32 VertexID)
+		{
+			TransferBoneWeights.TransferWeightsToPoint(ClothFacade.GetRenderBoneIndices()[VertexID],
+				ClothFacade.GetRenderBoneWeights()[VertexID],
+				ClothFacade.GetRenderPosition()[VertexID],
+				&TargetBoneToIndex);
+
+		}, bUseParallel ? EParallelForFlags::None : EParallelForFlags::ForceSingleThread);
+	}
 }
 
 
@@ -131,131 +229,22 @@ void FChaosClothAssetTransferSkinWeightsNode::Evaluate(Dataflow::FContext& Conte
 
 
 			//
-			// Iterate over the LODs and transfer the bone weights from the source Skeletal mesh to the Cloth asset.
+			// Transfer the bone weights from the source Skeletal mesh to the Cloth asset.
 			//
 			FCollectionClothFacade ClothFacade(ClothCollection);
-			for (int32 TargetLODIdx = 0; TargetLODIdx < ClothFacade.GetNumLods(); ++TargetLODIdx)
+			ClothFacade.SetSkeletonAssetPathName(InputSkeletalMesh->GetSkeleton()->GetPathName());
+
+			if (TransferMethod == EChaosClothAssetTransferSkinWeightsMethod::InpaintWeights)
+			{
+				TransferInpaintWeights(TargetRefSkeleton, NormalThreshold, RadiusPercentage, bUseParallel, ClothFacade, TransferBoneWeights);
+			}
+			else if (TransferMethod == EChaosClothAssetTransferSkinWeightsMethod::ClosestPointOnSurface)
+			{
+				TransferClosestPointOnSurface(TargetRefSkeleton, bUseParallel, ClothFacade, TransferBoneWeights);
+			}
+			else
 			{
-				FCollectionClothLodFacade ClothLodFacade = ClothFacade.GetLod(TargetLODIdx);
-				ClothLodFacade.SetSkeletonAssetPathName(InputSkeletalMesh->GetSkeleton()->GetPathName());
-
-				if (TransferMethod == EChaosClothAssetTransferSkinWeightsMethod::InpaintWeights)
-				{	
-					//
-					// Convert cloth sim mesh LOD to the welded dynamic sim mesh.
-					//
-					FDynamicMesh3 WeldedSimMesh;
-					TArray<TArray<int32>> WeldedToPatternIndices;
-					if (!ensure(SimClothLodToDynamicMesh(ClothLodFacade, TargetRefSkeleton, WeldedSimMesh, WeldedToPatternIndices)))
-					{
-						UE::Chaos::ClothAsset::DataflowNodes::LogAndToastWarning(LOCTEXT("Warning_TransferWeightsFailedLodToDynamicMesh", "TransferSkinWeightsNode: Failed to weld the simulation mesh for LOD."));
-						break;
-					}
-
-
-					//
-					// Transfer the weights from the body to the welded sim mesh.
-					//
-					TransferBoneWeights.NormalThreshold = FMathd::DegToRad * NormalThreshold;
-					TransferBoneWeights.SearchRadius = RadiusPercentage * WeldedSimMesh.GetBounds().DiagonalLength();
-					if (TransferBoneWeights.Validate() != EOperationValidationResult::Ok)
-					{
-						UE::Chaos::ClothAsset::DataflowNodes::LogAndToastWarning(LOCTEXT("Warning_TransferWeightsInpaintWeightsInvalidParameters", "TransferSkinWeightsNode: Transfer method parameters are invalid."));
-						break;
-					}
-					if (!TransferBoneWeights.TransferWeightsToMesh(WeldedSimMesh, FSkeletalMeshAttributes::DefaultSkinWeightProfileName))
-					{
-						UE::Chaos::ClothAsset::DataflowNodes::LogAndToastWarning(LOCTEXT("Warning_TransferWeightsFailed", "TransferSkinWeightsNode: Transferring skin weights failed"));
-						break;
-					}
-
-
-					//
-					// Copy the new bone weight data from the welded sim mesh back to the cloth patterns.
-					//
-					ParallelFor(WeldedSimMesh.MaxVertexID(), [&ClothLodFacade, &WeldedToPatternIndices, &WeldedSimMesh](int32 WeldedID)
-					{
-						const FDynamicMeshVertexSkinWeightsAttribute* OutAttribute = WeldedSimMesh.Attributes()->GetSkinWeightsAttribute(FSkeletalMeshAttributes::DefaultSkinWeightProfileName);
-						
-						checkSlow(OutAttribute);
-						checkSlow(WeldedSimMesh.IsVertex(WeldedID));
-
-						for (const int32 VertexID : WeldedToPatternIndices[WeldedID])
-						{	
-							checkSlow(VertexID < ClothLodFacade.GetNumSimVertices());
-							
-							OutAttribute->GetValue(WeldedID, ClothLodFacade.GetSimBoneIndices()[VertexID], 
-															 ClothLodFacade.GetSimBoneWeights()[VertexID]);
-							ClothLodFacade.GetSimNumBoneInfluences()[VertexID] = ClothLodFacade.GetSimBoneIndices()[VertexID].Num();
-						}
-					}, bUseParallel ? EParallelForFlags::None : EParallelForFlags::ForceSingleThread);
-
-
-					//
-					// Compute the bone weights for the render mesh by transferring weights from the sim mesh
-					// @todo If render mesh eventually supports welding, we should be transferring weights from the 
-					// body instead, same as we do for the sim mesh.
-					//
-					FTransferBoneWeights SimToRenderMeshTransfer(&WeldedSimMesh, FSkeletalMeshAttributes::DefaultSkinWeightProfileName);
-					SimToRenderMeshTransfer.bUseParallel = bUseParallel;
-					SimToRenderMeshTransfer.TransferMethod = FTransferBoneWeights::ETransferBoneWeightsMethod::ClosestPointOnSurface;
-					ParallelFor(ClothLodFacade.GetNumRenderVertices(), [&SimToRenderMeshTransfer, &ClothLodFacade](int32 VertexID)
-					{
-						SimToRenderMeshTransfer.TransferWeightsToPoint(ClothLodFacade.GetRenderBoneIndices()[VertexID],
-																	   ClothLodFacade.GetRenderBoneWeights()[VertexID], 
-																	   ClothLodFacade.GetRenderPosition()[VertexID]);
-						ClothLodFacade.GetRenderNumBoneInfluences()[VertexID] = ClothLodFacade.GetRenderBoneIndices()[VertexID].Num();
-
-					}, bUseParallel ? EParallelForFlags::None : EParallelForFlags::ForceSingleThread);
-				}
-				else if (TransferMethod == EChaosClothAssetTransferSkinWeightsMethod::ClosestPointOnSurface)
-				{	
-					//
-					// Compute the bone index mappings. This allows the transfer operator to retarget weights to the correct skeleton.
-					//
-					TMap<FName, FBoneIndexType> TargetBoneToIndex;
-					TargetBoneToIndex.Reserve(TargetRefSkeleton.GetRawBoneNum());
-					for (int32 BoneIdx = 0; BoneIdx < TargetRefSkeleton.GetRawBoneNum(); ++BoneIdx)
-					{
-						TargetBoneToIndex.Add(TargetRefSkeleton.GetRawRefBoneInfo()[BoneIdx].Name, BoneIdx);
-					}
-
-					if (TransferBoneWeights.Validate() != EOperationValidationResult::Ok)
-					{
-						UE::Chaos::ClothAsset::DataflowNodes::LogAndToastWarning(LOCTEXT("Warning_TransferWeightsClosestPointOnSurfaceInvalidParameters", "TransferSkinWeightsNode: Transfer method parameters are invalid."));
-						break;
-					}
-					
-					//
-					// Transfer weights to the sim mesh.
-					//
-					ParallelFor(ClothLodFacade.GetNumSimVertices(), [&TransferBoneWeights, &TargetBoneToIndex, &ClothLodFacade](int32 VertexID)
-					{
-						TransferBoneWeights.TransferWeightsToPoint(ClothLodFacade.GetSimBoneIndices()[VertexID],
-																   ClothLodFacade.GetSimBoneWeights()[VertexID],
-																   ClothLodFacade.GetSimRestPosition()[VertexID],
-																   &TargetBoneToIndex);
-					    ClothLodFacade.GetSimNumBoneInfluences()[VertexID] = ClothLodFacade.GetSimBoneIndices()[VertexID].Num();
-
-					}, bUseParallel ? EParallelForFlags::None : EParallelForFlags::ForceSingleThread);
-
-					//
-					// Transfer weights to the render mesh.
-					//
-					ParallelFor(ClothLodFacade.GetNumRenderVertices(), [&TransferBoneWeights, &TargetBoneToIndex, &ClothLodFacade](int32 VertexID)
-					{
-						TransferBoneWeights.TransferWeightsToPoint(ClothLodFacade.GetRenderBoneIndices()[VertexID],
-																   ClothLodFacade.GetRenderBoneWeights()[VertexID],
-																   ClothLodFacade.GetRenderPosition()[VertexID],
-																   &TargetBoneToIndex);
-						ClothLodFacade.GetRenderNumBoneInfluences()[VertexID] = ClothLodFacade.GetRenderBoneIndices()[VertexID].Num();
-
-					}, bUseParallel ? EParallelForFlags::None : EParallelForFlags::ForceSingleThread);
-				}
-				else
-				{
-					checkNoEntry();
-				}
+				checkNoEntry();
 			}
 		}
 
diff --git a/Engine/Plugins/Experimental/ChaosClothAssetEditor/Source/ChaosClothAssetDataflowNodes/Public/ChaosClothAsset/CopySimulationToRenderMeshNode.h b/Engine/Plugins/Experimental/ChaosClothAssetEditor/Source/ChaosClothAssetDataflowNodes/Public/ChaosClothAsset/CopySimulationToRenderMeshNode.h
--- a/Engine/Plugins/Experimental/ChaosClothAssetEditor/Source/ChaosClothAssetDataflowNodes/Public/ChaosClothAsset/CopySimulationToRenderMeshNode.h
+++ b/Engine/Plugins/Experimental/ChaosClothAssetEditor/Source/ChaosClothAssetDataflowNodes/Public/ChaosClothAsset/CopySimulationToRenderMeshNode.h
@@ -28,6 +28,10 @@ public:
 	UPROPERTY(EditAnywhere, Category = "Render Mesh")
 	TObjectPtr<const UMaterialInterface> Material;
 
+	/** Generate a single render pattern rather than a render pattern per sim pattern */
+	UPROPERTY(EditAnywhere, Category = "Render Mesh")
+	bool bGenerateSingleRenderPattern = false;
+
 	FChaosClothAssetCopySimulationToRenderMeshNode(const Dataflow::FNodeParameters& InParam, FGuid InGuid = FGuid::NewGuid());
 
 	virtual void Evaluate(Dataflow::FContext& Context, const FDataflowOutput* Out) const override;
diff --git a/Engine/Plugins/Experimental/ChaosClothAssetEditor/Source/ChaosClothAssetDataflowNodes/Public/ChaosClothAsset/ReverseNormalsNode.h b/Engine/Plugins/Experimental/ChaosClothAssetEditor/Source/ChaosClothAssetDataflowNodes/Public/ChaosClothAsset/ReverseNormalsNode.h
--- a/Engine/Plugins/Experimental/ChaosClothAssetEditor/Source/ChaosClothAssetDataflowNodes/Public/ChaosClothAsset/ReverseNormalsNode.h
+++ b/Engine/Plugins/Experimental/ChaosClothAssetEditor/Source/ChaosClothAssetDataflowNodes/Public/ChaosClothAsset/ReverseNormalsNode.h
@@ -18,9 +18,13 @@ public:
 	UPROPERTY(Meta = (Dataflowinput, DataflowOutput, DisplayName = "Collection", DataflowPassthrough = "Collection"))
 	FManagedArrayCollection Collection;
 
-	/** List of patterns to apply the operation on. All patterns will be used if left empty. */
+	/** List of sim patterns to apply the operation on. All patterns will be used if left empty. */
 	UPROPERTY(EditAnywhere, Category = "Reverse Normals", Meta = (EditCondition = "false"))
-	TArray<int32> Patterns;
+	TArray<int32> SimPatterns;
+
+	/** List of render patterns to apply the operation on. All patterns will be used if left empty. */
+	UPROPERTY(EditAnywhere, Category = "Reverse Normals", Meta = (EditCondition = "false"))
+	TArray<int32> RenderPatterns;
 
 	/** Whether to reverse the simulation mesh normals. */
 	UPROPERTY(EditAnywhere, Category = "Reverse Normals")
diff --git a/Engine/Plugins/Experimental/ChaosClothAssetEditor/Source/ChaosClothAssetDataflowNodes/Public/ChaosClothAsset/StaticMeshImportNode.h b/Engine/Plugins/Experimental/ChaosClothAssetEditor/Source/ChaosClothAssetDataflowNodes/Public/ChaosClothAsset/StaticMeshImportNode.h
--- a/Engine/Plugins/Experimental/ChaosClothAssetEditor/Source/ChaosClothAssetDataflowNodes/Public/ChaosClothAsset/StaticMeshImportNode.h
+++ b/Engine/Plugins/Experimental/ChaosClothAssetEditor/Source/ChaosClothAssetDataflowNodes/Public/ChaosClothAsset/StaticMeshImportNode.h
@@ -22,6 +22,10 @@ public:
 	UPROPERTY(EditAnywhere, Category = "Static Mesh Import")
 	TObjectPtr<const UStaticMesh> StaticMesh;
 
+	/* Which static mesh Lod to import.*/
+	UPROPERTY(EditAnywhere, Category = "Static Mesh Import")
+	int32 LodIndex = 0;
+
 	/* Import static mesh data as a simulation mesh data.*/
 	UPROPERTY(EditAnywhere, Category = "Static Mesh Import")
 	bool bImportAsSimMesh = true;
@@ -38,10 +42,6 @@ public:
 	UPROPERTY(EditAnywhere, Category = "Static Mesh Import", Meta = (AllowPreserveRatio, EditCondition = "bImportAsSimMesh && UVChannel != INDEX_NONE"))
 	FVector2f UVScale = { 1.f, 1.f };
 
-	/* Use Sim UV Islands to generate separate patterns (Note: render mesh will always be a single pattern)*/
-	UPROPERTY(EditAnywhere, Category = "Static Mesh Import", meta = (EditCondition = "bImportAsSimMesh && UVChannel != INDEX_NONE"))
-	bool bUseUVIslands = true;
-
 	FChaosClothAssetStaticMeshImportNode(const Dataflow::FNodeParameters& InParam, FGuid InGuid = FGuid::NewGuid());
 
 	virtual void Evaluate(Dataflow::FContext& Context, const FDataflowOutput* Out) const override;
diff --git a/Engine/Plugins/Experimental/ChaosClothAssetEditor/Source/ChaosClothAssetEditor/Private/ChaosClothAsset/ClothEditorMode.cpp b/Engine/Plugins/Experimental/ChaosClothAssetEditor/Source/ChaosClothAssetEditor/Private/ChaosClothAsset/ClothEditorMode.cpp
--- a/Engine/Plugins/Experimental/ChaosClothAssetEditor/Source/ChaosClothAssetEditor/Private/ChaosClothAsset/ClothEditorMode.cpp
+++ b/Engine/Plugins/Experimental/ChaosClothAssetEditor/Source/ChaosClothAssetEditor/Private/ChaosClothAsset/ClothEditorMode.cpp
@@ -290,16 +290,19 @@ void UChaosClothAssetEditorMode::Exit()
 {
 	UActorComponent::MarkRenderStateDirtyEvent.RemoveAll(this);
 
-	for (TObjectPtr<UDynamicMeshComponent> DynamicMeshComp : DynamicMeshComponents)
+	if (DynamicMeshComponent)
 	{
-		DynamicMeshComp->UnregisterComponent();
-		DynamicMeshComp->SelectionOverrideDelegate.Unbind();
+		DynamicMeshComponent->UnregisterComponent();
+		DynamicMeshComponent->SelectionOverrideDelegate.Unbind();
 	}
+	DynamicMeshComponent = nullptr;
+	DynamicMeshComponentParentActor = nullptr;
 
-	for (TObjectPtr<UMeshElementsVisualizer> WireframeDisplay : WireframesToTick)
+	if (WireframeToTick)
 	{
-		WireframeDisplay->Disconnect();
+		WireframeToTick->Disconnect();
 	}
+	WireframeToTick = nullptr;
 
 	if (DataflowComponent)
 	{
@@ -307,8 +310,6 @@ void UChaosClothAssetEditorMode::Exit()
 		DataflowComponent->DestroyComponent();
 	}
 
-	WireframesToTick.Reset();
-	DynamicMeshComponents.Reset();
 	PropertyObjectsToTick.Empty();
 	PreviewScene = nullptr;
 
@@ -392,7 +393,7 @@ void UChaosClothAssetEditorMode::ReinitializeDynamicMeshComponents()
 {
 	using namespace UE::Chaos::ClothAsset;
 
-	auto SetUpDynamicMeshComponentMaterial = [this](const UE::Chaos::ClothAsset::FCollectionClothLodConstFacade& ClothLodFacade, UDynamicMeshComponent& MeshComponent)
+	auto SetUpDynamicMeshComponentMaterial = [this](const UE::Chaos::ClothAsset::FCollectionClothConstFacade& ClothFacade, UDynamicMeshComponent& MeshComponent)
 	{
 		switch (ConstructionViewMode)
 		{
@@ -410,7 +411,7 @@ void UChaosClothAssetEditorMode::ReinitializeDynamicMeshComponents()
 			break;
 			case EClothPatternVertexType::Render:
 			{
-				const TArrayView<const FString> MaterialPaths = ClothLodFacade.GetRenderMaterialPathName();
+				const TArrayView<const FString> MaterialPaths = ClothFacade.GetRenderMaterialPathName();
 				for (int32 MaterialIndex = 0; MaterialIndex < MaterialPaths.Num(); ++MaterialIndex)
 				{
 					const FString& Path = MaterialPaths[MaterialIndex];
@@ -422,36 +423,33 @@ void UChaosClothAssetEditorMode::ReinitializeDynamicMeshComponents()
 		}
 	};
 	
-	// Clean up any existing DynamicMeshComponents
+	// Clean up existing DynamicMeshComponent
 	// Save indices of selected mesh components
-	const int32 TotalNumberExistingDynamicMeshComponents = DynamicMeshComponents.Num();
 
 	USelection* SelectedComponents = GetModeManager()->GetSelectedComponents();
-	for (int32 DynamicMeshComponentIndex = 0; DynamicMeshComponentIndex < DynamicMeshComponents.Num(); ++DynamicMeshComponentIndex)
+
+	if (DynamicMeshComponent)
 	{
 		// TODO: Check for any outstanding changes on the dynamic mesh component?
+		DynamicMeshComponent->UnregisterComponent();
+		DynamicMeshComponent->SelectionOverrideDelegate.Unbind();
 
-		TObjectPtr<UDynamicMeshComponent> DynamicMeshComp = DynamicMeshComponents[DynamicMeshComponentIndex];
-		DynamicMeshComp->UnregisterComponent();
-		DynamicMeshComp->SelectionOverrideDelegate.Unbind();
-
-		if (SelectedComponents->IsSelected(DynamicMeshComp))
+		if (SelectedComponents->IsSelected(DynamicMeshComponent))
 		{
-			SelectedComponents->Deselect(DynamicMeshComp);
-			DynamicMeshComp->PushSelectionToProxy();
+			SelectedComponents->Deselect(DynamicMeshComponent);
+			DynamicMeshComponent->PushSelectionToProxy();
 		}
 	}
 
-	for (const TObjectPtr<UMeshElementsVisualizer> Wireframe : WireframesToTick)
+	if (WireframeToTick)
 	{
-		Wireframe->Disconnect();
+		WireframeToTick->Disconnect();
 	}
 
-	DynamicMeshComponents.Empty();
-	DynamicMeshComponentParentActors.Empty();
 	PropertyObjectsToTick.Empty();	// TODO: We only want to empty the wireframe display properties. Is anything else using this array?
-	WireframesToTick.Empty();
-	DynamicMeshSourceInfos.Empty();
+	DynamicMeshComponent = nullptr;
+	DynamicMeshComponentParentActor = nullptr;
+	WireframeToTick = nullptr;
 
 	TSharedPtr<FManagedArrayCollection> Collection = GetClothCollection();
 	if (!Collection)
@@ -461,164 +459,89 @@ void UChaosClothAssetEditorMode::ReinitializeDynamicMeshComponents()
 
 	const UE::Chaos::ClothAsset::FCollectionClothConstFacade ClothFacade(Collection);
 
-	for (int32 LodIndex = 0; LodIndex < ClothFacade.GetNumLods(); ++LodIndex)
-	{
-		const UE::Chaos::ClothAsset::FCollectionClothLodConstFacade ClothLodFacade = ClothFacade.GetLod(LodIndex);
-
-		if (bCombineAllPatterns)
-		{
-			DynamicMeshComponents.SetNum(1);
-			DynamicMeshComponentParentActors.SetNum(1);
-
-			UE::Geometry::FDynamicMesh3 LodMesh;
-			LodMesh.EnableAttributes();
-			UE::Geometry::FDynamicMeshEditor MeshEditor(&LodMesh);
-
-			for (int32 PatternIndex = 0; PatternIndex < ClothLodFacade.GetNumPatterns(); ++PatternIndex)
-			{
-				UE::Geometry::FDynamicMesh3 PatternMesh;
-				FClothPatternToDynamicMesh Converter;
-				Converter.Convert(Collection, LodIndex, PatternIndex, ConstructionViewMode, PatternMesh);
-
-				if (ConstructionViewMode == EClothPatternVertexType::Sim2D)
-				{
-					// Use per-triangle normals for the 2D view
-					UE::Geometry::FMeshNormals::InitializeMeshToPerTriangleNormals(&PatternMesh);
-				}
-
-				UE::Geometry::FMeshIndexMappings IndexMapping;
-				MeshEditor.AppendMesh(&PatternMesh, IndexMapping);
-			}
-
-			// We only need an actor to allow use of HHitProxy for selection
-			const FRotator Rotation(0.0f, 0.0f, 0.0f);
-			const FActorSpawnParameters SpawnInfo;
-			const TObjectPtr<AActor> ParentActor = this->GetWorld()->SpawnActor<AActor>(FVector::ZeroVector, Rotation, SpawnInfo);
-			DynamicMeshComponentParentActors[0] = ParentActor;
-
-			DynamicMeshComponents[0] = NewObject<UDynamicMeshComponent>(ParentActor);
-			DynamicMeshComponents[0]->SetMesh(MoveTemp(LodMesh));		
-
-			SetUpDynamicMeshComponentMaterial(ClothLodFacade, *DynamicMeshComponents[0]);
-
-			DynamicMeshComponents[0]->SelectionOverrideDelegate = UPrimitiveComponent::FSelectionOverride::CreateUObject(this, &UChaosClothAssetEditorMode::IsComponentSelected);
-			DynamicMeshComponents[0]->RegisterComponentWithWorld(this->GetWorld());
-
-			DynamicMeshSourceInfos.Add(FDynamicMeshSourceInfo{ LodIndex, -1 });
-		}
-		else
-		{
-			const int32 NumComponents = ClothLodFacade.GetNumPatterns();
-			DynamicMeshComponents.SetNum(NumComponents);
-			DynamicMeshComponentParentActors.SetNum(NumComponents);
+	UE::Geometry::FDynamicMesh3 LodMesh;
+	LodMesh.EnableAttributes();
+	FClothPatternToDynamicMesh Converter;
+	Converter.Convert(Collection, INDEX_NONE, ConstructionViewMode, LodMesh);
 
-			for (int32 PatternIndex = 0; PatternIndex < ClothLodFacade.GetNumPatterns(); ++PatternIndex)
-			{
-				UE::Geometry::FDynamicMesh3 PatternMesh;
-				FClothPatternToDynamicMesh Converter;
-				Converter.Convert(Collection, LodIndex, PatternIndex, ConstructionViewMode, PatternMesh);
+	if (ConstructionViewMode == EClothPatternVertexType::Sim2D)
+	{
+		// Use per-triangle normals for the 2D view
+		UE::Geometry::FMeshNormals::InitializeMeshToPerTriangleNormals(&LodMesh);
+	}
 
-				if (ConstructionViewMode == EClothPatternVertexType::Sim2D)
-				{
-					// Use per-triangle normals for the 2D view
-					UE::Geometry::FMeshNormals::InitializeMeshToPerTriangleNormals(&PatternMesh);
-				}
+	// We only need an actor to allow use of HHitProxy for selection
+	const FRotator Rotation(0.0f, 0.0f, 0.0f);
+	const FActorSpawnParameters SpawnInfo;
+	DynamicMeshComponentParentActor = this->GetWorld()->SpawnActor<AActor>(FVector::ZeroVector, Rotation, SpawnInfo);
 
-				// We only need an actor to allow use of HHitProxy for selection
-				const FRotator Rotation(0.0f, 0.0f, 0.0f);
-				const FActorSpawnParameters SpawnInfo;
-				const TObjectPtr<AActor> ParentActor = this->GetWorld()->SpawnActor<AActor>(FVector::ZeroVector, Rotation, SpawnInfo);
-				DynamicMeshComponentParentActors[PatternIndex] = ParentActor;
+	DynamicMeshComponent = NewObject<UDynamicMeshComponent>(DynamicMeshComponentParentActor);
+	DynamicMeshComponent->SetMesh(MoveTemp(LodMesh));
 
-				TObjectPtr<UDynamicMeshComponent> PatternMeshComponent = NewObject<UDynamicMeshComponent>(ParentActor);
-				PatternMeshComponent->SetMesh(MoveTemp(PatternMesh));
+	SetUpDynamicMeshComponentMaterial(ClothFacade, *DynamicMeshComponent);
 
-				SetUpDynamicMeshComponentMaterial(ClothLodFacade, *PatternMeshComponent);
+	DynamicMeshComponent->SelectionOverrideDelegate = UPrimitiveComponent::FSelectionOverride::CreateUObject(this, &UChaosClothAssetEditorMode::IsComponentSelected);
+	DynamicMeshComponent->RegisterComponentWithWorld(this->GetWorld());
 
-				DynamicMeshComponents[PatternIndex]->SelectionOverrideDelegate = UPrimitiveComponent::FSelectionOverride::CreateUObject(this, &UChaosClothAssetEditorMode::IsComponentSelected);
-				DynamicMeshComponents[PatternIndex]->RegisterComponentWithWorld(this->GetWorld());
+	// Set up the wireframe display of the rest space mesh.
+	WireframeToTick = NewObject<UMeshElementsVisualizer>(this);
+	WireframeToTick->CreateInWorld(GetWorld(), FTransform::Identity);
 
-				DynamicMeshComponents[PatternIndex] = PatternMeshComponent;
-				DynamicMeshSourceInfos.Add(FDynamicMeshSourceInfo{ LodIndex, PatternIndex });
-			}
-		}
+	WireframeToTick->Settings->DepthBias = 2.0;
+	WireframeToTick->Settings->bAdjustDepthBiasUsingMeshSize = false;
+	WireframeToTick->Settings->bShowWireframe = true;
+	WireframeToTick->Settings->bShowBorders = true;
+	WireframeToTick->Settings->bShowUVSeams = false;
+	WireframeToTick->Settings->bShowNormalSeams = false;
 
-		// If we found any patterns, use this LOD
-		// TODO: Give the user the ability to specify which LOD to display
-		if (ClothLodFacade.GetNumPatterns() > 0)
-		{
-			break;
-		}
-	}
+	// These are not exposed at the visualizer level yet
+	// TODO: Should they be?
+	WireframeToTick->WireframeComponent->BoundaryEdgeThickness = 2;
 
-	for (TObjectPtr<UDynamicMeshComponent> RestSpaceMeshComponent : DynamicMeshComponents)
+	WireframeToTick->SetMeshAccessFunction([this](UMeshElementsVisualizer::ProcessDynamicMeshFunc ProcessFunc)
 	{
-		// Set up the wireframe display of the rest space mesh.
-		TObjectPtr<UMeshElementsVisualizer> WireframeDisplay = NewObject<UMeshElementsVisualizer>(this);
-		WireframeDisplay->CreateInWorld(GetWorld(), FTransform::Identity);
-
-		WireframeDisplay->Settings->DepthBias = 2.0;
-		WireframeDisplay->Settings->bAdjustDepthBiasUsingMeshSize = false;
-		WireframeDisplay->Settings->bShowWireframe = true;
-		WireframeDisplay->Settings->bShowBorders = true;
-		WireframeDisplay->Settings->bShowUVSeams = false;
-		WireframeDisplay->Settings->bShowNormalSeams = false;
-
-		// These are not exposed at the visualizer level yet
-		// TODO: Should they be?
-		WireframeDisplay->WireframeComponent->BoundaryEdgeThickness = 2;
+		ProcessFunc(*DynamicMeshComponent->GetMesh());
+	});
 
-		WireframeDisplay->SetMeshAccessFunction([RestSpaceMeshComponent](UMeshElementsVisualizer::ProcessDynamicMeshFunc ProcessFunc)
+	DynamicMeshComponent->OnMeshChanged.Add(
+		FSimpleMulticastDelegate::FDelegate::CreateLambda([this]()
 		{
-			ProcessFunc(*RestSpaceMeshComponent->GetMesh());
-		});
-
-		RestSpaceMeshComponent->OnMeshChanged.Add(
-			FSimpleMulticastDelegate::FDelegate::CreateLambda([WireframeDisplay]()
-		{
-			WireframeDisplay->NotifyMeshChanged();
+			WireframeToTick->NotifyMeshChanged();
 		}));
 
-		// The settings object and wireframe are not part of a tool, so they won't get ticked like they
-		// are supposed to (to enable property watching), unless we add this here.
-		PropertyObjectsToTick.Add(WireframeDisplay->Settings);
-		int32 WireframeIndex = WireframesToTick.Add(WireframeDisplay);
+	// The settings object and wireframe are not part of a tool, so they won't get ticked like they
+	// are supposed to (to enable property watching), unless we add this here.
+	PropertyObjectsToTick.Add(WireframeToTick->Settings);
 
-		// Some interactive tools will hide the input DynamicMeshComponent and create their own temporary PreviewMesh for visualization. If this
-		// occurs, we should also hide the corresponding WireframeDisplay (and un-hide it when the tool finishes).
-		UActorComponent::MarkRenderStateDirtyEvent.AddWeakLambda(this, [WireframeIndex, this](UActorComponent& ActorComponent)
+	// Some interactive tools will hide the input DynamicMeshComponent and create their own temporary PreviewMesh for visualization. If this
+	// occurs, we should also hide the corresponding WireframeDisplay (and un-hide it when the tool finishes).
+	UActorComponent::MarkRenderStateDirtyEvent.AddWeakLambda(this, [this](UActorComponent& ActorComponent)
+	{
+		if (!WireframeToTick || !DynamicMeshComponent)
 		{
-			if (WireframeIndex >= WireframesToTick.Num() || WireframeIndex >= DynamicMeshComponents.Num())
-			{
-				return;
-			}
-
-			TObjectPtr<UMeshElementsVisualizer> WireframeDisplay = WireframesToTick[WireframeIndex];
-			const TObjectPtr<UDynamicMeshComponent> RestSpaceMesh = DynamicMeshComponents[WireframeIndex];
-
-			if (!WireframeDisplay || !RestSpaceMesh)
-			{
-				return;
-			}
-
-			const bool bRestSpaceMeshVisible = RestSpaceMesh->GetVisibleFlag();
-			WireframeDisplay->Settings->bVisible = bRestSpaceMeshVisible && bConstructionViewWireframe;
-		});
+			return;
+		}
+		const bool bRestSpaceMeshVisible = DynamicMeshComponent->GetVisibleFlag();
+		WireframeToTick->Settings->bVisible = bRestSpaceMeshVisible && bConstructionViewWireframe;
+	});
 
-		const bool bRestSpaceMeshVisible = RestSpaceMeshComponent->GetVisibleFlag();
-		WireframeDisplay->Settings->bVisible = bRestSpaceMeshVisible && bConstructionViewWireframe;
-	}
+	const bool bRestSpaceMeshVisible = DynamicMeshComponent->GetVisibleFlag();
+	WireframeToTick->Settings->bVisible = bRestSpaceMeshVisible && bConstructionViewWireframe;
 
 	SelectedComponents->DeselectAll();
 	if (ConstructionViewMode != EClothPatternVertexType::Render)
 	{
-		SelectedComponents->BeginBatchSelectOperation();
-		for (TObjectPtr<UDynamicMeshComponent> RestSpaceMeshComponent : DynamicMeshComponents)
-		{
-			SelectedComponents->Select(RestSpaceMeshComponent);
-			RestSpaceMeshComponent->PushSelectionToProxy();
-		}
-		SelectedComponents->EndBatchSelectOperation();
+		SelectedComponents->Select(DynamicMeshComponent);
+		DynamicMeshComponent->PushSelectionToProxy();
+	}
+
+	// Update the context object with the ConstructionViewMode and Collection used to build the DynamicMeshComponents, so 
+	// tools know how to use the components.
+	UEditorInteractiveToolsContext* const RestSpaceToolsContext = GetInteractiveToolsContext();
+	UClothEditorContextObject* EditorContextObject = RestSpaceToolsContext->ContextObjectStore->FindContext<UClothEditorContextObject>();
+	if (ensure(EditorContextObject))
+	{
+		EditorContextObject->SetClothCollection(ConstructionViewMode, Collection);
 	}
 }
 
@@ -755,12 +678,9 @@ void UChaosClothAssetEditorMode::ModeTick(float DeltaTime)
 		}
 	}
 
-	for (TWeakObjectPtr<UMeshElementsVisualizer> WireframeDisplay : WireframesToTick)
+	if (WireframeToTick)
 	{
-		if (WireframeDisplay.IsValid())
-		{
-			WireframeDisplay->OnTick(DeltaTime);
-		}
+		WireframeToTick->OnTick(DeltaTime);
 	}
 
 
@@ -806,20 +726,9 @@ FBox UChaosClothAssetEditorMode::SceneBoundingBox() const
 {
 	FBoxSphereBounds TotalBounds(ForceInitToZero);
 	
-	bool bFoundMeshBounds = false;
-	for (const TObjectPtr<UDynamicMeshComponent> DynamicMeshComponent : DynamicMeshComponents)
+	if (DynamicMeshComponent)
 	{
-		const FBoxSphereBounds& CurrentBounds = DynamicMeshComponent->Bounds;
-
-		if (!bFoundMeshBounds)
-		{
-			TotalBounds = CurrentBounds;
-			bFoundMeshBounds = true;
-		}
-		else
-		{
-			TotalBounds = TotalBounds + CurrentBounds;
-		}
+		TotalBounds = DynamicMeshComponent->Bounds;
 	}
 
 	return TotalBounds.GetBox();
@@ -827,40 +736,15 @@ FBox UChaosClothAssetEditorMode::SceneBoundingBox() const
 
 FBox UChaosClothAssetEditorMode::SelectionBoundingBox() const
 {
-	FBoxSphereBounds TotalBounds;
-
 	const USelection* const SelectedComponents = GetModeManager()->GetSelectedComponents();
 
-	bool bFoundSelected = false;
-	for (const TObjectPtr<UDynamicMeshComponent> DynamicMeshComponent : DynamicMeshComponents)
-	{
-		if (SelectedComponents->IsSelected(DynamicMeshComponent))
-		{
-			const FBoxSphereBounds& CurrentBounds = DynamicMeshComponent->Bounds;
-
-			if (!bFoundSelected)
-			{
-				TotalBounds = CurrentBounds;
-				bFoundSelected = true;
-			}
-			else
-			{
-				TotalBounds = TotalBounds + CurrentBounds;
-			}
-			
-		}
-	}
-
-	if (bFoundSelected)
+	if (DynamicMeshComponent && SelectedComponents->IsSelected(DynamicMeshComponent))
 	{
-		return TotalBounds.GetBox();
+		return DynamicMeshComponent->Bounds.GetBox();
 	}
-	else
-	{
-		// Nothing selected, return the whole scene
-		return SceneBoundingBox();
-	}
-
+	
+	// Nothing selected, return the whole scene
+	return SceneBoundingBox();
 }
 
 
@@ -939,7 +823,7 @@ void UChaosClothAssetEditorMode::InitializeContextObject()
 		RestSpaceToolsContext->ContextObjectStore->AddContextObject(EditorContextObject);
 	}
 
-	EditorContextObject->Init(DataflowGraphEditor, DataflowGraph);
+	EditorContextObject->Init(DataflowGraphEditor, DataflowGraph, ConstructionViewMode, SelectedClothCollection);
 
 	check(EditorContextObject);
 
diff --git a/Engine/Plugins/Experimental/ChaosClothAssetEditor/Source/ChaosClothAssetEditor/Private/ChaosClothAsset/ClothEditorMode.h b/Engine/Plugins/Experimental/ChaosClothAssetEditor/Source/ChaosClothAssetEditor/Private/ChaosClothAsset/ClothEditorMode.h
--- a/Engine/Plugins/Experimental/ChaosClothAssetEditor/Source/ChaosClothAssetEditor/Private/ChaosClothAsset/ClothEditorMode.h
+++ b/Engine/Plugins/Experimental/ChaosClothAssetEditor/Source/ChaosClothAssetEditor/Private/ChaosClothAsset/ClothEditorMode.h
@@ -146,9 +146,9 @@ private:
 
 	bool IsComponentSelected(const UPrimitiveComponent* InComponent);
 
-	// Rest space wireframes. They have to get ticked to be able to respond to setting changes. 
+	// Rest space wireframe. They have to get ticked to be able to respond to setting changes. 
 	UPROPERTY()
-	TArray<TObjectPtr<UMeshElementsVisualizer>> WireframesToTick;
+	TObjectPtr<UMeshElementsVisualizer> WireframeToTick = nullptr;
 
 	// Preview Scene, here largely for convenience to avoid having to pass it around functions. Owned by the ClothEditorToolkit.
 	UE::Chaos::ClothAsset::FChaosClothPreviewScene* PreviewScene = nullptr;
@@ -157,21 +157,13 @@ private:
 	UPROPERTY()
 	TArray<TObjectPtr<UInteractiveToolPropertySet>> PropertyObjectsToTick;
 
-	// Rest space editable meshes
+	// Rest space editable mesh
 	UPROPERTY()
-	TArray<TObjectPtr<UDynamicMeshComponent>> DynamicMeshComponents;
+	TObjectPtr<UDynamicMeshComponent> DynamicMeshComponent = nullptr;
 
-	// Actors required for hit testing DynamicMeshComponents
+	// Actor required for hit testing DynamicMeshComponent
 	UPROPERTY()
-	TArray<TObjectPtr<AActor>> DynamicMeshComponentParentActors;
-
-	// Map back to original asset location for each DynamicMeshComponent
-	struct FDynamicMeshSourceInfo
-	{
-		int32 LodIndex;
-		int32 PatternIndex;
-	};
-	TArray<FDynamicMeshSourceInfo> DynamicMeshSourceInfos;
+	TObjectPtr<AActor> DynamicMeshComponentParentActor = nullptr;
 
 	TWeakPtr<UE::Chaos::ClothAsset::FChaosClothEditorRestSpaceViewportClient, ESPMode::ThreadSafe> RestSpaceViewportClient;
 
@@ -183,10 +175,6 @@ private:
 
 	void RestSpaceViewportResized(FViewport* RestspaceViewport, uint32 Unused);
 
-	// Whether to combine all patterns into a single DynamicMeshComponent, or have separate components for each pattern
-	// TODO: Expose this to the user
-	bool bCombineAllPatterns = true;
-
 	UE::Chaos::ClothAsset::EClothPatternVertexType ConstructionViewMode;
 	bool bCanChangeConstructionViewMode = true;
 
diff --git a/Engine/Plugins/Experimental/ChaosClothAssetEditor/Source/ChaosClothAssetEditor/Private/ChaosClothAsset/ClothEditorToolkit.cpp b/Engine/Plugins/Experimental/ChaosClothAssetEditor/Source/ChaosClothAssetEditor/Private/ChaosClothAsset/ClothEditorToolkit.cpp
--- a/Engine/Plugins/Experimental/ChaosClothAssetEditor/Source/ChaosClothAssetEditor/Private/ChaosClothAsset/ClothEditorToolkit.cpp
+++ b/Engine/Plugins/Experimental/ChaosClothAssetEditor/Source/ChaosClothAssetEditor/Private/ChaosClothAsset/ClothEditorToolkit.cpp
@@ -754,7 +754,11 @@ void FChaosClothAssetEditorToolkit::PopulateOutliner()
 	{
 		if (UChaosClothAsset* const ClothAsset = ClothComponent->GetClothAsset())
 		{
-			ClothCollection = ClothAsset->GetClothCollection();
+			// TODO: handle multiple LODs
+			if (ClothAsset->GetClothCollections().Num())
+			{
+				ClothCollection = ClothAsset->GetClothCollections()[0];
+			}
 		}
 	}
 	
