ClothCollection Schema rework. Existing cloth assets will be emptied, but their associated Dataflow graphs and nodes are still valid (with the exception of AddWeightMapNodes which need to be repainted

diff --git a/Engine/Plugins/Experimental/ChaosClothAssetEditor/Source/ChaosClothAssetEditorTools/Private/ChaosClothAsset/ClothEditorContextObject.cpp b/Engine/Plugins/Experimental/ChaosClothAssetEditor/Source/ChaosClothAssetEditorTools/Private/ChaosClothAsset/ClothEditorContextObject.cpp
--- a/Engine/Plugins/Experimental/ChaosClothAssetEditor/Source/ChaosClothAssetEditorTools/Private/ChaosClothAsset/ClothEditorContextObject.cpp
+++ b/Engine/Plugins/Experimental/ChaosClothAssetEditor/Source/ChaosClothAssetEditorTools/Private/ChaosClothAsset/ClothEditorContextObject.cpp
@@ -8,10 +8,12 @@
 
 #include UE_INLINE_GENERATED_CPP_BY_NAME(ClothEditorContextObject)
 
-void UClothEditorContextObject::Init(TWeakPtr<SDataflowGraphEditor> InDataflowGraphEditor, TObjectPtr<UDataflow> InDataflowGraph)
+void UClothEditorContextObject::Init(TWeakPtr<SDataflowGraphEditor> InDataflowGraphEditor, TObjectPtr<UDataflow> InDataflowGraph, UE::Chaos::ClothAsset::EClothPatternVertexType InConstructionViewMode, TWeakPtr<FManagedArrayCollection> InSelectedClothCollection)
 {
 	DataflowGraphEditor = InDataflowGraphEditor;
 	DataflowGraph = InDataflowGraph;
+	ConstructionViewMode = InConstructionViewMode;
+	SelectedClothCollection = InSelectedClothCollection;
 }
 
 TWeakPtr<SDataflowGraphEditor> UClothEditorContextObject::GetDataflowGraphEditor()
@@ -158,3 +160,8 @@ UEdGraphNode* UClothEditorContextObject::CreateAndConnectNewNode(
 	return NewEdNode;
 }
 
+void UClothEditorContextObject::SetClothCollection(UE::Chaos::ClothAsset::EClothPatternVertexType ViewMode, TWeakPtr<FManagedArrayCollection> ClothCollection)
+{
+	ConstructionViewMode = ViewMode;
+	SelectedClothCollection = ClothCollection;
+}
\ No newline at end of file
diff --git a/Engine/Plugins/Experimental/ChaosClothAssetEditor/Source/ChaosClothAssetEditorTools/Private/ChaosClothAsset/ClothWeightMapPaintTool.cpp b/Engine/Plugins/Experimental/ChaosClothAssetEditor/Source/ChaosClothAssetEditorTools/Private/ChaosClothAsset/ClothWeightMapPaintTool.cpp
--- a/Engine/Plugins/Experimental/ChaosClothAssetEditor/Source/ChaosClothAssetEditorTools/Private/ChaosClothAsset/ClothWeightMapPaintTool.cpp
+++ b/Engine/Plugins/Experimental/ChaosClothAssetEditor/Source/ChaosClothAssetEditorTools/Private/ChaosClothAsset/ClothWeightMapPaintTool.cpp
@@ -31,6 +31,8 @@
 #include "CanvasTypes.h"
 #include "CanvasItem.h"
 #include "ChaosClothAsset/ClothEditorContextObject.h"
+#include "ChaosClothAsset/ClothPatternVertexType.h"
+#include "ChaosClothAsset/CollectionClothFacade.h"
 #include "ContextObjectStore.h"
 #include "Dataflow/DataflowEdNode.h"
 #include "Dataflow/DataflowObject.h"
@@ -288,6 +290,23 @@ void UClothEditorWeightMapPaintTool::Setup()
 	// configure panels
 	UpdateSubToolType(FilterProperties->SubToolType);
 
+	// Setup DynamicMeshToWeight conversion
+	if (ClothEditorContextObject)
+	{
+		if (TSharedPtr<const FManagedArrayCollection> ClothCollection = ClothEditorContextObject->GetSelectedClothCollection().Pin())
+		{
+			using namespace UE::Chaos::ClothAsset;
+			bHaveDynamicMeshToWeightConversion = ClothEditorContextObject->GetConstructionViewMode() == EClothPatternVertexType::Sim2D;
+			if (bHaveDynamicMeshToWeightConversion)
+			{
+				FCollectionClothConstFacade Cloth(ClothCollection);
+				check(Cloth.IsValid());
+				DynamicMeshToWeight = Cloth.GetSimVertex3DLookup();
+				WeightToDynamicMesh = Cloth.GetSimVertex2DLookup();
+			}
+		}
+	}
+
 	PrecomputeFuture.Wait();
 	OctreeFuture.Wait();
 
@@ -302,16 +321,30 @@ void UClothEditorWeightMapPaintTool::Setup()
 	FString CurrentMapName;
 	TArray<float> CurrentWeights;
 	const bool bWeightMapNodeIsSelected = GetSelectedNodeInfo(CurrentMapName, CurrentWeights);
+	const int32 ExpectedWeightsCount = bHaveDynamicMeshToWeightConversion ? WeightToDynamicMesh.Num() : Mesh->VertexCount();
 
-	bShouldUpdateExistingNode = bWeightMapNodeIsSelected && (Mesh->VertexCount() == CurrentWeights.Num());
+	bShouldUpdateExistingNode = bWeightMapNodeIsSelected && (ExpectedWeightsCount == CurrentWeights.Num());
 
 	if (bShouldUpdateExistingNode)
 	{
-		check(Mesh->VertexCount() == CurrentWeights.Num());		// true if bShouldUpdateExistingNode is true
-
-		for (int32 VertexID = 0; VertexID < Mesh->VertexCount(); ++VertexID)
+		if (bHaveDynamicMeshToWeightConversion)
 		{
-			ActiveWeightMap->SetValue(VertexID, &CurrentWeights[VertexID]);
+			check(WeightToDynamicMesh.Num() == CurrentWeights.Num());
+			for (int32 WeightID = 0; WeightID < CurrentWeights.Num(); ++WeightID)
+			{
+				for (const int32 VertexID : WeightToDynamicMesh[WeightID])
+				{
+					ActiveWeightMap->SetValue(VertexID, &CurrentWeights[WeightID]);
+				}
+			}
+		}
+		else
+		{
+			check(Mesh->VertexCount() == CurrentWeights.Num());
+			for (int32 VertexID = 0; VertexID < Mesh->VertexCount(); ++VertexID)
+			{
+				ActiveWeightMap->SetValue(VertexID, &CurrentWeights[VertexID]);
+			}
 		}
 
 		UpdateWeightMapProperties->Name = CurrentMapName;
@@ -697,8 +730,19 @@ bool UClothEditorWeightMapPaintTool::SyncMeshWithWeightBuffer(FDynamicMesh3* Mes
 
 			if (ROIWeightValueBuffer[k] != CurWeight)
 			{
-				ActiveWeightEditChangeTracker->SaveVertexOneRingTriangles(VertIdx, true);
-				ActiveWeightMap->SetValue(VertIdx, &ROIWeightValueBuffer[k]);
+				if (bHaveDynamicMeshToWeightConversion)
+				{
+					for (const int32 Idx : WeightToDynamicMesh[DynamicMeshToWeight[VertIdx]])
+					{
+						ActiveWeightEditChangeTracker->SaveVertexOneRingTriangles(Idx, true);
+						ActiveWeightMap->SetValue(Idx, &ROIWeightValueBuffer[k]);
+					}
+				}
+				else
+				{
+					ActiveWeightEditChangeTracker->SaveVertexOneRingTriangles(VertIdx, true);
+					ActiveWeightMap->SetValue(VertIdx, &ROIWeightValueBuffer[k]);
+				}
 				NumModified++;
 			}
 		}
@@ -951,9 +995,23 @@ void UClothEditorWeightMapPaintTool::ComputeGradient()
 	{
 		TempROIBuffer.Add(vid);
 	}
-	for (int32 vid : TempROIBuffer)
+
+	if (bHaveDynamicMeshToWeightConversion)
 	{
-		ActiveWeightEditChangeTracker->SaveVertexOneRingTriangles(vid, true);
+		for (int32 vid : TempROIBuffer)
+		{
+			for (const int32 Idx : WeightToDynamicMesh[DynamicMeshToWeight[vid]])
+			{
+				ActiveWeightEditChangeTracker->SaveVertexOneRingTriangles(Idx, true);
+			}
+		}
+	}
+	else
+	{
+		for (int32 vid : TempROIBuffer)
+		{
+			ActiveWeightEditChangeTracker->SaveVertexOneRingTriangles(vid, true);
+		}
 	}
 
 
@@ -988,7 +1046,17 @@ void UClothEditorWeightMapPaintTool::ComputeGradient()
 		}
 
 		const float Value = FMath::LerpStable(0.0f, 1.0f, DistanceToLowSq / (DistanceToLowSq + DistanceToHighSq));
-		ActiveWeightMap->SetValue(VertexIndex, &Value);
+		if (bHaveDynamicMeshToWeightConversion)
+		{
+			for (const int32 Idx : WeightToDynamicMesh[DynamicMeshToWeight[VertexIndex]])
+			{
+				ActiveWeightMap->SetValue(Idx, &Value);
+			}
+		}
+		else
+		{
+			ActiveWeightMap->SetValue(VertexIndex, &Value);
+		}
 	}
 
 
@@ -1044,14 +1112,28 @@ void UClothEditorWeightMapPaintTool::SetVerticesToWeightMap(const TSet<int32>& V
 		ApplyVisibilityFilter(TempROIBuffer, VisibleVertices);
 		TempROIBuffer = MoveTemp(VisibleVertices);
 	}
-
-	for (int32 vid : TempROIBuffer)
+	
+	if (bHaveDynamicMeshToWeightConversion)
 	{
-		ActiveWeightEditChangeTracker->SaveVertexOneRingTriangles(vid, true);
+		for (int32 vid : TempROIBuffer)
+		{
+			for (const int32 Idx : WeightToDynamicMesh[DynamicMeshToWeight[vid]])
+			{
+				ActiveWeightEditChangeTracker->SaveVertexOneRingTriangles(Idx, true);
+				ActiveWeightMap->SetValue(Idx, &WeightValue);
+			}
+		}
 	}
-	for (int32 vid : TempROIBuffer)
-	{		
-		ActiveWeightMap->SetValue(vid, &WeightValue);
+	else
+	{
+		for (int32 vid : TempROIBuffer)
+		{
+			ActiveWeightEditChangeTracker->SaveVertexOneRingTriangles(vid, true);
+		}
+		for (int32 vid : TempROIBuffer)
+		{
+			ActiveWeightMap->SetValue(vid, &WeightValue);
+		}
 	}
 
 
@@ -1457,13 +1539,28 @@ void UClothEditorWeightMapPaintTool::FloodFillCurrentWeightAction()
 	{
 		TempROIBuffer.Add(vid);
 	}
-	for (int32 vid : TempROIBuffer)
+
+	if (bHaveDynamicMeshToWeightConversion)
 	{
-		ActiveWeightEditChangeTracker->SaveVertexOneRingTriangles(vid, true);
+		for (int32 vid : TempROIBuffer)
+		{
+			for (const int32 Idx : WeightToDynamicMesh[DynamicMeshToWeight[vid]])
+			{
+				ActiveWeightEditChangeTracker->SaveVertexOneRingTriangles(Idx, true);
+				ActiveWeightMap->SetValue(Idx, &SetWeightValue);
+			}
+		}
 	}
-	for (int32 vid : TempROIBuffer)
+	else
 	{
-		ActiveWeightMap->SetValue(vid, &SetWeightValue);
+		for (int32 vid : TempROIBuffer)
+		{
+			ActiveWeightEditChangeTracker->SaveVertexOneRingTriangles(vid, true);
+		}
+		for (int32 vid : TempROIBuffer)
+		{
+			ActiveWeightMap->SetValue(vid, &SetWeightValue);
+		}
 	}
 
 	DynamicMeshComponent->FastNotifyVertexAttributesUpdated(EMeshRenderAttributeFlags::VertexColors);
@@ -1488,13 +1585,28 @@ void UClothEditorWeightMapPaintTool::ClearAllWeightsAction()
 	{
 		TempROIBuffer.Add(vid);
 	}
-	for (int32 vid : TempROIBuffer)
+
+	if (bHaveDynamicMeshToWeightConversion)
 	{
-		ActiveWeightEditChangeTracker->SaveVertexOneRingTriangles(vid, true);
+		for (int32 vid : TempROIBuffer)
+		{
+			for (const int32 Idx : WeightToDynamicMesh[DynamicMeshToWeight[vid]])
+			{
+				ActiveWeightEditChangeTracker->SaveVertexOneRingTriangles(Idx, true);
+				ActiveWeightMap->SetValue(Idx, &SetWeightValue);
+			}
+		}
 	}
-	for (int32 vid : TempROIBuffer)
+	else
 	{
-		ActiveWeightMap->SetValue(vid, &SetWeightValue);
+		for (int32 vid : TempROIBuffer)
+		{
+			ActiveWeightEditChangeTracker->SaveVertexOneRingTriangles(vid, true);
+		}
+		for (int32 vid : TempROIBuffer)
+		{
+			ActiveWeightMap->SetValue(vid, &SetWeightValue);
+		}
 	}
 
 	DynamicMeshComponent->FastNotifyVertexAttributesUpdated(EMeshRenderAttributeFlags::VertexColors);
@@ -1543,7 +1655,18 @@ void UClothEditorWeightMapPaintTool::AddNewNode(const FString& NewMapName)
 	const TSharedPtr<FDataflowNode> NewDataflowNode = NewDataflowEdNode->GetDataflowNode();
 	FChaosClothAssetAddWeightMapNode* const NewWeightMapNode = NewDataflowNode->AsType<FChaosClothAssetAddWeightMapNode>();
 
-	NewWeightMapNode->VertexWeights = CurrentWeights;
+	if (bHaveDynamicMeshToWeightConversion)
+	{
+		NewWeightMapNode->VertexWeights.Init(0.f, WeightToDynamicMesh.Num());
+		for (int32 DynamicMeshIdx = 0; DynamicMeshIdx < CurrentWeights.Num(); ++DynamicMeshIdx)
+		{
+			NewWeightMapNode->VertexWeights[DynamicMeshToWeight[DynamicMeshIdx]] = CurrentWeights[DynamicMeshIdx];
+		}
+	}
+	else
+	{
+		NewWeightMapNode->VertexWeights = CurrentWeights;
+	}
 	NewWeightMapNode->Name = AddWeightMapProperties->Name;
 	// NewWeightMapNode->SetName(FName(AddWeightMapProperties->NewMapName));   TODO: This doesn't seem to work to set the Node name
 
@@ -1563,7 +1686,18 @@ void UClothEditorWeightMapPaintTool::UpdateSelectedNode()
 	UDataflowEdNode* const SelectedDataflowEdNode = CastChecked<UDataflowEdNode>(SingleSelectedNode);
 	const TSharedPtr<FDataflowNode> DataflowNode = SelectedDataflowEdNode->GetDataflowNode();
 	FChaosClothAssetAddWeightMapNode* const WeightMapNode = DataflowNode->AsType<FChaosClothAssetAddWeightMapNode>();
-	WeightMapNode->VertexWeights = CurrentWeights;
+	if (bHaveDynamicMeshToWeightConversion)
+	{
+		WeightMapNode->VertexWeights.Init(0.f, WeightToDynamicMesh.Num());
+		for (int32 DynamicMeshIdx = 0; DynamicMeshIdx < CurrentWeights.Num(); ++DynamicMeshIdx)
+		{
+			WeightMapNode->VertexWeights[DynamicMeshToWeight[DynamicMeshIdx]] = CurrentWeights[DynamicMeshIdx];
+		}
+	}
+	else
+	{
+		WeightMapNode->VertexWeights = CurrentWeights;
+	}
 	WeightMapNode->Invalidate();
 }
 
diff --git a/Engine/Plugins/Experimental/ChaosClothAssetEditor/Source/ChaosClothAssetEditorTools/Public/ChaosClothAsset/ClothEditorContextObject.h b/Engine/Plugins/Experimental/ChaosClothAssetEditor/Source/ChaosClothAssetEditorTools/Public/ChaosClothAsset/ClothEditorContextObject.h
--- a/Engine/Plugins/Experimental/ChaosClothAssetEditor/Source/ChaosClothAssetEditorTools/Public/ChaosClothAsset/ClothEditorContextObject.h
+++ b/Engine/Plugins/Experimental/ChaosClothAssetEditor/Source/ChaosClothAssetEditorTools/Public/ChaosClothAsset/ClothEditorContextObject.h
@@ -9,6 +9,13 @@ class SDataflowGraphEditor;
 class UDataflow;
 class UEdGraphNode;
 
+namespace UE::Chaos::ClothAsset
+{
+enum class EClothPatternVertexType : uint8;
+}
+struct FManagedArrayCollection;
+
+
 UCLASS()
 class CHAOSCLOTHASSETEDITORTOOLS_API UClothEditorContextObject : public UObject
 {
@@ -16,7 +23,7 @@ class CHAOSCLOTHASSETEDITORTOOLS_API UClothEditorContextObject : public UObject
 
 public:
 
-	void Init(TWeakPtr<SDataflowGraphEditor> DataflowGraphEditor, TObjectPtr<UDataflow> DataflowGraph);
+	void Init(TWeakPtr<SDataflowGraphEditor> DataflowGraphEditor, TObjectPtr<UDataflow> DataflowGraph, UE::Chaos::ClothAsset::EClothPatternVertexType InConstructionViewMode, TWeakPtr<FManagedArrayCollection> SelectedClothCollection);
 
 	TWeakPtr<SDataflowGraphEditor> GetDataflowGraphEditor();
 	const TWeakPtr<const SDataflowGraphEditor> GetDataflowGraphEditor() const;
@@ -63,6 +70,10 @@ public:
 		UEdGraphNode& UpstreamNode,
 		const FName& ConnectionTypeName);
 
+	void SetClothCollection(UE::Chaos::ClothAsset::EClothPatternVertexType ViewMode, TWeakPtr<FManagedArrayCollection> ClothCollection);
+
+	const TWeakPtr<const FManagedArrayCollection> GetSelectedClothCollection() const { return SelectedClothCollection; }
+	UE::Chaos::ClothAsset::EClothPatternVertexType GetConstructionViewMode() const { return ConstructionViewMode; }
 private:
 
 	TWeakPtr<SDataflowGraphEditor> DataflowGraphEditor;
@@ -70,6 +81,8 @@ private:
 	UPROPERTY()
 	TObjectPtr<UDataflow> DataflowGraph = nullptr;
 
+	UE::Chaos::ClothAsset::EClothPatternVertexType ConstructionViewMode;
+	TWeakPtr<const FManagedArrayCollection> SelectedClothCollection;
 };
 
 
diff --git a/Engine/Plugins/Experimental/ChaosClothAssetEditor/Source/ChaosClothAssetTools/Private/ChaosClothAsset/ClothAssetBuilderEditor.cpp b/Engine/Plugins/Experimental/ChaosClothAssetEditor/Source/ChaosClothAssetTools/Private/ChaosClothAsset/ClothAssetBuilderEditor.cpp
--- a/Engine/Plugins/Experimental/ChaosClothAssetEditor/Source/ChaosClothAssetTools/Private/ChaosClothAsset/ClothAssetBuilderEditor.cpp
+++ b/Engine/Plugins/Experimental/ChaosClothAssetEditor/Source/ChaosClothAssetTools/Private/ChaosClothAsset/ClothAssetBuilderEditor.cpp
@@ -34,14 +34,14 @@ void UClothAssetBuilderEditor::BuildLod(FSkeletalMeshLODModel& LODModel, const U
 	// Create a table to remap the LOD materials to the asset materials
 	const TArray<FSkeletalMaterial>& Materials = ClothAsset.GetMaterials();
 
-	const TSharedPtr<const FManagedArrayCollection> ClothCollection = ClothAsset.GetClothCollection();
+	const TSharedPtr<const FManagedArrayCollection> ClothCollection = ClothAsset.GetClothCollections()[LodIndex];
 
 	const FCollectionClothConstFacade ClothFacade(ClothCollection);
-	check(LodIndex < ClothFacade.GetNumLods())
-	const FCollectionClothLodConstFacade ClothLodFacade = ClothFacade.GetLod(LodIndex);
 
-	const int32 NumLodMaterials = ClothLodFacade.GetNumMaterials();
-	const TConstArrayView<FString> LodRenderMaterialPathName = ClothLodFacade.GetRenderMaterialPathName();
+	// Cloth Render Patterns == Skeletal Mesh Sections
+	const int32 NumSections = ClothFacade.GetNumRenderPatterns();
+	const int32 NumLodMaterials = NumSections;
+	const TConstArrayView<FString> LodRenderMaterialPathName = ClothFacade.GetRenderMaterialPathName();
 
 	TArray<int32> LodMaterialMap;
 	LodMaterialMap.AddUninitialized(NumLodMaterials);
@@ -56,26 +56,6 @@ void UClothAssetBuilderEditor::BuildLod(FSkeletalMeshLODModel& LODModel, const U
 		LodMaterialMap[LodMaterialIndex] = SkeletalMaterialIndex;
 	}
 
-	// Build the section/faces map from the LOD patterns
-	TMap<int32, TArray<int32>> SectionFacesMap;
-	SectionFacesMap.Reserve(Materials.Num());
-
-	TArray<uint32> LodRenderIndexRemap;
-
-	const int32 NumRenderFaces = ClothLodFacade.GetNumRenderFaces();
-	const TConstArrayView<int32> LodRenderMaterialIndex = ClothLodFacade.GetRenderMaterialIndex();
-
-	for (int32 RenderFaceIndex = 0; RenderFaceIndex < NumRenderFaces; ++RenderFaceIndex)
-	{
-		const int32 RenderMaterialIndex = LodMaterialMap[LodRenderMaterialIndex[RenderFaceIndex]];
-
-		TArray<int32>& SectionFaces = SectionFacesMap.FindOrAdd(RenderMaterialIndex);
-		SectionFaces.Add(RenderFaceIndex);
-	}
-
-	const int32 NumRenderVertices = ClothLodFacade.GetNumRenderVertices();
-	LodRenderIndexRemap.SetNumUninitialized(NumRenderVertices);
-
 	// Keep track of the active bone indices for this LOD model
 	TSet<FBoneIndexType> ActiveBoneIndices;
 	ActiveBoneIndices.Reserve(ClothAsset.RefSkeleton.GetNum());
@@ -99,21 +79,22 @@ void UClothAssetBuilderEditor::BuildLod(FSkeletalMeshLODModel& LODModel, const U
 		MaxDistances[Index] = 200.f;
 	}
 
+	const int32 NumRenderVertices = ClothFacade.GetNumRenderVertices();
+	LODModel.MeshToImportVertexMap.Reserve(NumRenderVertices);
+
 	// Populate this LOD's sections and the LOD index buffer
-	const int32 NumSections = SectionFacesMap.Num();
 	LODModel.Sections.SetNum(NumSections);
 
-	int32 SectionIndex = 0;
 	int32 BaseIndex = 0;
-	for (const TPair<int32, TArray<int32>>& SectionFaces : SectionFacesMap)
+	for(int32 SectionIndex = 0; SectionIndex < NumSections; ++SectionIndex)
 	{
+		FCollectionClothRenderPatternConstFacade RenderPatternFacade = ClothFacade.GetRenderPattern(SectionIndex);
 		FSkelMeshSection& Section = LODModel.Sections[SectionIndex];
 
-		Section.OriginalDataSectionIndex = SectionIndex++;
+		Section.OriginalDataSectionIndex = SectionIndex;
 
-		const int32 MaterialIndex = SectionFaces.Key;
-		const TArray<int32>& Faces = SectionFaces.Value;
-		const int32 NumFaces = Faces.Num();
+		const int32 MaterialIndex = SectionIndex;
+		const int32 NumFaces = RenderPatternFacade.GetNumRenderFaces();
 		const int32 NumIndices = NumFaces * 3;
 
 		// Build the section face data (indices)
@@ -125,19 +106,17 @@ void UClothAssetBuilderEditor::BuildLod(FSkeletalMeshLODModel& LODModel, const U
 		Section.NumTriangles = (uint32)NumFaces;
 
 		TArray<uint32> Indices;
-		TSet<int32> UniqueIndicesSet;
 		Indices.SetNumUninitialized(NumIndices);
 
-		const TConstArrayView<FIntVector3> LodRenderIndices = ClothLodFacade.GetRenderIndices();
+		const TConstArrayView<FIntVector3> PatternRenderIndices = RenderPatternFacade.GetRenderIndices();
 
 		for (int32 FaceIndex = 0; FaceIndex < NumFaces; ++FaceIndex)
 		{
-			const FIntVector3& RenderIndices = LodRenderIndices[Faces[FaceIndex]];
+			const FIntVector3& RenderIndices = PatternRenderIndices[FaceIndex];
 			for (int32 VertexIndex = 0; VertexIndex < 3; ++VertexIndex)
 			{
 				const int32 RenderIndex = RenderIndices[VertexIndex];
 				Indices[FaceIndex * 3 + VertexIndex] = (uint32)RenderIndex;
-				UniqueIndicesSet.Add(RenderIndex);
 			}
 		}
 
@@ -145,61 +124,62 @@ void UClothAssetBuilderEditor::BuildLod(FSkeletalMeshLODModel& LODModel, const U
 
 		LODModel.IndexBuffer.Append(MoveTemp(Indices));
 
-		// Build the section vertex data from the unique indices
-		TArray<int32> UniqueIndices = UniqueIndicesSet.Array();
-		const int32 NumVertices = UniqueIndices.Num();
+		// Build the section vertex data 
+		const int32 NumVertices = RenderPatternFacade.GetNumRenderVertices();
 
 		Section.SoftVertices.SetNumUninitialized(NumVertices);
 		Section.NumVertices = NumVertices;
 		Section.BaseVertexIndex = LODModel.NumVertices;
 		LODModel.NumVertices += (uint32)NumVertices;
 
+		const int32 PatternVertexOffset = RenderPatternFacade.GetRenderVerticesOffset();
+		checkSlow(Section.BaseVertexIndex == PatternVertexOffset); // Otherwise Section's Indices don't match RenderPattern indices
+
 		// Map reference skeleton bone index to the index in the section's bone map
 		TMap<FBoneIndexType, FBoneIndexType> ReferenceToSectionBoneMap; 
 
 		// Track how many bones we added to the section's bone map so far
 		int CurSectionBoneMapNum = 0; 
 
-		const TConstArrayView<FVector3f> LodRenderPosition = ClothLodFacade.GetRenderPosition();
-		const TConstArrayView<FVector3f> LodRenderTangentU = ClothLodFacade.GetRenderTangentU();
-		const TConstArrayView<FVector3f> LodRenderTangentV = ClothLodFacade.GetRenderTangentV();
-		const TConstArrayView<FVector3f> LodRenderNormal = ClothLodFacade.GetRenderNormal();
-		const TConstArrayView<FLinearColor> LodRenderColor = ClothLodFacade.GetRenderColor();
-		const TConstArrayView<TArray<FVector2f>> LodRenderUVs = ClothLodFacade.GetRenderUVs();
-		const TConstArrayView<int32> LodRenderNumBoneInfluences = ClothLodFacade.GetRenderNumBoneInfluences();
-		const TConstArrayView<TArray<int32>> LodRenderBoneIndices = ClothLodFacade.GetRenderBoneIndices();
-		const TConstArrayView<TArray<float>> LodRenderBoneWeights = ClothLodFacade.GetRenderBoneWeights();
+		const TConstArrayView<FVector3f> PatternRenderPosition = RenderPatternFacade.GetRenderPosition();
+		const TConstArrayView<FVector3f> PatternRenderTangentU = RenderPatternFacade.GetRenderTangentU();
+		const TConstArrayView<FVector3f> PatternRenderTangentV = RenderPatternFacade.GetRenderTangentV();
+		const TConstArrayView<FVector3f> PatternRenderNormal = RenderPatternFacade.GetRenderNormal();
+		const TConstArrayView<FLinearColor> PatternRenderColor = RenderPatternFacade.GetRenderColor();
+		const TConstArrayView<TArray<FVector2f>> PatternRenderUVs = RenderPatternFacade.GetRenderUVs();
+		const TConstArrayView<TArray<int32>> PatternRenderBoneIndices = RenderPatternFacade.GetRenderBoneIndices();
+		const TConstArrayView<TArray<float>> PatternRenderBoneWeights = RenderPatternFacade.GetRenderBoneWeights();
 
 		for (int32 VertexIndex = 0; VertexIndex < NumVertices; ++VertexIndex)
 		{
-			const int32 RenderIndex = UniqueIndices[VertexIndex];
-			LODModel.MaxImportVertex = FMath::Max(LODModel.MaxImportVertex, RenderIndex);
-
-			LodRenderIndexRemap[RenderIndex] = Section.BaseVertexIndex + (uint32)VertexIndex;
+			// Save the original indices for the newly added vertices
+			LODModel.MeshToImportVertexMap.Add(VertexIndex + PatternVertexOffset);
+			LODModel.MaxImportVertex = VertexIndex + PatternVertexOffset;
 
 			FSoftSkinVertex& SoftVertex = Section.SoftVertices[VertexIndex];
 
-			SoftVertex.Position = LodRenderPosition[RenderIndex];
-			SoftVertex.TangentX = LodRenderTangentU[RenderIndex];
-			SoftVertex.TangentY = LodRenderTangentV[RenderIndex];
-			SoftVertex.TangentZ = LodRenderNormal[RenderIndex];
+			SoftVertex.Position = PatternRenderPosition[VertexIndex];
+			SoftVertex.TangentX = PatternRenderTangentU[VertexIndex];
+			SoftVertex.TangentY = PatternRenderTangentV[VertexIndex];
+			SoftVertex.TangentZ = PatternRenderNormal[VertexIndex];
 
 			constexpr bool bSRGB = false; // Avoid linear to srgb conversion
-			SoftVertex.Color = LodRenderColor[RenderIndex].ToFColor(bSRGB);
+			SoftVertex.Color = PatternRenderColor[VertexIndex].ToFColor(bSRGB);
 
-			const TArray<FVector2f>& RenderUVs = LodRenderUVs[RenderIndex];
+			const TArray<FVector2f>& RenderUVs = PatternRenderUVs[VertexIndex];
 			for (int32 TexCoord = 0; TexCoord < FMath::Min(RenderUVs.Num(), (int32)MAX_TEXCOORDS); ++TexCoord)
 			{
 				SoftVertex.UVs[TexCoord] = RenderUVs[TexCoord];
 			}
 
-			const int32 NumInfluences = LodRenderNumBoneInfluences[RenderIndex];
+			const int32 NumInfluences = PatternRenderBoneIndices[VertexIndex].Num();
+			check(NumInfluences <= MAX_TOTAL_INFLUENCES);
 			
 			// Add all of the bones that have non-zero influence to the section's bone map and keep track of the order
 			// that we added the reference bone via CurSectionBoneMapNum
 			for (int32 Influence = 0; Influence < NumInfluences; ++Influence)
 			{
-				const FBoneIndexType InfluenceBone = (FBoneIndexType)LodRenderBoneIndices[RenderIndex][Influence];
+				const FBoneIndexType InfluenceBone = (FBoneIndexType)PatternRenderBoneIndices[VertexIndex][Influence];
 
 				if (ReferenceToSectionBoneMap.Contains(InfluenceBone) == false)
 				{
@@ -211,8 +191,8 @@ void UClothAssetBuilderEditor::BuildLod(FSkeletalMeshLODModel& LODModel, const U
 			int32 Influence = 0;
 			for (;Influence < NumInfluences; ++Influence)
 			{
-				const FBoneIndexType InfluenceBone = (FBoneIndexType)LodRenderBoneIndices[RenderIndex][Influence];
-				const float InWeight = LodRenderBoneWeights[RenderIndex][Influence];
+				const FBoneIndexType InfluenceBone = (FBoneIndexType)PatternRenderBoneIndices[VertexIndex][Influence];
+				const float InWeight = PatternRenderBoneWeights[VertexIndex][Influence];
 				const uint16 InfluenceWeight = static_cast<uint16>(InWeight * static_cast<float>(UE::AnimationCore::MaxRawBoneWeight) + 0.5f);
 				
 				// FSoftSkinVertex::InfluenceBones contain indices into the section's bone map and not the reference
@@ -242,13 +222,6 @@ void UClothAssetBuilderEditor::BuildLod(FSkeletalMeshLODModel& LODModel, const U
 			Section.BoneMap[Pair.Value] = Pair.Key;
 		}
 
-		// Remap the LOD indices with the new vertex indices
-		const TArrayView<uint32> SectionIndexBuffer(LODModel.IndexBuffer.GetData() + Section.BaseIndex, NumIndices);
-		for (uint32& RenderIndex : SectionIndexBuffer)
-		{
-			RenderIndex = LodRenderIndexRemap[RenderIndex];
-		}
-
 		ActiveBoneIndices.Append(Section.BoneMap);
 
 		// Update max bone influences
@@ -278,6 +251,7 @@ void UClothAssetBuilderEditor::BuildLod(FSkeletalMeshLODModel& LODModel, const U
 			}
 			TArray<uint32> SectionRenderIndices;
 			SectionRenderIndices.Reserve(NumIndices);
+			const TArrayView<uint32> SectionIndexBuffer(LODModel.IndexBuffer.GetData() + Section.BaseIndex, NumIndices);
 			for (const uint32 LodModelVertIndex : SectionIndexBuffer)
 			{
 				SectionRenderIndices.Add(LodModelVertIndex - Section.BaseVertexIndex);
@@ -299,8 +273,6 @@ void UClothAssetBuilderEditor::BuildLod(FSkeletalMeshLODModel& LODModel, const U
 				ClothAsset.SkinningKernelRadius);
 		}
 
-		// Save the original indices for the newly added vertices
-		LODModel.MeshToImportVertexMap.Append(MoveTemp(UniqueIndices));
 
 		// Compute the overlapping vertices map (inspired from MeshUtilities::BuildSkeletalMesh)
 		const TArray<FSoftSkinVertex>& SoftVertices = Section.SoftVertices;
diff --git a/Engine/Plugins/Experimental/ChaosClothAssetEditor/Source/ChaosClothAssetTools/Private/ChaosClothAsset/ClothAssetFactory.cpp b/Engine/Plugins/Experimental/ChaosClothAssetEditor/Source/ChaosClothAssetTools/Private/ChaosClothAsset/ClothAssetFactory.cpp
--- a/Engine/Plugins/Experimental/ChaosClothAssetEditor/Source/ChaosClothAssetTools/Private/ChaosClothAsset/ClothAssetFactory.cpp
+++ b/Engine/Plugins/Experimental/ChaosClothAssetEditor/Source/ChaosClothAssetTools/Private/ChaosClothAsset/ClothAssetFactory.cpp
@@ -3,6 +3,7 @@
 #include "ClothAssetFactory.h"
 #include "ChaosClothAsset/ClothAsset.h"
 #include "ChaosClothAsset/CollectionClothFacade.h"
+#include "GeometryCollection/ManagedArrayCollection.h"
 #include "Animation/Skeleton.h"
 
 #include UE_INLINE_GENERATED_CPP_BY_NAME(ClothAssetFactory)
@@ -23,11 +24,11 @@ UObject* UChaosClothAssetFactory::FactoryCreateNew(UClass* Class, UObject* Paren
 	ClothAsset->MarkPackageDirty();
 
 	// Add an empty default LOD, to avoid LOD mismatch with render data
-	TSharedPtr<FManagedArrayCollection> ClothCollection = ClothAsset->GetClothCollection();
+	TArray<TSharedPtr<FManagedArrayCollection>>& ClothCollections = ClothAsset->GetClothCollections();
+	TSharedPtr<FManagedArrayCollection>& ClothCollection = ClothCollections.Emplace_GetRef(MakeShared<FManagedArrayCollection>());
 
 	FCollectionClothFacade ClothFacade(ClothCollection);
 	ClothFacade.DefineSchema();
-	ClothFacade.AddLod();
 
 	// Set the default skeleton on this new LOD and rebuild the static data models (which is done by default with this override)
 	ClothAsset->SetSkeleton(nullptr);
diff --git a/Engine/Plugins/Experimental/ChaosClothAssetEditor/Source/ChaosClothAssetTools/Private/ChaosClothAsset/ClothPatternToDynamicMesh.cpp b/Engine/Plugins/Experimental/ChaosClothAssetEditor/Source/ChaosClothAssetTools/Private/ChaosClothAsset/ClothPatternToDynamicMesh.cpp
--- a/Engine/Plugins/Experimental/ChaosClothAssetEditor/Source/ChaosClothAssetTools/Private/ChaosClothAsset/ClothPatternToDynamicMesh.cpp
+++ b/Engine/Plugins/Experimental/ChaosClothAssetEditor/Source/ChaosClothAssetTools/Private/ChaosClothAsset/ClothPatternToDynamicMesh.cpp
@@ -4,6 +4,7 @@
 
 #include "ChaosClothAsset/ClothPatternToDynamicMesh.h"
 #include "ChaosClothAsset/ClothAsset.h"
+#include "ChaosClothAsset/ClothCollection.h"
 #include "ChaosClothAsset/ClothComponent.h"
 #include "ChaosClothAsset/CollectionClothFacade.h"
 #include "ToDynamicMesh.h"
@@ -29,31 +30,126 @@ public:
 	typedef int32 NormalIDType;
 	typedef int32 ColorIDType;
 
-	FClothPatternWrapper(const FCollectionClothConstFacade& ClothFacade, int32 LodIndex, int32 PatternIndex, EClothPatternVertexType VertexDataType) :
+	FClothPatternWrapper(const FCollectionClothConstFacade& ClothFacade, int32 PatternIndex, EClothPatternVertexType VertexDataType) :
 		VertexDataType(VertexDataType),
-		Pattern(ClothFacade.GetLod(LodIndex).GetPattern(PatternIndex))
+		Cloth(ClothFacade)
 	{
-		const int32 NumVertices = (VertexDataType == EClothPatternVertexType::Render) ? Pattern.GetNumRenderVertices() : Pattern.GetNumSimVertices();
-		VertIDs.SetNum(NumVertices);
-		for (int32 VtxIndex = 0; VtxIndex < NumVertices; ++VtxIndex)
+		int32 NumFaces = 0;
+		int32 FaceOffset = 0;
+		if (PatternIndex == INDEX_NONE)
 		{
-			VertIDs[VtxIndex] = VtxIndex;
+			// All patterns in one dynamic mesh
+			switch (VertexDataType)
+			{
+			case EClothPatternVertexType::Render:
+			{
+				const int32 NumVertices = Cloth.GetNumRenderVertices();
+				VertIDs.SetNum(NumVertices);
+				for (int32 VtxIndex = 0; VtxIndex < NumVertices; ++VtxIndex)
+				{
+					VertIDs[VtxIndex] = VtxIndex;
+				}
+				NormalIDs = VertIDs;
+
+				NumFaces = Cloth.GetNumRenderFaces();
+				FaceOffset = 0;
+			} break;
+			case EClothPatternVertexType::Sim2D:
+			{
+				const int32 NumVertices = Cloth.GetNumSimVertices2D();
+				VertIDs.SetNum(NumVertices);
+				for (int32 VtxIndex = 0; VtxIndex < NumVertices; ++VtxIndex)
+				{
+					VertIDs[VtxIndex] = VtxIndex;
+				}
+				NormalIDs = Cloth.GetSimVertex3DLookup();
+
+				NumFaces = Cloth.GetNumSimFaces();
+				FaceOffset = 0;
+			}
+			break;
+			case EClothPatternVertexType::Sim3D:
+			{
+				const int32 NumVertices = Cloth.GetNumSimVertices3D();
+				VertIDs.SetNum(NumVertices);
+				for (int32 VtxIndex = 0; VtxIndex < NumVertices; ++VtxIndex)
+				{
+					VertIDs[VtxIndex] = VtxIndex;
+				}
+				NormalIDs = VertIDs;
+
+				NumFaces = Cloth.GetNumSimFaces();
+				FaceOffset = 0;
+			}break;
+			default:
+				checkNoEntry();
+			}
+		}
+		else
+		{
+			switch (VertexDataType)
+			{
+			case EClothPatternVertexType::Render:
+			{
+				FCollectionClothRenderPatternConstFacade Pattern = Cloth.GetRenderPattern(PatternIndex);
+				const int32 NumVertices = Pattern.GetNumRenderVertices();
+				const int32 VertexOffset = Pattern.GetRenderVerticesOffset();
+				VertIDs.SetNum(NumVertices);
+				for (int32 VtxIndex = 0; VtxIndex < NumVertices; ++VtxIndex)
+				{
+					VertIDs[VtxIndex] = VtxIndex + VertexOffset;
+				}
+				NormalIDs = VertIDs;
+
+				NumFaces = Pattern.GetNumRenderFaces();
+				FaceOffset = Pattern.GetRenderFacesOffset();
+			} break;
+			case EClothPatternVertexType::Sim2D:
+			{
+				FCollectionClothSimPatternConstFacade Pattern = Cloth.GetSimPattern(PatternIndex);
+				const int32 NumVertices = Pattern.GetNumSimVertices2D();
+				const int32 VertexOffset = Pattern.GetSimVertices2DOffset();
+				VertIDs.SetNum(NumVertices);
+				for (int32 VtxIndex = 0; VtxIndex < NumVertices; ++VtxIndex)
+				{
+					VertIDs[VtxIndex] = VtxIndex + VertexOffset;
+				}
+				NormalIDs = Pattern.GetSimVertex3DLookup();
+
+				NumFaces = Pattern.GetNumSimFaces();
+				FaceOffset = Pattern.GetSimFacesOffset();
+			} break;
+			case EClothPatternVertexType::Sim3D:
+			{
+				FCollectionClothSimPatternConstFacade Pattern = Cloth.GetSimPattern(PatternIndex);
+				VertIDs = Pattern.GetSimVertex3DLookup();
+				NormalIDs = Pattern.GetSimVertex3DLookup();
+
+				NumFaces = Pattern.GetNumSimFaces();
+				FaceOffset = Pattern.GetSimFacesOffset();
+			} break;
+			default:
+				checkNoEntry();
+			}
 		}
 
-		const int32 NumFaces = (VertexDataType == EClothPatternVertexType::Render) ? Pattern.GetNumRenderFaces() : Pattern.GetNumSimFaces();
 		TriIDs.SetNum(NumFaces);
 		for (int32 TriIndex = 0; TriIndex < NumFaces; ++TriIndex)
 		{
-			TriIDs[TriIndex] = TriIndex;
+			TriIDs[TriIndex] = TriIndex + FaceOffset;;
 		}
 
-		//
-		// Weight map layers precomputation
-		//
-		WeightMapNames = ClothFacade.GetWeightMapNames();
+		if(VertexDataType != EClothPatternVertexType::Render)
+		{
+			//
+			// Weight map layers precomputation
+			//
+			// No weight maps for render mesh
+			WeightMapNames = Cloth.GetWeightMapNames();
+		}
 
 		// Set the reference skeleton if available
-		const FString& SkeletonPathName = ClothFacade.GetLod(LodIndex).GetSkeletonAssetPathName();
+		const FString& SkeletonPathName = Cloth.GetSkeletonAssetPathName();
 		const USkeleton* const Skeleton = !SkeletonPathName.IsEmpty() ?
 			LoadObject<USkeleton>(nullptr, *SkeletonPathName, nullptr, LOAD_None, nullptr) :
 			nullptr;
@@ -86,9 +182,11 @@ public:
 		return WeightMapNames[LayerIndex];
 	}
 
-	float GetVertexWeight(int32 LayerIndex, int32 VertexIndex) const
+	float GetVertexWeight(int32 LayerIndex, VertIDType VertexIndex) const
 	{
-		return Pattern.GetWeightMap(WeightMapNames[LayerIndex])[VertexIndex];
+		// All weights live on 3D indices.
+		const int32 VertexWeightIndex = (VertexDataType == EClothPatternVertexType::Sim2D) ? Cloth.GetSimVertex3DLookup()[VertexIndex] : VertexIndex;
+		return Cloth.GetWeightMap(WeightMapNames[LayerIndex])[VertexWeightIndex];
 	}
 
 	// --"Vertex Buffer" info
@@ -103,19 +201,19 @@ public:
 		{
 			case EClothPatternVertexType::Render:
 			{
-				const TConstArrayView<FVector3f> RenderPositions = Pattern.GetRenderPosition();
+				const TConstArrayView<FVector3f> RenderPositions = Cloth.GetRenderPosition();
 				return FVector3d(RenderPositions[VtxID]);
 			}
 			case EClothPatternVertexType::Sim2D:
 			{
-				const TConstArrayView<FVector2f> SimPositions = Pattern.GetSimPosition();
+				const TConstArrayView<FVector2f> SimPositions = Cloth.GetSimPosition2D();
 				const FVector2f& Pos = SimPositions[VtxID];
 				return FVector3d(Pos[0], Pos[1], 0.0);
 			}
 			case EClothPatternVertexType::Sim3D:
 			{
-				const TConstArrayView<FVector3f> RestPositions = Pattern.GetSimRestPosition();
-				const FVector3f& Pos = RestPositions[VtxID];
+				const TConstArrayView<FVector3f> SimPositions = Cloth.GetSimPosition3D();
+				const FVector3f& Pos = SimPositions[VtxID];
 				return FVector3d(Pos[0], Pos[1], Pos[2]);
 			}
 			default:
@@ -132,23 +230,26 @@ public:
 
 	bool GetTri(TriIDType TriID, VertIDType& VID0, VertIDType& VID1, VertIDType& VID2) const
 	{
-		FIntVector Face;		// Indices are using an LOD based indexation
-		int32 VerticesOffset;	// and need to have the pattern offset removed
+		FIntVector Face;
 
-		if (VertexDataType == EClothPatternVertexType::Render)
-		{
-			Face = Pattern.GetRenderIndices()[TriID];
-			VerticesOffset = Pattern.GetRenderVerticesOffset();
-		}
-		else
+		switch (VertexDataType)
 		{
-			Face = Pattern.GetSimIndices()[TriID];
-			VerticesOffset = Pattern.GetSimVerticesOffset();
+		case EClothPatternVertexType::Render:
+			Face = Cloth.GetRenderIndices()[TriID];
+			break;
+		case EClothPatternVertexType::Sim2D:
+			Face = Cloth.GetSimIndices2D()[TriID];
+			break;
+		case EClothPatternVertexType::Sim3D:
+			Face = Cloth.GetSimIndices3D()[TriID];
+			break;
+		default:
+			checkNoEntry();
 		}
 
-		VID0 = Face[0] - VerticesOffset;
-		VID1 = Face[1] - VerticesOffset;
-		VID2 = Face[2] - VerticesOffset;
+		VID0 = Face[0];
+		VID1 = Face[1];
+		VID2 = Face[2];
 
 		return true;
 	}
@@ -221,22 +322,25 @@ public:
 		return 1;
 	}
 
-	UE::AnimationCore::FBoneWeights GetVertexSkinWeight(int32 SkinWeightAttributeIndex, VertIDType VertexID) const 
+	UE::AnimationCore::FBoneWeights GetVertexSkinWeight(int32 SkinWeightAttributeIndex, VertIDType InVertexID) const 
 	{
 		using namespace UE::AnimationCore;
 		
 		checkfSlow(SkinWeightAttributeIndex == 0, TEXT("Cloth assets should only have one skin weight profile")); 
 
 		const bool bGetRenderMeshData = (VertexDataType == EClothPatternVertexType::Render);
-		const TConstArrayView<int32> NumBoneInfluences = bGetRenderMeshData ? Pattern.GetRenderNumBoneInfluences() : Pattern.GetSimNumBoneInfluences();
-		const TConstArrayView<TArray<int32>> BoneIndices = bGetRenderMeshData ? Pattern.GetRenderBoneIndices() : Pattern.GetSimBoneIndices();
-		const TConstArrayView<TArray<float>> BoneWeights = bGetRenderMeshData ? Pattern.GetRenderBoneWeights() : Pattern.GetSimBoneWeights();
+		const TConstArrayView<TArray<int32>> BoneIndices = bGetRenderMeshData ? Cloth.GetRenderBoneIndices() : Cloth.GetSimBoneIndices();
+		const TConstArrayView<TArray<float>> BoneWeights = bGetRenderMeshData ? Cloth.GetRenderBoneWeights() : Cloth.GetSimBoneWeights();
+
+		// All weights live on 3D indices. Need to convert 2D index to 3D index.
+		const VertIDType VertexID = (VertexDataType == EClothPatternVertexType::Sim2D) ? Cloth.GetSimVertex3DLookup()[InVertexID] : InVertexID;
 
-		if (ensure(VertexID >= 0 && VertexID < NumBoneInfluences.Num()))
+		if (ensure(VertexID >= 0 && VertexID < BoneIndices.Num()))
 		{
-			const int32 NumInfluences = NumBoneInfluences[VertexID];
 			const TArray<int32> Indices = BoneIndices[VertexID];
 			const TArray<float> Weights = BoneWeights[VertexID];
+			const int32 NumInfluences = Indices.Num();
+			check(Weights.Num() == NumInfluences);
 			
 			TArray<FBoneWeight> BoneWeightArray;
 			BoneWeightArray.SetNumUninitialized(NumInfluences);
@@ -303,28 +407,42 @@ public:
 
 	const TArray<int32>& GetNormalIDs() const 
 	{ 
-		return VertIDs; 
+		return NormalIDs; 
 	}
 
 	FVector3f GetNormal(NormalIDType ID) const 
 	{ 
-		return (VertexDataType == EClothPatternVertexType::Render) ? Pattern.GetRenderNormal()[ID] : Pattern.GetSimRestNormal()[ID];
+		return (VertexDataType == EClothPatternVertexType::Render) ? Cloth.GetRenderNormal()[ID] : Cloth.GetSimNormal()[ID];
 	}
 
 	bool GetNormalTri(const TriIDType& TriID, NormalIDType& ID0, NormalIDType& ID1, NormalIDType& ID2) const 
 	{ 
-		return GetTri(TriID, ID0, ID1, ID2); 
+		if (VertexDataType == EClothPatternVertexType::Sim2D)
+		{
+			// All normal data lives on 3D indices.
+			const FIntVector Face = Cloth.GetSimIndices3D()[TriID];
+
+			ID0 = Face[0];
+			ID1 = Face[1];
+			ID2 = Face[2];
+
+			return true;
+		}
+		else
+		{
+			return GetTri(TriID, ID0, ID1, ID2);
+		}
 	}
 	
 	const TArray<int32>& GetUVIDs(int32 LayerID) const 
 	{ 
-		return (VertexDataType == EClothPatternVertexType::Render) ? VertIDs : EmptyArray;
+		return (VertexDataType == EClothPatternVertexType::Render) ? NormalIDs : EmptyArray;
 	}
 
 	FVector2f GetUV(int32 LayerID, UVIDType UVID) const 
 	{ 
 		checkf(VertexDataType == EClothPatternVertexType::Render, TEXT("Requested UVs from a Sim mesh"));
-		const TConstArrayView<FVector2f> VertexUVs = Pattern.GetRenderUVs()[UVID];
+		const TConstArrayView<FVector2f> VertexUVs = Cloth.GetRenderUVs()[UVID];
 		return VertexUVs[LayerID];
 	}
 
@@ -335,76 +453,76 @@ public:
 	
 	const TArray<int32>& GetTangentIDs() const 
 	{ 
-		return (VertexDataType == EClothPatternVertexType::Render) ? VertIDs : EmptyArray;
+		return (VertexDataType == EClothPatternVertexType::Render) ? NormalIDs : EmptyArray;
 	}
 
 	FVector3f GetTangent(NormalIDType ID) const 
 	{ 
 		checkf(VertexDataType == EClothPatternVertexType::Render, TEXT("Requested Tangent from a Sim mesh"));
-		return Pattern.GetRenderTangentU()[ID];
+		return Cloth.GetRenderTangentU()[ID];
 	}
 
 	bool GetTangentTri(const TriIDType& TriID, NormalIDType& ID0, NormalIDType& ID1, NormalIDType& ID2) const 
 	{ 
-		return GetTri(TriID, ID0, ID1, ID2);
+		return GetNormalTri(TriID, ID0, ID1, ID2);
 	}
 
 	const TArray<int32>& GetBiTangentIDs() const 
 	{ 
-		return (VertexDataType == EClothPatternVertexType::Render) ? VertIDs : EmptyArray;
+		return (VertexDataType == EClothPatternVertexType::Render) ? NormalIDs : EmptyArray;
 	}
 
 	FVector3f GetBiTangent(NormalIDType ID) const 
 	{
 		checkf(VertexDataType == EClothPatternVertexType::Render, TEXT("Requested Bitangent from a Sim mesh"));
-		return Pattern.GetRenderTangentV()[ID];
+		return Cloth.GetRenderTangentV()[ID];
 	}
 
 	bool GetBiTangentTri(const TriIDType& TriID, NormalIDType& ID0, NormalIDType& ID1, NormalIDType& ID2) const
 	{
-		return GetTri(TriID, ID0, ID1, ID2);
+		return GetNormalTri(TriID, ID0, ID1, ID2);
 	}
 
 	const TArray<int32>& GetColorIDs() const 
 	{ 
-		return (VertexDataType == EClothPatternVertexType::Render) ? VertIDs : EmptyArray;
+		return (VertexDataType == EClothPatternVertexType::Render) ? NormalIDs : EmptyArray;
 	}
 
 	FVector4f GetColor(ColorIDType VID) const 
 	{
 		checkf(VertexDataType == EClothPatternVertexType::Render, TEXT("Requested color from a Sim mesh"));
-		return Pattern.GetRenderColor()[VID];
+		return Cloth.GetRenderColor()[VID];
 	}
 
 	bool GetColorTri(const TriIDType& TriID, ColorIDType& ID0, ColorIDType& ID1, ColorIDType& ID2) const 
 	{
-		return GetTri(TriID, ID0, ID1, ID2);
+		return GetNormalTri(TriID, ID0, ID1, ID2);
 	}
 
 private:
 
 	const EClothPatternVertexType VertexDataType;
+	const FCollectionClothConstFacade& Cloth;
 
-	TArray<TriIDType> TriIDs;		// indices into Pattern.GetSimFaces()
-	TArray<VertIDType> VertIDs;		// indices into Pattern.GetSimVertices()
+	TArray<TriIDType> TriIDs;		// indices into Indices
+	TArray<VertIDType> VertIDs;		// indices into Positions2D or Positions3D
+	TArray<NormalIDType> NormalIDs; // indices into Normals (and all per vertex data other than positions)
 	
 	TArray<FName> WeightMapNames;
 
-	const FCollectionClothPatternConstFacade Pattern;
-
 	TArray<int32> EmptyArray;
 
 	const FReferenceSkeleton* RefSkeleton = nullptr;
 };
 
 
-void FClothPatternToDynamicMesh::Convert(const TSharedPtr<const FManagedArrayCollection> ClothCollection, int32 LODIndex, int32 PatternIndex, EClothPatternVertexType VertexDataType, UE::Geometry::FDynamicMesh3& MeshOut)
+void FClothPatternToDynamicMesh::Convert(const TSharedPtr<const FManagedArrayCollection> ClothCollection, int32 PatternIndex, EClothPatternVertexType VertexDataType, UE::Geometry::FDynamicMesh3& MeshOut)
 {
 	const FCollectionClothConstFacade ClothFacade(ClothCollection);
 
 	// Actual conversion
 	UE::Geometry::TToDynamicMesh<FClothPatternWrapper> PatternToDynamicMesh; 
-	FClothPatternWrapper PatternWrapper(ClothFacade, LODIndex, PatternIndex, VertexDataType);
+	FClothPatternWrapper PatternWrapper(ClothFacade, PatternIndex, VertexDataType);
 
 	const bool bDisableAttributes = false;
 	auto TriangleToGroupFunction = [](FClothPatternWrapper::TriIDType) { return 0; };
@@ -418,14 +536,20 @@ void FClothPatternToDynamicMesh::Convert(const TSharedPtr<const FManagedArrayCol
 	{
 		MeshOut.EnableAttributes();
 
-		const FCollectionClothLodConstFacade LodFacade(ClothFacade.GetLod(LODIndex));
-		const TArrayView<const int32> RenderMaterialIndex = LodFacade.GetRenderMaterialIndex();
-
-		auto TriangleToMaterialFunction = [&RenderMaterialIndex, VertexDataType](FClothPatternWrapper::TriIDType TriID)->int32
+		const bool bIsRenderType = VertexDataType == EClothPatternVertexType::Render;
+		auto TriangleToMaterialFunction = [PatternIndex, bIsRenderType, &ClothFacade](FClothPatternWrapper::TriIDType TriID)->int32
 		{
-			if (VertexDataType == EClothPatternVertexType::Render && ensure(RenderMaterialIndex.IsValidIndex(TriID)))
+			if (bIsRenderType)
 			{
-				return RenderMaterialIndex[TriID];
+				if (PatternIndex != INDEX_NONE)
+				{
+					return PatternIndex;
+				}
+				const int32 FoundPattern = ClothFacade.FindRenderPatternByFaceIndex(TriID);
+				if (ensure(FoundPattern != INDEX_NONE))
+				{
+					return FoundPattern;
+				}
 			}
 			return INDEX_NONE;
 		};
@@ -438,10 +562,11 @@ void FClothPatternToDynamicMesh::Convert(const TSharedPtr<const FManagedArrayCol
 
 void FClothPatternToDynamicMesh::Convert(const UChaosClothAsset* ClothAssetMeshIn, int32 LODIndex, int32 PatternIndex, EClothPatternVertexType VertexDataType, UE::Geometry::FDynamicMesh3& MeshOut)
 {
-	const TSharedPtr<const FManagedArrayCollection> ClothCollection = ClothAssetMeshIn->GetClothCollection();
-	check(ClothCollection.IsValid());
+	const TArray<TSharedPtr<const FManagedArrayCollection>>& ClothCollections = ClothAssetMeshIn->GetClothCollections();
+	check(ClothCollections.IsValidIndex(LODIndex));
+	check(ClothCollections[LODIndex].IsValid());
 
-	Convert(ClothCollection, LODIndex, PatternIndex, VertexDataType, MeshOut);
+	Convert(ClothCollections[LODIndex], PatternIndex, VertexDataType, MeshOut);
 }
 
 }	// namespace UE::Chaos::ClothAsset
diff --git a/Engine/Plugins/Experimental/ChaosClothAssetEditor/Source/ChaosClothAssetTools/Private/ChaosClothAsset/ClothingAssetToClothAssetExporter.cpp b/Engine/Plugins/Experimental/ChaosClothAssetEditor/Source/ChaosClothAssetTools/Private/ChaosClothAsset/ClothingAssetToClothAssetExporter.cpp
--- a/Engine/Plugins/Experimental/ChaosClothAssetEditor/Source/ChaosClothAssetTools/Private/ChaosClothAsset/ClothingAssetToClothAssetExporter.cpp
+++ b/Engine/Plugins/Experimental/ChaosClothAssetEditor/Source/ChaosClothAssetTools/Private/ChaosClothAsset/ClothingAssetToClothAssetExporter.cpp
@@ -3,9 +3,15 @@
 #include "ClothingAssetToClothAssetExporter.h"
 #include "ChaosClothAsset/ClothAsset.h"
 #include "ChaosClothAsset/CollectionClothFacade.h"
+#include "ChaosClothAsset/ClothGeometryTools.h"
 #include "Animation/Skeleton.h"
+#include "DynamicMesh/DynamicMesh3.h"
+#include "DynamicMesh/DynamicMeshAttributeSet.h"
+#include "DynamicMesh/NonManifoldMappingSupport.h"
+#include "ToDynamicMesh.h"
 #include "Engine/SkinnedAssetCommon.h"
 #include "Engine/SkeletalMesh.h"
+#include "GeometryCollection/ManagedArraycollection.h"
 #include "Materials/Material.h"
 #include "Misc/MessageDialog.h"
 #include "ClothingAsset.h"
@@ -14,6 +20,61 @@
 
 #define LOCTEXT_NAMESPACE "ClothingAssetToClothAssetExporter"
 
+namespace UE::Chaos::ClothAsset::Private
+{
+struct FSimpleSrcMeshInterface
+{
+	typedef int32     VertIDType;
+	typedef int32     TriIDType;
+
+	FSimpleSrcMeshInterface(const TArray<FVector3f>& InPositions, const TArray<uint32>& InIndices)
+		:Positions(InPositions), Indices(InIndices)
+	{
+		VertIDs.SetNumUninitialized(Positions.Num());
+		for (int32 VtxIndex = 0; VtxIndex < Positions.Num(); ++VtxIndex)
+		{
+			VertIDs[VtxIndex] = VtxIndex;
+		}
+
+		check(Indices.Num() % 3 == 0);
+		const int32 NumFaces = Indices.Num() / 3;
+		TriIDs.SetNumUninitialized(NumFaces);
+		for (int32 TriIndex = 0; TriIndex < NumFaces; ++TriIndex)
+		{
+			TriIDs[TriIndex] = 3 * TriIndex;
+		}
+	}
+
+	// accounting.
+	int32 NumTris() const { return TriIDs.Num(); }
+	int32 NumVerts() const { return VertIDs.Num(); }
+
+	// --"Vertex Buffer" info
+	const TArray<VertIDType>& GetVertIDs() const { return VertIDs; }
+	const FVector GetPosition(const VertIDType VtxID) const { return FVector(Positions[VtxID]); }
+
+	// --"Index Buffer" info
+	const TArray<TriIDType>& GetTriIDs() const { return TriIDs; }
+	// return false if this TriID is not contained in mesh.
+	bool GetTri(const TriIDType TriID, VertIDType& VID0, VertIDType& VID1, VertIDType& VID2) const
+	{
+		VID0 = Indices[TriID + 0];
+		VID1 = Indices[TriID + 1];
+		VID2 = Indices[TriID + 2];
+
+		return true;
+	}
+
+private:
+	const TArray<FVector3f>& Positions;
+	const TArray<uint32>& Indices;
+
+	TArray<TriIDType> TriIDs; // TriID = first index in flat Indices array
+	TArray<VertIDType> VertIDs;
+};
+} // namespace UE::Chaos::ClothAsset::Private
+
+
 UClass* UClothingAssetToChaosClothAssetExporter::GetExportedType() const
 {
 	return UChaosClothAsset::StaticClass();
@@ -34,10 +95,10 @@ void UClothingAssetToChaosClothAssetExporter::Export(const UClothingAssetBase* C
 
 	UChaosClothAsset* const ClothAsset = CastChecked<UChaosClothAsset>(ExportedAsset);
 	check(ClothAsset);
-	FCollectionClothFacade Cloth(ClothAsset->GetClothCollection());
-
+	TArray<TSharedPtr<FManagedArrayCollection>>& ClothCollections = ClothAsset->GetClothCollections();
+	
 	const int32 NumLods = ClothingAssetCommon->LodData.Num();
-	Cloth.SetNumLods(NumLods);
+	ClothCollections.SetNum(NumLods);
 
 	// Create the LODs
 	for (int32 LodIndex = 0; LodIndex < NumLods; ++LodIndex)
@@ -45,9 +106,26 @@ void UClothingAssetToChaosClothAssetExporter::Export(const UClothingAssetBase* C
 		const FClothLODDataCommon& ClothLODData = ClothingAssetCommon->LodData[LodIndex];
 		const FClothPhysicalMeshData& PhysicalMeshData = ClothLODData.PhysicalMeshData;
 
-		// Unwrap the physical mesh data into the pattern and rest meshes
-		FCollectionClothLodFacade ClothLod = Cloth.GetLod(LodIndex);
-		ClothLod.Initialize(PhysicalMeshData.Vertices, PhysicalMeshData.Indices);
+		TSharedPtr<FManagedArrayCollection>& ClothCollection = ClothCollections[LodIndex];
+		if (!ClothCollection.IsValid())
+		{
+			ClothCollection = MakeShared<FManagedArrayCollection>();
+			FCollectionClothFacade Cloth(ClothCollection);
+			Cloth.DefineSchema();
+		}
+
+		using namespace UE::Chaos::ClothAsset::Private;
+
+		// Build a DynamicMesh from Positions and Indices
+		UE::Geometry::TToDynamicMeshBase<FSimpleSrcMeshInterface> ToDynamicMesh;
+		FSimpleSrcMeshInterface SimpleSrc(PhysicalMeshData.Vertices, PhysicalMeshData.Indices);
+
+		UE::Geometry::FDynamicMesh3 DynamicMesh;
+		ToDynamicMesh.Convert(DynamicMesh, SimpleSrc, [](FSimpleSrcMeshInterface::TriIDType) {return 0; });
+		UE::Geometry::FNonManifoldMappingSupport::AttachNonManifoldVertexMappingData(ToDynamicMesh.ToSrcVertIDMap, DynamicMesh);
+
+		constexpr int32 UVChannelIndexNone = INDEX_NONE;
+		FClothGeometryTools::BuildSimMeshFromDynamicMesh(ClothCollection, DynamicMesh, UVChannelIndexNone, FVector2f(1.f));
 	}
 
 	if (NumLods)
@@ -58,7 +136,14 @@ void UClothingAssetToChaosClothAssetExporter::Export(const UClothingAssetBase* C
 	else
 	{
 		// Make sure that at least one empty LOD is always created
-		Cloth.AddLod();
+		ClothCollections.SetNum(1);
+		TSharedPtr<FManagedArrayCollection>& ClothCollection = ClothCollections[1];
+		if (!ClothCollection.IsValid())
+		{
+			ClothCollection = MakeShared<FManagedArrayCollection>();
+			FCollectionClothFacade Cloth(ClothCollection);
+			Cloth.DefineSchema();
+		}
 	}
 
 	// Assign the physics asset if any (must be done after having added the LODs)
diff --git a/Engine/Plugins/Experimental/ChaosClothAssetEditor/Source/ChaosClothAssetTools/Public/ChaosClothAsset/ClothPatternToDynamicMesh.h b/Engine/Plugins/Experimental/ChaosClothAssetEditor/Source/ChaosClothAssetTools/Public/ChaosClothAsset/ClothPatternToDynamicMesh.h
--- a/Engine/Plugins/Experimental/ChaosClothAssetEditor/Source/ChaosClothAssetTools/Public/ChaosClothAsset/ClothPatternToDynamicMesh.h
+++ b/Engine/Plugins/Experimental/ChaosClothAssetEditor/Source/ChaosClothAssetTools/Public/ChaosClothAsset/ClothPatternToDynamicMesh.h
@@ -14,13 +14,16 @@ namespace UE::Chaos::ClothAsset
 {
 
 /**
-* Convert a single pattern from a ClothAsset to a FDynamicMesh3
+* Convert a single pattern from a ClothAsset to a FDynamicMesh3. 
+* When PatternIndex = INDEX_NONE, convert the entire cloth asset to an FDynamicMesh3.
+* When VertexDataType = EClothPatternVertexType::Sim3D, this generates a welded mesh (with native welded indexing) when PatternIndex = INDEX_NONE,
+*	and an unwelded (with native 2D unwelded indexing) when PatternIndex is a valid SimPattern Index.
 */
 class CHAOSCLOTHASSETTOOLS_API FClothPatternToDynamicMesh
 {
 public:
 
-	void Convert(const TSharedPtr<const FManagedArrayCollection> ClothCollection, int32 LODIndex, int32 PatternIndex, EClothPatternVertexType VertexDataType, UE::Geometry::FDynamicMesh3& MeshOut);
+	void Convert(const TSharedPtr<const FManagedArrayCollection> ClothCollection, int32 PatternIndex, EClothPatternVertexType VertexDataType, UE::Geometry::FDynamicMesh3& MeshOut);
 
 	void Convert(const UChaosClothAsset* ClothAssetMeshIn, int32 LODIndex, int32 PatternIndex, EClothPatternVertexType VertexDataType, UE::Geometry::FDynamicMesh3& MeshOut);
 };
diff --git a/Engine/Source/Runtime/ClothingSystemRuntimeCommon/Private/ClothTetherData.cpp b/Engine/Source/Runtime/ClothingSystemRuntimeCommon/Private/ClothTetherData.cpp
--- a/Engine/Source/Runtime/ClothingSystemRuntimeCommon/Private/ClothTetherData.cpp
+++ b/Engine/Source/Runtime/ClothingSystemRuntimeCommon/Private/ClothTetherData.cpp
@@ -18,9 +18,10 @@ public:
 		const TConstArrayView<float>& MaxDistances,
 		bool bUseGeodesicDistance);
 
-	const TArray<TArray<FTether>>& GetTetherSlots() const { return TetherSlots; }
+	FClothTetherDataPrivate(
+		TArray<TArray<TPair<float, int32>>>&& PerDynamicNodeTethers);
 
-	const TArray<int32>& GetTetherNums() const { return TetherNums; }
+	void GetBatchedTetherData(TArray<TArray<TTuple<int32, int32, float>>>& Tethers) const;
 
 private:
 
@@ -35,6 +36,9 @@ private:
 	// The choice of the tether is determined by the shortest geodesic (curvature) distance.
 	void GenerateGeodesicTethers(const TConstArrayView<FVector3f>& Points, const TConstArrayView<float>& MaxDistances);
 
+	// Update TetherNums after computing tethers.
+	void UpdateCounts();
+
 private:
 	TMap<int32, TSet<int32>> NodeToNeighbors;
 	TArray<int32> DynamicNodes;
@@ -59,28 +63,14 @@ void FClothTetherData::GenerateTethers(
 
 	// Calculate the tethers
 	const FClothTetherDataPrivate ClothTetherData(Points, Indices, MaxDistances, bUseGeodesicDistance);
+	ClothTetherData.GetBatchedTetherData(Tethers);
+}
 
-	const TArray<TArray<FClothTetherDataPrivate::FTether>>& TetherSlots = ClothTetherData.GetTetherSlots();
-	const TArray<int32> TetherNums = ClothTetherData.GetTetherNums();
-
-	// Reorganize the multiple tethers per node array into a single sequential batches of tethers more suitable for parallel processing
-	const int32 NumDynamicNodes = TetherSlots.Num();
-	const int32 MaxNumUsedSlots = TetherNums.Num();
-	Tethers.Reset(MaxNumUsedSlots);
-	Tethers.SetNum(MaxNumUsedSlots);
-
-	for (int32 Slot = 0; Slot < MaxNumUsedSlots; ++Slot)
-	{
-		Tethers[Slot].Reserve(TetherNums[Slot]);
-
-		for (int32 Index = 0; Index < NumDynamicNodes; ++Index)
-		{
-			if (TetherSlots[Index].IsValidIndex(Slot))
-			{
-				Tethers[Slot].Emplace(TetherSlots[Index][Slot]);
-			}
-		}
-	}
+void FClothTetherData::GenerateTethers(
+	TArray<TArray<TPair<float, int32>>>&& PerDynamicNodeTethers)
+{
+	const FClothTetherDataPrivate ClothTetherData(MoveTemp(PerDynamicNodeTethers));
+	ClothTetherData.GetBatchedTetherData(Tethers);
 }
 
 bool FClothTetherData::Serialize(FArchive& Ar)
@@ -160,18 +150,61 @@ FClothTetherDataPrivate::FClothTetherDataPrivate(
 			GenerateGeodesicTethers(Points, MaxDistances);
 		}
 
-		// Update counts
-		TetherNums.Reserve(MaxNumAttachments);
-		for (const TArray<FTether>& TetherSlot : TetherSlots)
+		UpdateCounts();
+	}
+}
+
+FClothTetherDataPrivate::FClothTetherDataPrivate(TArray<TArray<TPair<float, int32>>>&& PerDynamicNodeTethers)
+{
+	TetherSlots.SetNum(PerDynamicNodeTethers.Num());
+
+	// For each dynamic node
+	ParallelFor(PerDynamicNodeTethers.Num(), [this, &PerDynamicNodeTethers](int32 Index)
+	{
+		const int32 DynamicNode = Index;
+		TArray<TPair<float, int32>>& ClosestKinematicNodes = PerDynamicNodeTethers[Index];
+		// Only keep the first MaxNumAttachments closest kinematic nodes
+		if (ClosestKinematicNodes.Num() > MaxNumAttachments)
+		{
+			// Order all by distance, smallest first
+			ClosestKinematicNodes.Sort();
+
+			// Shrink the list...
+			constexpr bool bAllowShrinking = false;  // ... but not the array
+			ClosestKinematicNodes.SetNum(MaxNumAttachments, bAllowShrinking);
+		}
+		// Finally create the tethers between this dynamic node and the N closests kinematic ones
+		TetherSlots[Index].Reserve(ClosestKinematicNodes.Num());
+
+		for (const TPair<float, int32> ClosestKinematicNode : ClosestKinematicNodes)
 		{
-			// Resize the count array whenever needed
-			const int32 NumUsedSlots = TetherSlot.Num();
-			TetherNums.SetNum(FMath::Max(TetherNums.Num(), NumUsedSlots));
+			const float Length = ClosestKinematicNode.Key;
+			const int32 KinematicNode = ClosestKinematicNode.Value;
 
-			// Increment each used slot count
-			for (int32 UsedSlot = 0; UsedSlot < NumUsedSlots; ++UsedSlot)
+			TetherSlots[Index].Emplace(KinematicNode, DynamicNode, Length);
+		}
+	});
+
+	UpdateCounts();
+}
+
+void FClothTetherDataPrivate::GetBatchedTetherData(TArray<TArray<TTuple<int32, int32, float>>>& Tethers) const
+{
+	// Reorganize the multiple tethers per node array into a single sequential batches of tethers more suitable for parallel processing
+	const int32 NumDynamicNodes = TetherSlots.Num();
+	const int32 MaxNumUsedSlots = TetherNums.Num();
+	Tethers.Reset(MaxNumUsedSlots);
+	Tethers.SetNum(MaxNumUsedSlots);
+
+	for (int32 Slot = 0; Slot < MaxNumUsedSlots; ++Slot)
+	{
+		Tethers[Slot].Reserve(TetherNums[Slot]);
+
+		for (int32 Index = 0; Index < NumDynamicNodes; ++Index)
+		{
+			if (TetherSlots[Index].IsValidIndex(Slot))
 			{
-				++TetherNums[UsedSlot];
+				Tethers[Slot].Emplace(TetherSlots[Index][Slot]);
 			}
 		}
 	}
@@ -508,3 +541,20 @@ void FClothTetherDataPrivate::GenerateGeodesicTethers(const TConstArrayView<FVec
 	});
 }
 
+void FClothTetherDataPrivate::UpdateCounts()
+{
+	// Update counts
+	TetherNums.Reserve(MaxNumAttachments);
+	for (const TArray<FTether>& TetherSlot : TetherSlots)
+	{
+		// Resize the count array whenever needed
+		const int32 NumUsedSlots = TetherSlot.Num();
+		TetherNums.SetNum(FMath::Max(TetherNums.Num(), NumUsedSlots));
+
+		// Increment each used slot count
+		for (int32 UsedSlot = 0; UsedSlot < NumUsedSlots; ++UsedSlot)
+		{
+			++TetherNums[UsedSlot];
+		}
+	}
+}
\ No newline at end of file
diff --git a/Engine/Source/Runtime/ClothingSystemRuntimeCommon/Public/ClothTetherData.h b/Engine/Source/Runtime/ClothingSystemRuntimeCommon/Public/ClothTetherData.h
--- a/Engine/Source/Runtime/ClothingSystemRuntimeCommon/Public/ClothTetherData.h
+++ b/Engine/Source/Runtime/ClothingSystemRuntimeCommon/Public/ClothTetherData.h
@@ -35,6 +35,13 @@ struct FClothTetherData
 		const TConstArrayView<float>& MaxDistances,  // Mask for sorting the kinematic from the dynamic points
 		bool bUseGeodesicDistance);  // Whether to use geodesic (walking along the sruface) or euclidean (beeline) distances to find the tethers.
 
+	/**
+	 * Generate the tether batches from existing tether data.
+	 * PerDynamicNodeTethers is an array where Index = DynamicIndex, and Value is an array of TPair<float, int32>  = {Distance, KinematicIndex}
+	 */
+	CLOTHINGSYSTEMRUNTIMECOMMON_API void GenerateTethers(
+		TArray<TArray<TPair<float, int32>>>&& PerDynamicNodeTethers);
+
 	/** Custom serializer, since neiher an array of array nor a tuple can be set as a UPROPERTY. */ 
 	bool Serialize(FArchive& Ar);
 };
diff --git a/Engine/Source/Runtime/Core/Public/UObject/UE5MainStreamObjectVersion.h b/Engine/Source/Runtime/Core/Public/UObject/UE5MainStreamObjectVersion.h
--- a/Engine/Source/Runtime/Core/Public/UObject/UE5MainStreamObjectVersion.h
+++ b/Engine/Source/Runtime/Core/Public/UObject/UE5MainStreamObjectVersion.h
@@ -366,6 +366,9 @@ struct CORE_API FUE5MainStreamObjectVersion
 
 		// OpenColorIO transforms now serialize their generated texture(s) and shader code normally into the uasset.
 		OpenColorIOAssetCacheSerialization,
+
+		// Cloth collection single lod schema
+		ClothCollectionSingleLodSchema,
 		
 		// -----<new versions can be added above this line>-------------------------------------------------
 		VersionPlusOne,
