ChaosClothAsset: Update SkeletalMeshImportNode to be able to import sim meshes.

diff --git a/Engine/Plugins/Experimental/ChaosClothAsset/Source/ChaosClothAsset/Private/ChaosClothAsset/ClothGeometryTools.cpp b/Engine/Plugins/Experimental/ChaosClothAsset/Source/ChaosClothAsset/Private/ChaosClothAsset/ClothGeometryTools.cpp
--- a/Engine/Plugins/Experimental/ChaosClothAsset/Source/ChaosClothAsset/Private/ChaosClothAsset/ClothGeometryTools.cpp
+++ b/Engine/Plugins/Experimental/ChaosClothAsset/Source/ChaosClothAsset/Private/ChaosClothAsset/ClothGeometryTools.cpp
@@ -5,6 +5,7 @@
 #include "Containers/Queue.h"
 #include "DynamicMesh/DynamicMesh3.h"
 #include "DynamicMesh/DynamicMeshAttributeSet.h"
+#include "DynamicMesh/DynamicVertexSkinWeightsAttribute.h"
 #include "DynamicMesh/NonManifoldMappingSupport.h"
 #include "Math/Vector.h"
 #include "Util/IndexUtil.h"
@@ -718,14 +719,18 @@ namespace UE::Chaos::ClothAsset
 
 	void FClothGeometryTools::BuildSimMeshFromDynamicMesh(
 		const TSharedPtr<FManagedArrayCollection>& ClothCollection,
-		const UE::Geometry::FDynamicMesh3& DynamicMesh, int32 UVChannelIndex, const FVector2f& UVScale)
+		const UE::Geometry::FDynamicMesh3& DynamicMesh, int32 UVChannelIndex, const FVector2f& UVScale, bool bAppend)
 	{
 		using namespace Private::SimMeshBuilder;
 
-		DeleteSimMesh(ClothCollection);
+		if (!bAppend)
+		{
+			DeleteSimMesh(ClothCollection);
+		}
 
 		const UE::Geometry::FDynamicMeshAttributeSet* const AttributeSet = DynamicMesh.Attributes();
 		const UE::Geometry::FDynamicMeshUVOverlay* const UVOverlay = AttributeSet ? AttributeSet->GetUVLayer(UVChannelIndex) : nullptr;
+		const UE::Geometry::FDynamicMeshVertexSkinWeightsAttribute* SkinWeights = AttributeSet ? AttributeSet->GetSkinWeightsAttribute(FName("Default")) : nullptr;
 
 		TArray<FIsland> Islands;
 		if (UVOverlay)
@@ -743,7 +748,18 @@ namespace UE::Chaos::ClothAsset
 			if (Island.Indices.Num() && Island.Positions2D.Num() && Island.Positions3D.Num())
 			{
 				FCollectionClothSimPatternFacade Pattern = Cloth.AddGetSimPattern();
+				const int32 VertexOffset = Cloth.GetNumSimVertices3D();
 				Pattern.Initialize(Island.Positions2D, Island.Positions3D, Island.Indices);
+				if (SkinWeights)
+				{
+					TArrayView<TArray<int32>> BoneIndices = Cloth.GetSimBoneIndices();
+					TArrayView<TArray<float>> BoneWeights = Cloth.GetSimBoneWeights();
+					const int32 VertexCount = Island.Positions3D.Num();
+					for (int32 VertexIndex = 0; VertexIndex < VertexCount; ++VertexIndex)
+					{						
+						SkinWeights->GetValue(Island.PositionToSourceIndex[VertexIndex], BoneIndices[VertexIndex + VertexOffset], BoneWeights[VertexIndex + VertexOffset]);
+					}
+				}
 			}
 		}
 
diff --git a/Engine/Plugins/Experimental/ChaosClothAsset/Source/ChaosClothAsset/Public/ChaosClothAsset/ClothGeometryTools.h b/Engine/Plugins/Experimental/ChaosClothAsset/Source/ChaosClothAsset/Public/ChaosClothAsset/ClothGeometryTools.h
--- a/Engine/Plugins/Experimental/ChaosClothAsset/Source/ChaosClothAsset/Public/ChaosClothAsset/ClothGeometryTools.h
+++ b/Engine/Plugins/Experimental/ChaosClothAsset/Source/ChaosClothAsset/Public/ChaosClothAsset/ClothGeometryTools.h
@@ -58,6 +58,6 @@ namespace UE::Chaos::ClothAsset
 		*/
 		static void BuildSimMeshFromDynamicMesh(
 			const TSharedPtr<FManagedArrayCollection>& ClothCollection,
-			const UE::Geometry::FDynamicMesh3& DynamicMesh, int32 UVChannelIndex, const FVector2f& UVScale);
+			const UE::Geometry::FDynamicMesh3& DynamicMesh, int32 UVChannelIndex, const FVector2f& UVScale, bool bAppend);
 	};
 }  // End namespace UE::Chaos::ClothAsset
\ No newline at end of file
diff --git a/Engine/Plugins/Experimental/ChaosClothAssetEditor/Source/ChaosClothAssetDataflowNodes/Private/ChaosClothAsset/ClothDataflowTools.cpp b/Engine/Plugins/Experimental/ChaosClothAssetEditor/Source/ChaosClothAssetDataflowNodes/Private/ChaosClothAsset/ClothDataflowTools.cpp
--- a/Engine/Plugins/Experimental/ChaosClothAssetEditor/Source/ChaosClothAssetDataflowNodes/Private/ChaosClothAsset/ClothDataflowTools.cpp
+++ b/Engine/Plugins/Experimental/ChaosClothAssetEditor/Source/ChaosClothAssetDataflowNodes/Private/ChaosClothAsset/ClothDataflowTools.cpp
@@ -1,70 +1,438 @@
 // Copyright Epic Games, Inc. All Rights Reserved.
 
 #include "ChaosClothAsset/ClothDataflowTools.h"
+#include "ChaosClothAsset/ClothGeometryTools.h"
 #include "ChaosClothAsset/CollectionClothFacade.h"
 #include "ChaosClothAsset/ClothCollection.h"
+#include "Animation/Skeleton.h"
 #include "BoneWeights.h"
+#include "DynamicMesh/NonManifoldMappingSupport.h"
+#include "Math/UnrealMathUtility.h"
 #include "Rendering/SkeletalMeshLODModel.h"
+#include "SkeletalMeshAttributes.h"
+#include "ToDynamicMesh.h"
 
 namespace UE::Chaos::ClothAsset
 {
-void FClothDataflowTools::AddRenderPatternFromSkeletalMeshSection(const TSharedPtr<FManagedArrayCollection>& ClothCollection, const FSkeletalMeshLODModel& SkeletalMeshModel, const int32 SectionIndex, const FString& RenderMaterialPathName)
-{
-	check(SectionIndex < SkeletalMeshModel.Sections.Num());
-	check(ClothCollection.IsValid());
-
-	FCollectionClothFacade Cloth(ClothCollection);
-	FCollectionClothRenderPatternFacade ClothPatternFacade = Cloth.AddGetRenderPattern();
-
-	const FSkelMeshSection& Section = SkeletalMeshModel.Sections[SectionIndex];
-	ClothPatternFacade.SetNumRenderVertices(Section.NumVertices);
-	ClothPatternFacade.SetNumRenderFaces(Section.NumTriangles);
-
-	TArrayView<FVector3f> RenderPosition = ClothPatternFacade.GetRenderPosition();
-	TArrayView<FVector3f> RenderNormal = ClothPatternFacade.GetRenderNormal();
-	TArrayView<FVector3f> RenderTangentU = ClothPatternFacade.GetRenderTangentU();
-	TArrayView<FVector3f> RenderTangentV = ClothPatternFacade.GetRenderTangentV();
-	TArrayView<TArray<FVector2f>> RenderUVs = ClothPatternFacade.GetRenderUVs();
-	TArrayView<FLinearColor> RenderColor = ClothPatternFacade.GetRenderColor();
-	TArrayView<TArray<int32>> RenderBoneIndices = ClothPatternFacade.GetRenderBoneIndices();
-	TArrayView<TArray<float>> RenderBoneWeights = ClothPatternFacade.GetRenderBoneWeights();
-	for (int32 VertexIndex = 0; VertexIndex < Section.NumVertices; ++VertexIndex)
+	namespace Private
 	{
-		const FSoftSkinVertex& SoftVertex = Section.SoftVertices[VertexIndex];
-
-		RenderPosition[VertexIndex] = SoftVertex.Position;
-		RenderNormal[VertexIndex] = SoftVertex.TangentZ;
-		RenderTangentU[VertexIndex] = SoftVertex.TangentX;
-		RenderTangentV[VertexIndex] = SoftVertex.TangentY;
-		RenderUVs[VertexIndex].SetNum(MAX_TEXCOORDS);
-		for (int32 TexCoordIndex = 0; TexCoordIndex < MAX_TEXCOORDS; ++TexCoordIndex)
+		//
+		// Wrapper for accessing a SkelMeshSection. Implements the interface expected by TToDynamicMesh<>.
+		// This will weld all vertices which are the same.
+		//
+		template<bool bHasNormals = false, bool bHasTangents = false, bool bHasBiTangents = false, bool bHasColors = false>
+		struct FSkelMeshSectionWrapper
 		{
-			RenderUVs[VertexIndex][TexCoordIndex] = SoftVertex.UVs[TexCoordIndex];
-		}
+			typedef int32 TriIDType;
+			typedef int32 VertIDType;
+			typedef int32 WedgeIDType;
+			typedef int32 UVIDType;
+			typedef int32 NormalIDType;
+			typedef int32 ColorIDType;
+
+			FSkelMeshSectionWrapper(const FSkeletalMeshLODModel& SkeletalMeshModel, const int32 SectionIndex)
+				: SourceSection(SkeletalMeshModel.Sections[SectionIndex])
+				, IndexBuffer(SkeletalMeshModel.IndexBuffer.GetData() + SourceSection.BaseIndex, SourceSection.NumTriangles * 3)
+			{
+				const int32 NumVerts = SourceSection.SoftVertices.Num();
+				const int32 NumTriangles = SourceSection.NumTriangles;
+
+				// We need to weld the mesh verts to get rid of duplicates (happens for smoothing groups)
+				TArray<FVector> UniqueVerts;
+				IndexRemap.AddDefaulted(NumVerts);
+				constexpr float ThreshSq = UE_THRESH_POINTS_ARE_SAME * UE_THRESH_POINTS_ARE_SAME;
+				for (int32 VertIndex = 0; VertIndex < NumVerts; ++VertIndex)
+				{
+					const FSoftSkinVertex& SourceVert = SourceSection.SoftVertices[VertIndex];
+
+					bool bUnique = true;
+					int32 RemapIndex = INDEX_NONE;
+
+					const int32 NumUniqueVerts = UniqueVerts.Num();
+					for (int32 UniqueVertIndex = 0; UniqueVertIndex < NumUniqueVerts; ++UniqueVertIndex)
+					{
+						FVector& UniqueVert = UniqueVerts[UniqueVertIndex];
+
+						if ((UniqueVert - (FVector)SourceVert.Position).SizeSquared() <= ThreshSq)
+						{
+							// Not unique
+							bUnique = false;
+							RemapIndex = UniqueVertIndex;
+
+							break;
+						}
+					}
+
+					if (bUnique)
+					{
+						// Unique
+						UniqueVerts.Add((FVector)SourceVert.Position);
+						OriginalIndexes.Add(VertIndex);
+						IndexRemap[VertIndex] = UniqueVerts.Num() - 1;
+					}
+					else
+					{
+						IndexRemap[VertIndex] = RemapIndex;
+					}
+				}
+
+				TriIDs.SetNum(NumTriangles);
+				for (int32 Tri = 0; Tri < NumTriangles; ++Tri)
+				{
+					TriIDs[Tri] = Tri;
+				}
+			}
+
+			int32 NumTris() const
+			{
+				return TriIDs.Num();
+			}
+
+			int32 NumVerts() const
+			{
+				return OriginalIndexes.Num();
+			}
+
+			int32 NumUVLayers() const
+			{
+				return MAX_TEXCOORDS;
+			}
+
+			// --"Vertex Buffer" info 
+			const TArray<int32>& GetVertIDs() const
+			{
+				return OriginalIndexes;
+			}
+
+			const FVector3d GetPosition(const VertIDType VtxID) const
+			{
+				return (FVector3d)SourceSection.SoftVertices[VtxID].Position;
+			}
+
+			// --"Index Buffer" info
+			const TArray<int32>& GetTriIDs() const
+			{
+				return TriIDs;
+			}
+
+			// return false if this TriID is not contained in mesh.
+			bool GetTri(const TriIDType TriID, VertIDType& VID0, VertIDType& VID1, VertIDType& VID2) const
+			{
+				if (TriID >= 0 && TriID < (int32)SourceSection.NumTriangles)
+				{
+					VID0 = IndexRemap[IndexBuffer[3 * TriID + 0] - SourceSection.BaseVertexIndex];
+					VID1 = IndexRemap[IndexBuffer[3 * TriID + 1] - SourceSection.BaseVertexIndex];
+					VID2 = IndexRemap[IndexBuffer[3 * TriID + 2] - SourceSection.BaseVertexIndex];
+					return true;
+				}
+				return false;
+			}
+
+			bool HasNormals() const
+			{
+				return bHasNormals;
+			}
+
+			bool HasTangents() const
+			{
+				return bHasTangents;
+			}
+
+			bool HasBiTangents() const
+			{
+				return bHasBiTangents;
+			}
+
+			bool HasColors() const
+			{
+				return bHasColors;
+			}
+
+			// Each triangle corner is a wedge. This will lookup into original unwelded soft verts
+			void GetWedgeIDs(const TriIDType& TriID, WedgeIDType& WID0, WedgeIDType& WID1, WedgeIDType& WID2) const
+			{
+				WID0 = IndexBuffer[3 * TriID + 0] - SourceSection.BaseVertexIndex;
+				WID1 = IndexBuffer[3 * TriID + 1] - SourceSection.BaseVertexIndex;
+				WID2 = IndexBuffer[3 * TriID + 2] - SourceSection.BaseVertexIndex;
+			}
+
+			// attribute access per-wedge
+			// NB:  ToDynamicMesh will attempt to weld identical attributes that are associated with the same vertex
+			FVector2f GetWedgeUV(int32 UVLayerIndex, WedgeIDType WID) const
+			{
+				check(UVLayerIndex < MAX_TEXCOORDS);
+				return SourceSection.SoftVertices[WID].UVs[UVLayerIndex];
+			}
+
+			FVector3f GetWedgeNormal(WedgeIDType WID) const
+			{
+				return SourceSection.SoftVertices[WID].TangentZ;
+			}
+
+			FVector3f GetWedgeTangent(WedgeIDType WID) const
+			{
+				return SourceSection.SoftVertices[WID].TangentX;
+			}
+
+			FVector3f GetWedgeBiTangent(WedgeIDType WID) const
+			{
+				return SourceSection.SoftVertices[WID].TangentY;
+			}
+
+			FVector4f GetWedgeColor(WedgeIDType WID) const
+			{
+				return FLinearColor(SourceSection.SoftVertices[WID].Color);
+			}
+
+			// attribute access that exploits shared attributes. 
+			// each group of shared attributes presents itself as a mesh with its own attribute vertex buffer.
+			// NB:  If the mesh has no shared Attr attributes, then Get{Attr}IDs() should return an empty array.
+			// NB:  Get{Attr}Tri() functions should return false if the triangle is not set in the attribute mesh. 
+			const TArray<UVIDType>& GetUVIDs(int32 LayerID) const
+			{
+				return EmptyArray;
+			}
+
+			FVector2f GetUV(int32 LayerID, UVIDType UVID) const
+			{
+				check(false);
+				return FVector2f();
+			}
+
+			bool GetUVTri(int32 LayerID, const TriIDType& TID, UVIDType& ID0, UVIDType& ID1, UVIDType& ID2) const
+			{
+				return false;
+			}
+
+			const TArray<NormalIDType>& GetNormalIDs() const
+			{
+				return EmptyArray;
+			}
+
+			FVector3f GetNormal(NormalIDType ID) const
+			{
+				check(false);
+				return FVector3f();
+			}
+
+			bool GetNormalTri(const TriIDType& TID, NormalIDType& NID0, NormalIDType& NID1, NormalIDType& NID2) const
+			{
+				return false;
+			}
+
+			const TArray<NormalIDType>& GetTangentIDs() const
+			{
+				return EmptyArray;
+			}
+
+			FVector3f GetTangent(NormalIDType ID) const
+			{
+				check(false);
+				return FVector3f();
+			}
+
+			bool GetTangentTri(const TriIDType& TID, NormalIDType& NID0, NormalIDType& NID1, NormalIDType& NID2) const
+			{
+				return false;
+			}
 
-		RenderColor[VertexIndex] = FLinearColor(SoftVertex.Color);
+			const TArray<NormalIDType>& GetBiTangentIDs() const
+			{
+				return EmptyArray;
+			}
 
-		const int32 NumBones = Section.MaxBoneInfluences;
-		RenderBoneIndices[VertexIndex].SetNum(NumBones);
-		RenderBoneWeights[VertexIndex].SetNum(NumBones);
-		for (int32 BoneIndex = 0; BoneIndex < NumBones; ++BoneIndex)
+			FVector3f GetBiTangent(NormalIDType ID) const
+			{
+				check(false);
+				return FVector3f();
+			}
+
+			bool GetBiTangentTri(const TriIDType& TID, NormalIDType& NID0, NormalIDType& NID1, NormalIDType& NID2) const
+			{
+				return false;
+			}
+
+			const TArray<ColorIDType>& GetColorIDs() const
+			{
+				return EmptyArray;
+			}
+
+			FVector4f GetColor(ColorIDType ID) const
+			{
+				check(false);
+				return FVector4f();
+			}
+
+			bool GetColorTri(const TriIDType& TID, ColorIDType& NID0, ColorIDType& NID1, ColorIDType& NID2) const
+			{
+				return false;
+			}
+
+			// weight maps information
+			int32 NumWeightMapLayers() const
+			{
+				return 0;
+			}
+
+			float GetVertexWeight(int32 WeightMapIndex, int32 SrcVertID) const
+			{
+				check(false);
+				return 0.f;
+			}
+
+			FName GetWeightMapName(int32 WeightMapIndex) const
+			{
+				check(false);
+				return FName();
+			}
+
+			// skin weight attributes information
+			int32 NumSkinWeightAttributes() const
+			{
+				return 1;
+			}
+
+			UE::AnimationCore::FBoneWeights GetVertexSkinWeight(int32 SkinWeightAttributeIndex, VertIDType VtxID) const
+			{
+				using namespace UE::AnimationCore;
+				check(SkinWeightAttributeIndex == 0);
+				const int32 NumInfluences = SourceSection.MaxBoneInfluences;
+				const FSoftSkinVertex& SoftVertex = SourceSection.SoftVertices[VtxID];
+				TArray<FBoneWeight> BoneWeightArray;
+				BoneWeightArray.SetNumUninitialized(NumInfluences);
+				for (int32 Idx = 0; Idx < NumInfluences; ++Idx)
+				{
+					BoneWeightArray[Idx] = FBoneWeight(SoftVertex.InfluenceBones[Idx], (float)SoftVertex.InfluenceWeights[Idx] * UE::AnimationCore::InvMaxRawBoneWeightFloat);
+				}
+				return FBoneWeights::Create(BoneWeightArray, FBoneWeightsSettings());
+			}
+
+			FName GetSkinWeightAttributeName(int32 SkinWeightAttributeIndex) const
+			{
+				checkfSlow(SkinWeightAttributeIndex == 0, TEXT("Cloth assets should only have one skin weight profile"));
+				return FSkeletalMeshAttributes::DefaultSkinWeightProfileName;
+			}
+
+			// bone attributes information
+			int32 GetNumBones() const
+			{
+				return 0;
+			}
+
+			FName GetBoneName(int32 BoneIdx) const
+			{
+				check(false);
+				return FName();
+			}
+
+			int32 GetBoneParentIndex(int32 BoneIdx) const
+			{
+				check(false);
+				return INDEX_NONE;
+			}
+
+			FTransform GetBonePose(int32 BoneIdx) const
+			{
+				check(false);
+				return FTransform();
+			}
+
+			FVector4f GetBoneColor(int32 BoneIdx) const
+			{
+				check(false);
+				return FLinearColor::White;
+			}
+
+
+			const FSkelMeshSection& SourceSection;
+			const TConstArrayView<uint32> IndexBuffer;
+			TArray<int32> OriginalIndexes; // UniqueIndex -> OrigIndex
+			TArray<int32> IndexRemap; // OrigIndex -> UniqueIndex
+			TArray<int32> TriIDs;
+			TArray<int32> EmptyArray;
+
+		};
+	} // Private
+
+	void FClothDataflowTools::AddRenderPatternFromSkeletalMeshSection(const TSharedPtr<FManagedArrayCollection>& ClothCollection, const FSkeletalMeshLODModel& SkeletalMeshModel, const int32 SectionIndex, const FString& RenderMaterialPathName)
+	{
+		check(SectionIndex < SkeletalMeshModel.Sections.Num());
+		check(ClothCollection.IsValid());
+
+		FCollectionClothFacade Cloth(ClothCollection);
+		check(Cloth.IsValid());
+
+		FCollectionClothRenderPatternFacade ClothPatternFacade = Cloth.AddGetRenderPattern();
+
+		const FSkelMeshSection& Section = SkeletalMeshModel.Sections[SectionIndex];
+		ClothPatternFacade.SetNumRenderVertices(Section.NumVertices);
+		ClothPatternFacade.SetNumRenderFaces(Section.NumTriangles);
+
+		TArrayView<FVector3f> RenderPosition = ClothPatternFacade.GetRenderPosition();
+		TArrayView<FVector3f> RenderNormal = ClothPatternFacade.GetRenderNormal();
+		TArrayView<FVector3f> RenderTangentU = ClothPatternFacade.GetRenderTangentU();
+		TArrayView<FVector3f> RenderTangentV = ClothPatternFacade.GetRenderTangentV();
+		TArrayView<TArray<FVector2f>> RenderUVs = ClothPatternFacade.GetRenderUVs();
+		TArrayView<FLinearColor> RenderColor = ClothPatternFacade.GetRenderColor();
+		TArrayView<TArray<int32>> RenderBoneIndices = ClothPatternFacade.GetRenderBoneIndices();
+		TArrayView<TArray<float>> RenderBoneWeights = ClothPatternFacade.GetRenderBoneWeights();
+		for (int32 VertexIndex = 0; VertexIndex < Section.NumVertices; ++VertexIndex)
 		{
-			RenderBoneIndices[VertexIndex][BoneIndex] = (int32)Section.BoneMap[(int32)SoftVertex.InfluenceBones[BoneIndex]];
-			RenderBoneWeights[VertexIndex][BoneIndex] = (float)SoftVertex.InfluenceWeights[BoneIndex] * UE::AnimationCore::InvMaxRawBoneWeightFloat;
+			const FSoftSkinVertex& SoftVertex = Section.SoftVertices[VertexIndex];
+
+			RenderPosition[VertexIndex] = SoftVertex.Position;
+			RenderNormal[VertexIndex] = SoftVertex.TangentZ;
+			RenderTangentU[VertexIndex] = SoftVertex.TangentX;
+			RenderTangentV[VertexIndex] = SoftVertex.TangentY;
+			RenderUVs[VertexIndex].SetNum(MAX_TEXCOORDS);
+			for (int32 TexCoordIndex = 0; TexCoordIndex < MAX_TEXCOORDS; ++TexCoordIndex)
+			{
+				RenderUVs[VertexIndex][TexCoordIndex] = SoftVertex.UVs[TexCoordIndex];
+			}
+
+			RenderColor[VertexIndex] = FLinearColor(SoftVertex.Color);
+
+			const int32 NumBones = Section.MaxBoneInfluences;
+			RenderBoneIndices[VertexIndex].SetNum(NumBones);
+			RenderBoneWeights[VertexIndex].SetNum(NumBones);
+			for (int32 BoneIndex = 0; BoneIndex < NumBones; ++BoneIndex)
+			{
+				RenderBoneIndices[VertexIndex][BoneIndex] = (int32)Section.BoneMap[(int32)SoftVertex.InfluenceBones[BoneIndex]];
+				RenderBoneWeights[VertexIndex][BoneIndex] = (float)SoftVertex.InfluenceWeights[BoneIndex] * UE::AnimationCore::InvMaxRawBoneWeightFloat;
+			}
 		}
-	}
 
-	const int32 VertexOffset = ClothPatternFacade.GetRenderVerticesOffset();
-	TArrayView<FIntVector3> RenderIndices = ClothPatternFacade.GetRenderIndices();
-	for (uint32 FaceIndex = 0; FaceIndex < Section.NumTriangles; ++FaceIndex)
+		const int32 VertexOffset = ClothPatternFacade.GetRenderVerticesOffset();
+		TArrayView<FIntVector3> RenderIndices = ClothPatternFacade.GetRenderIndices();
+		for (uint32 FaceIndex = 0; FaceIndex < Section.NumTriangles; ++FaceIndex)
+		{
+			const uint32 IndexOffset = Section.BaseIndex + FaceIndex * 3;
+			RenderIndices[FaceIndex] = FIntVector3(
+				SkeletalMeshModel.IndexBuffer[IndexOffset + 0] - Section.BaseVertexIndex + VertexOffset,
+				SkeletalMeshModel.IndexBuffer[IndexOffset + 1] - Section.BaseVertexIndex + VertexOffset,
+				SkeletalMeshModel.IndexBuffer[IndexOffset + 2] - Section.BaseVertexIndex + VertexOffset
+			);
+		}
+		ClothPatternFacade.SetRenderMaterialPathName(RenderMaterialPathName);
+	}
+	
+	void FClothDataflowTools::AddSimPatternsFromSkeletalMeshSection(const TSharedPtr<FManagedArrayCollection>& ClothCollection, const FSkeletalMeshLODModel& SkeletalMeshModel, const int32 SectionIndex, const int32 UVChannelIndex)
 	{
-		const uint32 IndexOffset = Section.BaseIndex + FaceIndex * 3;
-		RenderIndices[FaceIndex] = FIntVector3(
-			SkeletalMeshModel.IndexBuffer[IndexOffset + 0] - Section.BaseVertexIndex + VertexOffset,
-			SkeletalMeshModel.IndexBuffer[IndexOffset + 1] - Section.BaseVertexIndex + VertexOffset,
-			SkeletalMeshModel.IndexBuffer[IndexOffset + 2] - Section.BaseVertexIndex + VertexOffset
-		);
+		check(SectionIndex < SkeletalMeshModel.Sections.Num());
+
+		// Convert to DynamicMesh and then use that to create patterns.
+		UE::Geometry::TToDynamicMesh<Private::FSkelMeshSectionWrapper<>> SkelMeshSectionToDynamicMesh;
+		Private::FSkelMeshSectionWrapper<> SectionWrapper(SkeletalMeshModel, SectionIndex);
+
+		UE::Geometry::FDynamicMesh3 DynamicMesh;
+		DynamicMesh.EnableAttributes();
+		constexpr bool bCopyTangents = false;
+		SkelMeshSectionToDynamicMesh.Convert(DynamicMesh, SectionWrapper, [](int32) { return 0; }, [](int32) { return INDEX_NONE; }, bCopyTangents);
+
+		// Set ToSrcVertIDMap as an overlay that the build sim mesh code expects.
+		UE::Geometry::FNonManifoldMappingSupport::AttachNonManifoldVertexMappingData(SkelMeshSectionToDynamicMesh.ToSrcVertIDMap, DynamicMesh);
+
+		constexpr bool bAppend = true;
+		FClothGeometryTools::BuildSimMeshFromDynamicMesh(ClothCollection, DynamicMesh, UVChannelIndex, FVector2f(1.f), bAppend);
 	}
-	ClothPatternFacade.SetRenderMaterialPathName(RenderMaterialPathName);
-}
 }  // End namespace UE::Chaos::ClothAsset
diff --git a/Engine/Plugins/Experimental/ChaosClothAssetEditor/Source/ChaosClothAssetDataflowNodes/Private/ChaosClothAsset/ClothDataflowTools.h b/Engine/Plugins/Experimental/ChaosClothAssetEditor/Source/ChaosClothAssetDataflowNodes/Private/ChaosClothAsset/ClothDataflowTools.h
--- a/Engine/Plugins/Experimental/ChaosClothAssetEditor/Source/ChaosClothAssetDataflowNodes/Private/ChaosClothAsset/ClothDataflowTools.h
+++ b/Engine/Plugins/Experimental/ChaosClothAssetEditor/Source/ChaosClothAssetDataflowNodes/Private/ChaosClothAsset/ClothDataflowTools.h
@@ -16,5 +16,7 @@ namespace UE::Chaos::ClothAsset
 	struct FClothDataflowTools
 	{
 		static void AddRenderPatternFromSkeletalMeshSection(const TSharedPtr<FManagedArrayCollection>& ClothCollection, const FSkeletalMeshLODModel& SkeletalMeshModel, const int32 SectionIndex, const FString& RenderMaterialPathName);
+
+		static void AddSimPatternsFromSkeletalMeshSection(const TSharedPtr<FManagedArrayCollection>& ClothCollection, const FSkeletalMeshLODModel& SkeletalMeshModel, const int32 SectionIndex, const int32 UVChannelIndex);
 	};
 }  // End namespace UE::Chaos::ClothAsset
diff --git a/Engine/Plugins/Experimental/ChaosClothAssetEditor/Source/ChaosClothAssetDataflowNodes/Private/ChaosClothAsset/SkeletalMeshImportNode.cpp b/Engine/Plugins/Experimental/ChaosClothAssetEditor/Source/ChaosClothAssetDataflowNodes/Private/ChaosClothAsset/SkeletalMeshImportNode.cpp
--- a/Engine/Plugins/Experimental/ChaosClothAssetEditor/Source/ChaosClothAssetDataflowNodes/Private/ChaosClothAsset/SkeletalMeshImportNode.cpp
+++ b/Engine/Plugins/Experimental/ChaosClothAssetEditor/Source/ChaosClothAssetDataflowNodes/Private/ChaosClothAsset/SkeletalMeshImportNode.cpp
@@ -55,7 +55,7 @@ void FChaosClothAssetSkeletalMeshImportNode::Evaluate(Dataflow::FContext& Contex
 
 			if (bImportSimMesh)
 			{
-				ensureMsgf(false, TEXT("TODO: Import as sim mesh not implemented"));
+				FClothDataflowTools::AddSimPatternsFromSkeletalMeshSection(ClothCollection, LODModel, SectionIndex, UVChannel);
 			}
 
 			if (bImportRenderMesh)
diff --git a/Engine/Plugins/Experimental/ChaosClothAssetEditor/Source/ChaosClothAssetDataflowNodes/Private/ChaosClothAsset/StaticMeshImportNode.cpp b/Engine/Plugins/Experimental/ChaosClothAssetEditor/Source/ChaosClothAssetDataflowNodes/Private/ChaosClothAsset/StaticMeshImportNode.cpp
--- a/Engine/Plugins/Experimental/ChaosClothAssetEditor/Source/ChaosClothAssetDataflowNodes/Private/ChaosClothAsset/StaticMeshImportNode.cpp
+++ b/Engine/Plugins/Experimental/ChaosClothAssetEditor/Source/ChaosClothAssetDataflowNodes/Private/ChaosClothAsset/StaticMeshImportNode.cpp
@@ -147,7 +147,8 @@ void FChaosClothAssetStaticMeshImportNode::Evaluate(Dataflow::FContext& Context,
 					UE::Geometry::FDynamicMesh3 DynamicMesh;
 
 					Converter.Convert(MeshDescription, DynamicMesh);
-					UE::Chaos::ClothAsset::FClothGeometryTools::BuildSimMeshFromDynamicMesh(ClothCollection, DynamicMesh, UVChannel, UVScale);
+					constexpr bool bAppend = false;
+					UE::Chaos::ClothAsset::FClothGeometryTools::BuildSimMeshFromDynamicMesh(ClothCollection, DynamicMesh, UVChannel, UVScale, bAppend);
 				}
 				if (bImportAsRenderMesh)
 				{
diff --git a/Engine/Plugins/Experimental/ChaosClothAssetEditor/Source/ChaosClothAssetTools/Private/ChaosClothAsset/ClothingAssetToClothAssetExporter.cpp b/Engine/Plugins/Experimental/ChaosClothAssetEditor/Source/ChaosClothAssetTools/Private/ChaosClothAsset/ClothingAssetToClothAssetExporter.cpp
--- a/Engine/Plugins/Experimental/ChaosClothAssetEditor/Source/ChaosClothAssetTools/Private/ChaosClothAsset/ClothingAssetToClothAssetExporter.cpp
+++ b/Engine/Plugins/Experimental/ChaosClothAssetEditor/Source/ChaosClothAssetTools/Private/ChaosClothAsset/ClothingAssetToClothAssetExporter.cpp
@@ -125,7 +125,8 @@ void UClothingAssetToChaosClothAssetExporter::Export(const UClothingAssetBase* C
 		UE::Geometry::FNonManifoldMappingSupport::AttachNonManifoldVertexMappingData(ToDynamicMesh.ToSrcVertIDMap, DynamicMesh);
 
 		constexpr int32 UVChannelIndexNone = INDEX_NONE;
-		FClothGeometryTools::BuildSimMeshFromDynamicMesh(ClothCollection, DynamicMesh, UVChannelIndexNone, FVector2f(1.f));
+		constexpr bool bAppend = false;
+		FClothGeometryTools::BuildSimMeshFromDynamicMesh(ClothCollection, DynamicMesh, UVChannelIndexNone, FVector2f(1.f), bAppend);
 	}
 
 	if (NumLods)
