Cloth to DynamicMesh: set default material ID of zero on all sim triangles
Cloth Editor: if any triangle in the render mesh has an invalid material ID, apply the default WoldGridMaterial to it

diff --git a/Engine/Plugins/Experimental/ChaosClothAssetEditor/Source/ChaosClothAssetEditor/Private/ChaosClothAsset/ClothEditorMode.cpp b/Engine/Plugins/Experimental/ChaosClothAssetEditor/Source/ChaosClothAssetEditor/Private/ChaosClothAsset/ClothEditorMode.cpp
--- a/Engine/Plugins/Experimental/ChaosClothAssetEditor/Source/ChaosClothAssetEditor/Private/ChaosClothAsset/ClothEditorMode.cpp
+++ b/Engine/Plugins/Experimental/ChaosClothAssetEditor/Source/ChaosClothAssetEditor/Private/ChaosClothAsset/ClothEditorMode.cpp
@@ -57,6 +57,7 @@
 #include "ContextObjectStore.h"
 #include "BaseGizmos/TransformGizmoUtil.h"
 #include "Materials/Material.h"
+#include "MaterialDomain.h"
 
 #include UE_INLINE_GENERATED_CPP_BY_NAME(ClothEditorMode)
 
@@ -418,6 +419,23 @@ void UChaosClothAssetEditorMode::ReinitializeDynamicMeshComponents()
 					UMaterialInterface* const Material = LoadObject<UMaterialInterface>(nullptr, *Path);
 					MeshComponent.SetMaterial(MaterialIndex, Material);
 				}
+
+				// Fix up any triangles without valid material IDs
+				int32 DefaultMaterialID = INDEX_NONE;
+				for (const int32 TriID : MeshComponent.GetMesh()->TriangleIndicesItr())
+				{
+					const int32 MaterialID = MeshComponent.GetMesh()->Attributes()->GetMaterialID()->GetValue(TriID);
+					if (!MeshComponent.GetMaterial(MaterialID))
+					{
+						if (DefaultMaterialID == INDEX_NONE)
+						{
+							DefaultMaterialID = MeshComponent.GetNumMaterials();
+							MeshComponent.SetMaterial(DefaultMaterialID, UMaterial::GetDefaultMaterial(MD_Surface));
+						}
+						MeshComponent.GetMesh()->Attributes()->GetMaterialID()->SetValue(TriID, DefaultMaterialID);
+					}
+				}
+
 			}
 			break;
 		}
diff --git a/Engine/Plugins/Experimental/ChaosClothAssetEditor/Source/ChaosClothAssetTools/Private/ChaosClothAsset/ClothPatternToDynamicMesh.cpp b/Engine/Plugins/Experimental/ChaosClothAssetEditor/Source/ChaosClothAssetTools/Private/ChaosClothAsset/ClothPatternToDynamicMesh.cpp
--- a/Engine/Plugins/Experimental/ChaosClothAssetEditor/Source/ChaosClothAssetTools/Private/ChaosClothAsset/ClothPatternToDynamicMesh.cpp
+++ b/Engine/Plugins/Experimental/ChaosClothAssetEditor/Source/ChaosClothAssetTools/Private/ChaosClothAsset/ClothPatternToDynamicMesh.cpp
@@ -546,12 +546,12 @@ void FClothPatternToDynamicMesh::Convert(const TSharedPtr<const FManagedArrayCol
 					return PatternIndex;
 				}
 				const int32 FoundPattern = ClothFacade.FindRenderPatternByFaceIndex(TriID);
-				if (ensure(FoundPattern != INDEX_NONE))
-				{
-					return FoundPattern;
-				}
+				check(FoundPattern != INDEX_NONE);
+				return FoundPattern;
 			}
-			return INDEX_NONE;
+			
+			// Sim meshes will have a default material with MaterialID zero applied
+			return 0;
 		};
 
 		constexpr bool bCopyTangents = false;
