Chaos Cloth Editor- Removed an annoying ensure that triggered every time an unconnected node was selected.

#rnx

diff --git a/Engine/Plugins/Experimental/ChaosClothAssetEditor/Source/ChaosClothAssetTools/Private/ChaosClothAsset/ClothPatternToDynamicMesh.cpp b/Engine/Plugins/Experimental/ChaosClothAssetEditor/Source/ChaosClothAssetTools/Private/ChaosClothAsset/ClothPatternToDynamicMesh.cpp
--- a/Engine/Plugins/Experimental/ChaosClothAssetEditor/Source/ChaosClothAssetTools/Private/ChaosClothAsset/ClothPatternToDynamicMesh.cpp
+++ b/Engine/Plugins/Experimental/ChaosClothAssetEditor/Source/ChaosClothAssetTools/Private/ChaosClothAsset/ClothPatternToDynamicMesh.cpp
@@ -154,7 +154,6 @@ public:
 			LoadObject<USkeleton>(nullptr, *SkeletonPathName, nullptr, LOAD_None, nullptr) :
 			nullptr;
 		RefSkeleton = Skeleton ? &Skeleton->GetReferenceSkeleton() : nullptr;
-		ensureMsgf(RefSkeleton, TEXT("Reference skeleton is not set"));
 	}
 
 	int32 NumTris() const
