ClothTrainingTool: remove SaveGeometryCache function. Use GeometryCacheHelpers

diff --git a/Engine/Plugins/Experimental/ChaosClothAssetEditor/Source/ChaosClothAssetEditorTools/Private/ChaosClothAsset/ClothTrainingTool.cpp b/Engine/Plugins/Experimental/ChaosClothAssetEditor/Source/ChaosClothAssetEditorTools/Private/ChaosClothAsset/ClothTrainingTool.cpp
--- a/Engine/Plugins/Experimental/ChaosClothAssetEditor/Source/ChaosClothAssetEditorTools/Private/ChaosClothAsset/ClothTrainingTool.cpp
+++ b/Engine/Plugins/Experimental/ChaosClothAssetEditor/Source/ChaosClothAssetEditorTools/Private/ChaosClothAsset/ClothTrainingTool.cpp
@@ -25,6 +25,7 @@
 #include "GeometryCacheCodecV1.h"
 #include "GeometryCacheMeshData.h"
 #include "GeometryCacheTrackStreamable.h"
+#include "GeometryCacheConstantTopologyWriter.h"
 #include "InteractiveToolManager.h"
 #include "Internationalization/Regex.h"
 #include "ModelingOperators.h"
@@ -108,27 +109,6 @@ namespace UE::ClothTrainingTool::Private
 		return Result;
 	}
 
-	FBox3f GetBoundingBox(const TArray<FVector3f>& Positions)
-	{
-		FBox3f BoundingBox;
-		for (const FVector3f& Position : Positions)
-		{
-			BoundingBox += Position;
-		}
-		return BoundingBox;
-	}
-
-	void CopyMaterials(const USkinnedAsset& Asset, UGeometryCache& Cache)
-	{
-		const TArray<FSkeletalMaterial>& Materials = Asset.GetMaterials();
-		TArray<TObjectPtr<UMaterialInterface>>& Interfaces = Cache.Materials;
-		Interfaces.SetNumZeroed(Materials.Num());
-		for (int32 MaterialIndex = 0; MaterialIndex < Materials.Num(); ++MaterialIndex)
-		{
-			Interfaces[MaterialIndex] = Materials[MaterialIndex].MaterialInterface;
-		}
-	}
-
 	TArray<uint32> Range(uint32 Start, uint32 End)
 	{
 		TArray<uint32> Result;
@@ -151,105 +131,18 @@ namespace UE::ClothTrainingTool::Private
 		return NumVertices;
 	}
 
-	TArray<FVector2f> GetUV0s(const FSkeletalMeshLODRenderData& LODData)
-	{
-		TArray<FVector2f> UV0s;
-		const FStaticMeshVertexBuffer& StaticMeshVertexBuffer = LODData.StaticVertexBuffers.StaticMeshVertexBuffer;
-		const int32 NumVertices = StaticMeshVertexBuffer.GetNumVertices();
-		UV0s.SetNumZeroed(NumVertices);
-		for (int32 Index = 0; Index < NumVertices; ++Index)
-		{
-			UV0s[Index] = StaticMeshVertexBuffer.GetVertexUV(Index, 0);
-		}
-		return UV0s;
-	}
-
-	TArray<FColor> GetColors(const FSkeletalMeshLODRenderData& LODData, int32 NumVertices)
+	TArrayView<TArray<FVector3f>> ShrinkToValidFrames(const TArrayView<TArray<FVector3f>>& Positions, int32 NumVertices)
 	{
-		TArray<FColor> Colors;
-		Colors.SetNum(NumVertices);
-		const FColorVertexBuffer& ColorVertexBuffer = LODData.StaticVertexBuffers.ColorVertexBuffer;
-		if (ColorVertexBuffer.GetNumVertices() == NumVertices)
-		{
-			for (int32 Index = 0; Index < NumVertices; ++Index)
-			{
-				Colors[Index] = ColorVertexBuffer.VertexColor(Index);
-			}
-		}
-		else
+		int32 NumValidFrames = 0;
+		for (const TArray<FVector3f>& Frame : Positions)
 		{
-			for (int32 Index = 0; Index < NumVertices; ++Index)
+			if (Frame.Num() != NumVertices)
 			{
-				Colors[Index] = FColor::White;
-			}
-		}
-		return Colors;
-	}
-
-	TArray<FPackedNormal> GetPackedNormals(TConstArrayView<FVector3f> Normals)
-	{
-		TArray<FPackedNormal> PackedNormals;
-		PackedNormals.Reserve(Normals.Num());
-		for (const FVector3f& Normal : Normals)
-		{
-			PackedNormals.Add(Normal);
-		}
-		return PackedNormals;
-	}
-
-	TArray<FVector3f> ComputeNormals(const ::Chaos::FTriangleMesh& TriangleMesh, TConstArrayView<FVector3f> Positions)
-	{
-		using TVec3 = ::Chaos::TVec3<float>;
-		TConstArrayView<TVec3> ChaosPositions(reinterpret_cast<const TVec3*>(Positions.GetData()), Positions.Num());
-		const TArray<TVec3> FaceNormals = TriangleMesh.GetFaceNormals(ChaosPositions, false);
-		TriangleMesh.GetPointToTriangleMap();
-		TArray<FVector3f> Normals;
-		Normals.SetNumUninitialized(Positions.Num());
-		TArrayView<TVec3> ChaosNormals(reinterpret_cast<TVec3*>(Normals.GetData()), Normals.Num());
-		TriangleMesh.GetPointNormals(ChaosNormals, TConstArrayView<TVec3>(FaceNormals), false);
-		return Normals;
-	}
-	
-	TArray<FVector3f> ComputeTangentsX(const ::Chaos::FTriangleMesh& TriangleMesh, TConstArrayView<FVector3f> Positions, TConstArrayView<FVector3f> Normals, TConstArrayView<FVector2f> UVs)
-	{
-		using ::Chaos::TVec3;
-		TArray<FVector3f> Tangents;
-		const int32 NumVertices = Positions.Num();
-		Tangents.SetNumZeroed(NumVertices);
-		TConstArrayView<TVec3<int32>> Elements{TriangleMesh.GetElements()};
-		for (const TVec3<int32>& Elment : Elements)
-		{
-			const FVector3f& P0 = Positions[Elment[0]];
-			const FVector3f& P1 = Positions[Elment[1]];
-			const FVector3f& P2 = Positions[Elment[2]];
-			const FVector2f& UV0 = UVs[Elment[0]];
-			const FVector2f& UV1 = UVs[Elment[1]];
-			const FVector2f& UV2 = UVs[Elment[2]];
-	
-			const FVector3f Edge1 = P1 - P0;
-			const FVector3f Edge2 = P2 - P0;
-			const FVector2f UVEdge1 = UV1 - UV0;
-			const FVector2f UVEdge2 = UV2 - UV0;
-			const float Det = UVEdge1.Y * UVEdge2.X - UVEdge1.X * UVEdge2.Y;
-			if (FMath::Abs(Det) > UE_SMALL_NUMBER)
-			{
-				FVector3f TangentX = (-Edge1 * UVEdge2.Y + Edge2 * UVEdge1.Y) / Det;
-				TangentX.Normalize();
-				TangentX *= Edge1.Cross(Edge2).Size();
-				for (int32 Index = 0; Index < 3; ++Index)
-				{
-					Tangents[Elment[Index]] += TangentX;
-				}
+				break;
 			}
+			++NumValidFrames;
 		}
-		for (int32 Index = 0; Index < NumVertices; Index++)
-		{
-			FVector3f& TangentX = Tangents[Index];
-			const FVector3f& Normal = Normals[Index];
-			TangentX -= Normal * Normal.Dot(TangentX);
-			TangentX = TangentX.GetSafeNormal(UE_SMALL_NUMBER, FVector3f(1.0f, 0.0f, 0.0f));
-		}
-		return Tangents;
+		return TArrayView<TArray<FVector3f>>(Positions.GetData(), NumValidFrames);
 	}
 
 	void SaveGeometryCache(UGeometryCache& GeometryCache, const USkinnedAsset& Asset, TConstArrayView<uint32> ImportedVertexNumbers, TArrayView<TArray<FVector3f>> PositionsToMoveFrom)
@@ -261,108 +154,21 @@ namespace UE::ClothTrainingTool::Private
 			return;
 		}
 		const FSkeletalMeshLODRenderData& LODData = RenderData->LODRenderData[LODIndex];
-		TArray<uint32> Indices;
-		LODData.MultiSizeIndexContainer.GetIndexBuffer(Indices);
-
-		using ::Chaos::TVec3;
-		TArray<TVec3<int32>> Elements;
-		const int32 NumTriangles = Indices.Num() / 3;
-		Elements.SetNumUninitialized(NumTriangles);
-		for(int32 TriangleIndex = 0; TriangleIndex < NumTriangles; ++TriangleIndex)
-		{
-			const int32 Index = TriangleIndex * 3;
-			Elements[TriangleIndex] = TVec3<int32>(Indices[Index], Indices[Index + 1], Indices[Index + 2]);
-		}
-		::Chaos::FTriangleMesh TriangleMesh(MoveTemp(Elements));
-
-		const int32 NumFrames = PositionsToMoveFrom.Num();
-		if (NumFrames == 0)
-		{
-			return;
-		}
 		const int32 NumVertices = GetNumVertices(LODData);
-		if (!ensureAlways(NumVertices == PositionsToMoveFrom[0].Num()))
-		{
-			return;
-		}
-		const TArray<FVector2f> UV0 = GetUV0s(LODData);
-		const bool bHasUV0 = UV0.Num() == NumVertices;
-		const TArray<FColor> Colors = GetColors(LODData, NumVertices);
-
-		GeometryCache.ClearForReimporting();
+		PositionsToMoveFrom = ShrinkToValidFrames(PositionsToMoveFrom, NumVertices);
 
-		const FName BaseName = FName(*Asset.GetName());
-		const FName CodecName = MakeUniqueObjectName(&GeometryCache, UGeometryCacheCodecV1::StaticClass(), FName(BaseName.ToString() + FString(TEXT("_Codec"))));
-		UGeometryCacheCodecV1* Codec = NewObject<UGeometryCacheCodecV1>(&GeometryCache, CodecName, RF_Public);
-
-		constexpr float CompressedPositionPrecision = 1e-3f;
-		constexpr int32 CompressedTextureCoordinatesNumberOfBits = 4;
-		Codec->InitializeEncoder(CompressedPositionPrecision, CompressedTextureCoordinatesNumberOfBits);
-		
-		const FName TrackName = MakeUniqueObjectName(&GeometryCache, UGeometryCacheTrackStreamable::StaticClass(), BaseName);
-		UGeometryCacheTrackStreamable* Track = NewObject<UGeometryCacheTrackStreamable>(&GeometryCache, TrackName, RF_Public);
-		
-		constexpr bool bApplyConstantTopologyOptimizations = true;
-		constexpr bool bCalculateMotionVectors = false;
-		constexpr bool bOptimizeIndexBuffers = true;
-		Track->BeginCoding(Codec, bApplyConstantTopologyOptimizations, bCalculateMotionVectors, bOptimizeIndexBuffers);
-
-		constexpr float CacheFPS = 30;
-		int32 MaxRecordedFrame = -1;
-		for (int32 Frame = 0; Frame < NumFrames; ++Frame)
-		{
-			const float Time = Frame / CacheFPS;
-			if (PositionsToMoveFrom[Frame].Num() != NumVertices)
-			{
-				break;
-			}
-			FGeometryCacheMeshData MeshData;
-			MeshData.Positions = MoveTemp(PositionsToMoveFrom[Frame]);
-			// Ideally we should not need to save normals to save disk space. Save normals for now for visualization.
-			const TArray<FVector3f> Normals = ComputeNormals(TriangleMesh, MeshData.Positions);
-			MeshData.VertexInfo.bHasTangentZ = true;
-			MeshData.TangentsZ = GetPackedNormals(Normals);
-			if (bHasUV0)
-			{
-				MeshData.VertexInfo.bHasUV0 = true;
-				MeshData.TextureCoordinates = UV0;
-				MeshData.VertexInfo.bHasTangentX = true;
-				MeshData.TangentsX = GetPackedNormals(ComputeTangentsX(TriangleMesh, MeshData.Positions, Normals, UV0));
-			}
-			MeshData.VertexInfo.bHasColor0 = true;
-			MeshData.Colors = Colors;
-			MeshData.VertexInfo.bHasImportedVertexNumbers = true;
-			MeshData.ImportedVertexNumbers = ImportedVertexNumbers;
-			MeshData.Indices = Indices;
-			MeshData.BoundingBox = GetBoundingBox(MeshData.Positions);
-
-			MeshData.BatchesInfo.Reserve(LODData.RenderSections.Num());
-			for (const FSkelMeshRenderSection& Section : LODData.RenderSections)
-			{
-				FGeometryCacheMeshBatchInfo BatchInfo;
-				BatchInfo.StartIndex = Section.BaseIndex;
-				BatchInfo.NumTriangles = Section.NumTriangles;
-				BatchInfo.MaterialIndex = Section.MaterialIndex;
-				MeshData.BatchesInfo.Add(MoveTemp(BatchInfo));
-			}
-
-			constexpr bool bConstTopology = true;
-			Track->AddMeshSample(MeshData, Time, bConstTopology);
-			constexpr bool bVisible = true;
-			Track->AddVisibilitySample(bVisible, Time);
-			MaxRecordedFrame = Frame;
-		}
-
-		TArray<FMatrix> Mats { FMatrix::Identity, FMatrix::Identity };
-		TArray<float> MatTimes { 0.0f, MaxRecordedFrame / CacheFPS };
-		Track->SetMatrixSamples(Mats, MatTimes);
-
-		if (ensureAlways(Track->EndCoding()))
+		using UE::GeometryCacheHelpers::FGeometryCacheConstantTopologyWriter;
+		using UE::GeometryCacheHelpers::AddTrackWriterFromSkinnedAsset;
+		using FTrackWriter = FGeometryCacheConstantTopologyWriter::FTrackWriter;
+		FGeometryCacheConstantTopologyWriter Writer(GeometryCache);
+		const int32 Index = AddTrackWriterFromSkinnedAsset(Writer, Asset);
+		if (Index == INDEX_NONE)
 		{
-			GeometryCache.AddTrack(Track);
-			GeometryCache.SetFrameStartEnd(0, NumFrames - 1);
+			return;
 		}
-		CopyMaterials(Asset, GeometryCache);
+		FTrackWriter& TrackWriter = Writer.GetTrackWriter(Index);
+		TrackWriter.ImportedVertexNumbers = ImportedVertexNumbers;
+		TrackWriter.WriteAndClose(PositionsToMoveFrom);
 	}
 
 	class FTimeScope
