Cloth Editor: node/tool workflow changes
- Clicking the buttons in the left-side panel will now create the corresponding Dataflow node in the graph and that's all
- Selecting a Dataflow node will automatically launch the corresponding tool, if any
- Tools should only modify the node parameters
- As long as no changes are made in the tool, the user can select a different node and the tool will auto-Cancel
- After any edit operation in the tool, the user needs to choose to either Accept or Cancel to end the tool

TODO:
- Dragging a node in the Graph Editor will start the tool. Want to avoid this if possible

diff --git a/Engine/Plugins/Experimental/ChaosClothAssetEditor/Source/ChaosClothAssetDataflowNodes/Private/ChaosClothAsset/AddWeightMapNode.cpp b/Engine/Plugins/Experimental/ChaosClothAssetEditor/Source/ChaosClothAssetDataflowNodes/Private/ChaosClothAsset/AddWeightMapNode.cpp
--- a/Engine/Plugins/Experimental/ChaosClothAssetEditor/Source/ChaosClothAssetDataflowNodes/Private/ChaosClothAsset/AddWeightMapNode.cpp
+++ b/Engine/Plugins/Experimental/ChaosClothAssetEditor/Source/ChaosClothAssetDataflowNodes/Private/ChaosClothAsset/AddWeightMapNode.cpp
@@ -35,7 +35,7 @@ void FChaosClothAssetAddWeightMapNode::Evaluate(Dataflow::FContext& Context, con
 
 		TArrayView<float> ClothWeights = ClothFacade.GetWeightMap(InName);
 		const int32 MaxWeightIndex = FMath::Min(VertexWeights.Num(), ClothWeights.Num());
-		if (VertexWeights.Num() != ClothWeights.Num())
+		if (VertexWeights.Num() > 0 && VertexWeights.Num() != ClothWeights.Num())
 		{
 			DataflowNodes::LogAndToastWarning(
 				FText::Format(LOCTEXT("WeightMapSize", "FChaosClothAssetAddWeightMapNode: Vertex count mismatch: vertex weights in the node: {0}; 3D vertices in cloth: {1}"),
diff --git a/Engine/Plugins/Experimental/ChaosClothAssetEditor/Source/ChaosClothAssetEditor/Private/ChaosClothAsset/ClothEditorCommands.cpp b/Engine/Plugins/Experimental/ChaosClothAssetEditor/Source/ChaosClothAssetEditor/Private/ChaosClothAsset/ClothEditorCommands.cpp
--- a/Engine/Plugins/Experimental/ChaosClothAssetEditor/Source/ChaosClothAssetEditor/Private/ChaosClothAsset/ClothEditorCommands.cpp
+++ b/Engine/Plugins/Experimental/ChaosClothAssetEditor/Source/ChaosClothAssetEditor/Private/ChaosClothAsset/ClothEditorCommands.cpp
@@ -12,8 +12,10 @@ namespace UE::Chaos::ClothAsset
 const FString FChaosClothAssetEditorCommands::BeginRemeshToolIdentifier = TEXT("BeginRemeshTool");
 const FString FChaosClothAssetEditorCommands::BeginAttributeEditorToolIdentifier = TEXT("BeginAttributeEditorTool");
 const FString FChaosClothAssetEditorCommands::BeginWeightMapPaintToolIdentifier = TEXT("BeginWeightMapPaintTool");
+const FString FChaosClothAssetEditorCommands::AddWeightMapNodeIdentifier = TEXT("AddWeightMapNode");
 const FString FChaosClothAssetEditorCommands::BeginClothTrainingToolIdentifier = TEXT("BeginClothTrainingTool");
 const FString FChaosClothAssetEditorCommands::BeginTransferSkinWeightsToolIdentifier = TEXT("BeginTransferSkinWeightsTool");
+const FString FChaosClothAssetEditorCommands::AddTransferSkinWeightsNodeIdentifier = TEXT("AddTransferSkinWeightsNode");
 const FString FChaosClothAssetEditorCommands::ToggleSimulationSuspendedIdentifier = TEXT("ToggleSimulationSuspended");
 const FString FChaosClothAssetEditorCommands::SoftResetSimulationIdentifier = TEXT("SoftResetSimulation");
 const FString FChaosClothAssetEditorCommands::HardResetSimulationIdentifier = TEXT("HardResetSimulation");
@@ -37,11 +39,14 @@ void FChaosClothAssetEditorCommands::RegisterCommands()
 
 	UI_COMMAND(BeginRemeshTool, "Remesh", "Remesh the selected mesh", EUserInterfaceActionType::Button, FInputChord());
 	UI_COMMAND(BeginAttributeEditorTool, "AttrEd", "Edit/configure mesh attributes", EUserInterfaceActionType::Button, FInputChord());
-	UI_COMMAND(BeginWeightMapPaintTool, "MapPnt", "Paint Weight Maps on the mesh", EUserInterfaceActionType::Button, FInputChord());
+	
+	UI_COMMAND(BeginWeightMapPaintTool, "MapPnt", "Paint weight maps on the mesh", EUserInterfaceActionType::None, FInputChord());
+	UI_COMMAND(AddWeightMapNode, "MapPnt", "Paint weight maps on the mesh", EUserInterfaceActionType::Button, FInputChord());
 
 	UI_COMMAND(BeginClothTrainingTool, "Generate Train Data", "Generate cloth training data", EUserInterfaceActionType::Button, FInputChord());
 
-	UI_COMMAND(BeginTransferSkinWeightsTool, "Transfer Skin Weights", "Launch Transfer Skin Weights tool", EUserInterfaceActionType::Button, FInputChord());
+	UI_COMMAND(BeginTransferSkinWeightsTool, "TransferSkinWeights", "Transfer skinning weights from a SkeletalMesh", EUserInterfaceActionType::None, FInputChord());
+	UI_COMMAND(AddTransferSkinWeightsNode, "TransferSkinWeights", "Transfer skinning weights from a SkeletalMesh", EUserInterfaceActionType::Button, FInputChord());
 
 	UI_COMMAND(SetConstructionMode2D, "2D Sim", "Switches the viewport to 2D simulation mesh view", EUserInterfaceActionType::RadioButton, FInputChord());
 	UI_COMMAND(SetConstructionMode3D, "3D Sim", "Switches the viewport to 3D simulation mesh view", EUserInterfaceActionType::RadioButton, FInputChord());
diff --git a/Engine/Plugins/Experimental/ChaosClothAssetEditor/Source/ChaosClothAssetEditor/Private/ChaosClothAsset/ClothEditorCommands.h b/Engine/Plugins/Experimental/ChaosClothAssetEditor/Source/ChaosClothAssetEditor/Private/ChaosClothAsset/ClothEditorCommands.h
--- a/Engine/Plugins/Experimental/ChaosClothAssetEditor/Source/ChaosClothAssetEditor/Private/ChaosClothAsset/ClothEditorCommands.h
+++ b/Engine/Plugins/Experimental/ChaosClothAssetEditor/Source/ChaosClothAssetEditor/Private/ChaosClothAsset/ClothEditorCommands.h
@@ -32,8 +32,11 @@ public:
 
 	const static FString BeginRemeshToolIdentifier;
 	TSharedPtr<FUICommandInfo> BeginRemeshTool;
+
 	const static FString BeginWeightMapPaintToolIdentifier;
 	TSharedPtr<FUICommandInfo> BeginWeightMapPaintTool;
+	const static FString AddWeightMapNodeIdentifier;
+	TSharedPtr<FUICommandInfo> AddWeightMapNode;
 
 	const static FString BeginAttributeEditorToolIdentifier;
 	TSharedPtr<FUICommandInfo> BeginAttributeEditorTool;
@@ -43,6 +46,8 @@ public:
 
 	const static FString BeginTransferSkinWeightsToolIdentifier;
 	TSharedPtr<FUICommandInfo> BeginTransferSkinWeightsTool;
+	const static FString AddTransferSkinWeightsNodeIdentifier;
+	TSharedPtr<FUICommandInfo> AddTransferSkinWeightsNode;
 
 	// Construction viewport commands
 	const static FString ToggleConstructionViewWireframeIdentifier;
diff --git a/Engine/Plugins/Experimental/ChaosClothAssetEditor/Source/ChaosClothAssetEditor/Private/ChaosClothAsset/ClothEditorMode.cpp b/Engine/Plugins/Experimental/ChaosClothAssetEditor/Source/ChaosClothAssetEditor/Private/ChaosClothAsset/ClothEditorMode.cpp
--- a/Engine/Plugins/Experimental/ChaosClothAssetEditor/Source/ChaosClothAssetEditor/Private/ChaosClothAsset/ClothEditorMode.cpp
+++ b/Engine/Plugins/Experimental/ChaosClothAssetEditor/Source/ChaosClothAssetEditor/Private/ChaosClothAsset/ClothEditorMode.cpp
@@ -10,6 +10,8 @@
 #include "ChaosClothAsset/ClothPatternToDynamicMesh.h"
 #include "ChaosClothAsset/ClothEditorPreviewScene.h"
 #include "ChaosClothAsset/ClothEditorContextObject.h"
+#include "ChaosClothAsset/AddWeightMapNode.h"
+#include "ChaosClothAsset/TransferSkinWeightsNode.h"
 #include "AssetEditorModeManager.h"
 #include "Drawing/MeshElementsVisualizer.h"
 #include "EditorViewportClient.h"
@@ -58,6 +60,7 @@
 #include "BaseGizmos/TransformGizmoUtil.h"
 #include "Materials/Material.h"
 #include "MaterialDomain.h"
+#include "Dataflow/DataflowSNode.h"
 
 #include UE_INLINE_GENERATED_CPP_BY_NAME(ClothEditorMode)
 
@@ -168,6 +171,33 @@ void UChaosClothAssetEditorMode::RegisterClothTool(TSharedPtr<FUICommandInfo> UI
 
 }
 
+void UChaosClothAssetEditorMode::RegisterAddNodeCommand(TSharedPtr<FUICommandInfo> AddNodeCommand, const FName& NewNodeType, TSharedPtr<FUICommandInfo> StartToolCommand)
+{
+	auto AddNode = [this](const FName& NewNodeType)
+	{
+		const FName ConnectionType = FManagedArrayCollection::StaticType();
+		UEdGraphNode* const CurrentlySelectedNode = GetSingleSelectedNodeWithOutputType(ConnectionType);
+		checkf(CurrentlySelectedNode, TEXT("No node with FManagedArrayCollection output is currently selected in the Dataflow graph"));
+		verifyf(CreateAndConnectNewNode(NewNodeType, *CurrentlySelectedNode, ConnectionType), TEXT("Failed to create a new node: %s"), *NewNodeType.ToString());
+	};
+
+	auto CanAddNode = [this](const FName& NewNodeType) -> bool
+	{
+		const UEdGraphNode* const CurrentlySelectedNode = GetSingleSelectedNodeWithOutputType(FManagedArrayCollection::StaticType());
+		return (CurrentlySelectedNode != nullptr);
+	};
+
+	const TSharedRef<FUICommandList>& CommandList = Toolkit->GetToolkitCommands();
+
+	CommandList->MapAction(AddNodeCommand,
+		FExecuteAction::CreateWeakLambda(this, AddNode, NewNodeType),
+		FCanExecuteAction::CreateWeakLambda(this, CanAddNode, NewNodeType)
+	);
+
+	NodeTypeToToolCommandMap.Add(NewNodeType, StartToolCommand);
+}
+
+
 void UChaosClothAssetEditorMode::RegisterPreviewTools()
 {
 	using namespace UE::Chaos::ClothAsset;
@@ -184,24 +214,21 @@ void UChaosClothAssetEditorMode::RegisterTools()
 
 	const FChaosClothAssetEditorCommands& CommandInfos = FChaosClothAssetEditorCommands::Get();
 
-	// Note that the identifiers below need to match the command names so that the tool icons can 
-	// be easily retrieved from the active tool name in ChaosClothAssetEditorModeToolkit::OnToolStarted. Otherwise
-	// we would need to keep some other mapping from tool identifier to tool icon.
+	UEditorInteractiveToolsContext* const ConstructionViewportToolsContext = GetInteractiveToolsContext();
 
-	// TODO: Re-add the remesh tool when we have a way to remesh both 2d and 3d rest space meshes at the same time
-	//RegisterTool(CommandInfos.BeginRemeshTool, FChaosClothAssetEditorCommands::BeginRemeshToolIdentifier, NewObject<URemeshMeshToolBuilder>());
+	RegisterClothTool(CommandInfos.BeginAttributeEditorTool, FChaosClothAssetEditorCommands::BeginAttributeEditorToolIdentifier, NewObject<UAttributeEditorToolBuilder>(), ConstructionViewportToolsContext);
 
-	UEditorInteractiveToolsContext* const RestSpaceToolsContext = GetInteractiveToolsContext();
+	RegisterClothTool(CommandInfos.BeginWeightMapPaintTool, FChaosClothAssetEditorCommands::BeginWeightMapPaintToolIdentifier, NewObject<UClothEditorWeightMapPaintToolBuilder>(), ConstructionViewportToolsContext);
+	RegisterAddNodeCommand(CommandInfos.AddWeightMapNode, FChaosClothAssetAddWeightMapNode::StaticType(), CommandInfos.BeginWeightMapPaintTool);
 
-	RegisterClothTool(CommandInfos.BeginAttributeEditorTool, FChaosClothAssetEditorCommands::BeginAttributeEditorToolIdentifier, NewObject<UAttributeEditorToolBuilder>(), RestSpaceToolsContext);
-	RegisterClothTool(CommandInfos.BeginWeightMapPaintTool, FChaosClothAssetEditorCommands::BeginWeightMapPaintToolIdentifier, NewObject<UClothEditorWeightMapPaintToolBuilder>(), RestSpaceToolsContext);
-	RegisterClothTool(CommandInfos.BeginTransferSkinWeightsTool, FChaosClothAssetEditorCommands::BeginTransferSkinWeightsToolIdentifier, NewObject<UClothTransferSkinWeightsToolBuilder>(), RestSpaceToolsContext);
+	RegisterClothTool(CommandInfos.BeginTransferSkinWeightsTool, FChaosClothAssetEditorCommands::BeginTransferSkinWeightsToolIdentifier, NewObject<UClothTransferSkinWeightsToolBuilder>(), ConstructionViewportToolsContext);
+	RegisterAddNodeCommand(CommandInfos.AddTransferSkinWeightsNode, FChaosClothAssetTransferSkinWeightsNode::StaticType(), CommandInfos.BeginTransferSkinWeightsTool);
 }
 
 bool UChaosClothAssetEditorMode::ShouldToolStartBeAllowed(const FString& ToolIdentifier) const
 {
-	// For now we've decided to disallow switch-away on any tools in the Cloth editor.
-	if (GetInteractiveToolsContext()->HasActiveTool())
+	// Allow switching away from tool if no changes have been made in the tool yet (which we infer from the CanAccept status)
+	if (GetInteractiveToolsContext()->CanAcceptActiveTool())
 	{
 		return false;
 	}
@@ -229,12 +256,6 @@ void UChaosClothAssetEditorMode::OnToolStarted(UInteractiveToolManager* Manager,
 	FChaosClothAssetEditorCommands::UpdateToolCommandBinding(Tool, ToolCommandList, false);
 
 	bCanChangeConstructionViewMode = false;
-
-	// Don't let the Graph change during tool execution (or even which nodes are selected)
-	if (TSharedPtr<SDataflowGraphEditor> GraphEditor = DataflowGraphEditor.Pin())
-	{
-		GraphEditor->SetEnabled(false);
-	}
 }
 
 void UChaosClothAssetEditorMode::OnToolEnded(UInteractiveToolManager* Manager, UInteractiveTool* Tool)
@@ -243,8 +264,6 @@ void UChaosClothAssetEditorMode::OnToolEnded(UInteractiveToolManager* Manager, U
 
 	bCanChangeConstructionViewMode = true;
 
-	ActiveToolsContext = nullptr;
-
 	ReinitializeDynamicMeshComponents();
 
 	if (TSharedPtr<SDataflowGraphEditor> GraphEditor = DataflowGraphEditor.Pin())
@@ -681,6 +700,22 @@ void UChaosClothAssetEditorMode::ModeTick(float DeltaTime)
 {
 	Super::ModeTick(DeltaTime);
 	
+
+	if (TSharedPtr<SDataflowGraphEditor> GraphEditor = DataflowGraphEditor.Pin())
+	{
+		// For now don't allow selection change once the tool has uncommitted changes
+		// TODO: We might want to auto-accept unsaved changes and allow switching between nodes
+		if (GetInteractiveToolsContext()->CanAcceptActiveTool())
+		{
+			GraphEditor->SetEnabled(false);
+		}
+		else
+		{
+			GraphEditor->SetEnabled(true);
+		}
+	}
+
+
 	for (TObjectPtr<UInteractiveToolPropertySet>& Propset : PropertyObjectsToTick)
 	{
 		if (Propset)
@@ -841,7 +876,7 @@ void UChaosClothAssetEditorMode::InitializeContextObject()
 		RestSpaceToolsContext->ContextObjectStore->AddContextObject(EditorContextObject);
 	}
 
-	EditorContextObject->Init(DataflowGraphEditor, DataflowGraph, ConstructionViewMode, SelectedClothCollection);
+	EditorContextObject->Init(DataflowGraphEditor, ConstructionViewMode, SelectedClothCollection);
 
 	check(EditorContextObject);
 
@@ -869,5 +904,170 @@ void UChaosClothAssetEditorMode::SetDataflowGraphEditor(TSharedPtr<SDataflowGrap
 	}
 }
 
+void UChaosClothAssetEditorMode::OnDataflowNodeSelectionChanged(const TSet<UObject*>& NewSelection)
+{
+	using namespace UE::Chaos::ClothAsset;
+
+	// Start the tool associated with the selected node, if any
+
+	const TSharedRef<FUICommandList> CommandList = Toolkit->GetToolkitCommands();
+	const FChaosClothAssetEditorCommands& CommandInfos = FChaosClothAssetEditorCommands::Get();
+
+	bool bNewToolStarted = false;
+	if (NewSelection.Num() == 1)
+	{
+		if (const UDataflowEdNode* const Node = Cast<UDataflowEdNode>(*NewSelection.CreateConstIterator()))
+		{
+			if (const TSharedPtr<const FDataflowNode> DataflowNode = Node->GetDataflowNode())
+			{
+				const FName DataflowNodeName = DataflowNode->GetType();
+				if (NodeTypeToToolCommandMap.Contains(DataflowNodeName))
+				{
+					CommandList->TryExecuteAction(NodeTypeToToolCommandMap[DataflowNodeName].ToSharedRef());
+					bNewToolStarted = true;
+				}
+			}
+		}
+	}
+
+	UEditorInteractiveToolsContext* const ToolsContext = GetInteractiveToolsContext();
+	checkf(ToolsContext, TEXT("No valid ToolsContext found for UChaosClothAssetEditorMode"));
+	if (!bNewToolStarted && ToolsContext->HasActiveTool())
+	{
+		// The user has clicked away from the selected node, end the tool
+		// TODO: If we wanted we could probably auto-accept the pending changes in the active tool
+		ToolsContext->EndTool(EToolShutdownType::Completed);
+	}
+
+}
+
+
+void UChaosClothAssetEditorMode::OnDataflowNodeDeleted(const TSet<UObject*>& DeletedNodes)
+{
+	UEditorInteractiveToolsContext* const ToolsContext = GetInteractiveToolsContext();
+	checkf(ToolsContext, TEXT("No valid ToolsContext found for UChaosClothAssetEditorMode"));
+	const bool bCanCancel = ToolsContext->CanCancelActiveTool();
+	ToolsContext->EndTool(bCanCancel ? EToolShutdownType::Cancel : EToolShutdownType::Completed);
+}
+
+UEdGraphNode* UChaosClothAssetEditorMode::GetSingleSelectedNodeWithOutputType(const FName& SelectedNodeOutputTypeName) const
+{
+	const TSharedPtr<const SDataflowGraphEditor> PinnedDataflowGraphEditor = DataflowGraphEditor.Pin();
+	if (!PinnedDataflowGraphEditor)
+	{
+		return nullptr;
+	}
+
+	UEdGraphNode* const SelectedNode = PinnedDataflowGraphEditor->GetSingleSelectedNode();
+	if (!SelectedNode)
+	{
+		return nullptr;
+	}
+
+	const UDataflowEdNode* const SelectedDataflowEdNode = CastChecked<UDataflowEdNode>(SelectedNode);
+	const TSharedPtr<const FDataflowNode> SelectedDataflowNode = SelectedDataflowEdNode->GetDataflowNode();
+
+	if (!SelectedDataflowNode)
+	{
+		// This can happen when the user deletes a node. Seems like the Dataflow FGraph is updated with the removed node before the graph editor can update.
+		return nullptr;
+	}
+
+	for (const FDataflowOutput* const Output : SelectedDataflowNode->GetOutputs())
+	{
+		if (Output->GetType() == SelectedNodeOutputTypeName)
+		{
+			return SelectedNode;
+		}
+	}
+
+	return nullptr;
+}
+
+UEdGraphNode* UChaosClothAssetEditorMode::CreateNewNode(const FName& NewNodeTypeName)
+{
+	const TSharedPtr<const SDataflowGraphEditor> PinnedDataflowGraphEditor = DataflowGraphEditor.Pin();
+	if (!PinnedDataflowGraphEditor)
+	{
+		return nullptr;
+	}
+
+	const TSharedPtr<FAssetSchemaAction_Dataflow_CreateNode_DataflowEdNode> NodeAction =
+		FAssetSchemaAction_Dataflow_CreateNode_DataflowEdNode::CreateAction(DataflowGraph, NewNodeTypeName);
+	constexpr UEdGraphPin* FromPin = nullptr;
+	constexpr bool bSelectNewNode = true;
+	UEdGraphNode* const NewEdNode = NodeAction->PerformAction(DataflowGraph, FromPin, PinnedDataflowGraphEditor->GetPasteLocation(), bSelectNewNode);
+
+	return NewEdNode;
+}
+
+
+UEdGraphNode* UChaosClothAssetEditorMode::CreateAndConnectNewNode(const FName& NewNodeTypeName, UEdGraphNode& UpstreamNode, const FName& ConnectionTypeName)
+{
+	// First find the specified output of the upstream node, plus any pins it's connected to
+
+	UEdGraphPin* UpstreamNodeOutputPin = nullptr;
+	TArray<UEdGraphPin*> ExistingNodeInputPins;
+
+	const UDataflowEdNode* const UpstreamDataflowEdNode = CastChecked<UDataflowEdNode>(&UpstreamNode);
+	const TSharedPtr<const FDataflowNode> UpstreamDataflowNode = UpstreamDataflowEdNode->GetDataflowNode();
+
+	for (const FDataflowOutput* const Output : UpstreamDataflowNode->GetOutputs())
+	{
+		if (Output->GetType() == ConnectionTypeName)
+		{
+			UpstreamNodeOutputPin = UpstreamDataflowEdNode->FindPin(*Output->GetName().ToString(), EGPD_Output);
+			ExistingNodeInputPins = UpstreamNodeOutputPin->LinkedTo;
+			break;
+		}
+	}
+
+	// Add the new node 
+
+	UEdGraphNode* const NewEdNode = CreateNewNode(NewNodeTypeName);
+	checkf(NewEdNode, TEXT("Failed to create a new node in the DataflowGraph"));
+
+	UDataflowEdNode* const NewDataflowEdNode = CastChecked<UDataflowEdNode>(NewEdNode);
+	const TSharedPtr<FDataflowNode> NewDataflowNode = NewDataflowEdNode->GetDataflowNode();
+
+	// Re-wire the graph
+
+	if (UpstreamNodeOutputPin)
+	{
+		UEdGraphPin* NewNodeInputPin = nullptr;
+		for (const FDataflowInput* const NewNodeInput : NewDataflowNode->GetInputs())
+		{
+			if (NewNodeInput->GetType() == ConnectionTypeName)
+			{
+				NewNodeInputPin = NewDataflowEdNode->FindPin(*NewNodeInput->GetName().ToString(), EGPD_Input);
+			}
+		}
+
+		UEdGraphPin* NewNodeOutputPin = nullptr;
+		for (const FDataflowOutput* const NewNodeOutput : NewDataflowNode->GetOutputs())
+		{
+			if (NewNodeOutput->GetType() == ConnectionTypeName)
+			{
+				NewNodeOutputPin = NewDataflowEdNode->FindPin(*NewNodeOutput->GetName().ToString(), EGPD_Output);
+				break;
+			}
+		}
+
+		check(NewNodeInputPin);
+		check(NewNodeOutputPin);
+
+		DataflowGraph->GetSchema()->TryCreateConnection(UpstreamNodeOutputPin, NewNodeInputPin);
+
+		for (UEdGraphPin* DownstreamInputPin : ExistingNodeInputPins)
+		{
+			DataflowGraph->GetSchema()->TryCreateConnection(NewNodeOutputPin, DownstreamInputPin);
+		}
+	}
+
+	DataflowGraph->NotifyGraphChanged();
+
+	return NewEdNode;
+}
+
 #undef LOCTEXT_NAMESPACE
 
diff --git a/Engine/Plugins/Experimental/ChaosClothAssetEditor/Source/ChaosClothAssetEditor/Private/ChaosClothAsset/ClothEditorMode.h b/Engine/Plugins/Experimental/ChaosClothAssetEditor/Source/ChaosClothAssetEditor/Private/ChaosClothAsset/ClothEditorMode.h
--- a/Engine/Plugins/Experimental/ChaosClothAssetEditor/Source/ChaosClothAssetEditor/Private/ChaosClothAsset/ClothEditorMode.h
+++ b/Engine/Plugins/Experimental/ChaosClothAssetEditor/Source/ChaosClothAssetEditor/Private/ChaosClothAsset/ClothEditorMode.h
@@ -32,11 +32,11 @@ struct FManagedArrayCollection;
 namespace UE::Chaos::ClothAsset
 {
 class FChaosClothPreviewScene;
-class FChaosClothAssetEditor3DViewportClient;
 class FChaosClothAssetEditorModeToolkit;
 class FChaosClothAssetEditorToolkit;
 class FChaosClothEditorRestSpaceViewportClient;
 }
+class UEdGraphNode;
 
 /**
  * The cloth editor mode is the mode used in the cloth asset editor. It holds most of the inter-tool state.
@@ -122,6 +122,8 @@ private:
 		UEditorInteractiveToolsContext* UseToolsContext, 
 		EToolsContextScope ToolScope = EToolsContextScope::Default);
 
+	void RegisterAddNodeCommand(TSharedPtr<FUICommandInfo> AddNodeCommand, const FName& NewNodeType, TSharedPtr<FUICommandInfo> StartToolCommand);
+
 	// Register the set of tools that operate on objects in the 3D preview world 
 	void RegisterPreviewTools();
 
@@ -140,6 +142,40 @@ private:
 	TSharedPtr<FManagedArrayCollection> GetClothCollection();
 
 	void SetDataflowGraphEditor(TSharedPtr<SDataflowGraphEditor> InGraphEditor);
+	
+	void OnDataflowNodeSelectionChanged(const TSet<UObject*>& NewSelection);
+	void OnDataflowNodeDeleted(const TSet<UObject*>& DeletedNodes);
+
+	/**
+	* Return the single selected node in the Dataflow Graph Editor only if it has an output of the specified type
+	* If there is not a single node selected, or if it does not have the specified output, return null
+	*/
+	UEdGraphNode* GetSingleSelectedNodeWithOutputType(const FName& SelectedNodeOutputTypeName) const;
+
+	/**
+	 * Create a node with the specified type in the graph
+	*/
+	UEdGraphNode* CreateNewNode(const FName& NewNodeTypeName);
+
+	/** Create a node with the specified type, then connect it to the output of the specified UpstreamNode.
+	* If the specified output of the upstream node is already connected to another node downstream, we first break
+	* that connecttion, then insert the new node along the previous connection.
+	* We want to turn this:
+	*
+	* [UpstreamNode] ----> [DownstreamNode(s)]
+	*
+	* to this:
+	*
+	* [UpstreamNode] ----> [NewNode] ----> [DownstreamNode(s)]
+	*
+	*
+	* @param NewNodeTypeName The type of node to create, by name
+	* @param UpstreamNode Node to connect the new node to
+	* @param ConnectionTypeName The type of output of the upstream node to connect our new node to
+	* @return The newly-created node
+	*/
+	UEdGraphNode* CreateAndConnectNewNode(const FName& NewNodeTypeName,	UEdGraphNode& UpstreamNode,	const FName& ConnectionTypeName);
+
 
 	void InitializeContextObject();
 	void DeleteContextObject();
@@ -199,7 +235,9 @@ private:
 	UPROPERTY()
 	TObjectPtr<UEditorInteractiveToolsContext> ActiveToolsContext = nullptr;
 
-
 	TSharedPtr<FManagedArrayCollection> SelectedClothCollection = nullptr;
+
+	// Correspondence between node types and commands to launch tools
+	TMap<FName, TSharedPtr<const FUICommandInfo>> NodeTypeToToolCommandMap;
 };
 
diff --git a/Engine/Plugins/Experimental/ChaosClothAssetEditor/Source/ChaosClothAssetEditor/Private/ChaosClothAsset/ClothEditorModeToolkit.cpp b/Engine/Plugins/Experimental/ChaosClothAssetEditor/Source/ChaosClothAssetEditor/Private/ChaosClothAsset/ClothEditorModeToolkit.cpp
--- a/Engine/Plugins/Experimental/ChaosClothAssetEditor/Source/ChaosClothAssetEditor/Private/ChaosClothAsset/ClothEditorModeToolkit.cpp
+++ b/Engine/Plugins/Experimental/ChaosClothAssetEditor/Source/ChaosClothAssetEditor/Private/ChaosClothAsset/ClothEditorModeToolkit.cpp
@@ -147,9 +147,9 @@ void FChaosClothAssetEditorModeToolkit::BuildToolPalette(FName PaletteIndex, cla
 	// TODO: make separate palettes to split up the tools tab
 	if (PaletteIndex == FBaseCharacterFXEditorModeToolkit::ToolsTabName)
 	{
-		ToolbarBuilder.AddToolBarButton(Commands.BeginWeightMapPaintTool);
+		ToolbarBuilder.AddToolBarButton(Commands.AddWeightMapNode);
 		ToolbarBuilder.AddToolBarButton(Commands.BeginClothTrainingTool);
-		ToolbarBuilder.AddToolBarButton(Commands.BeginTransferSkinWeightsTool);
+		ToolbarBuilder.AddToolBarButton(Commands.AddTransferSkinWeightsNode);
 	}
 }
 
diff --git a/Engine/Plugins/Experimental/ChaosClothAssetEditor/Source/ChaosClothAssetEditor/Private/ChaosClothAsset/ClothEditorStyle.cpp b/Engine/Plugins/Experimental/ChaosClothAssetEditor/Source/ChaosClothAssetEditor/Private/ChaosClothAsset/ClothEditorStyle.cpp
--- a/Engine/Plugins/Experimental/ChaosClothAssetEditor/Source/ChaosClothAssetEditor/Private/ChaosClothAsset/ClothEditorStyle.cpp
+++ b/Engine/Plugins/Experimental/ChaosClothAssetEditor/Source/ChaosClothAssetEditor/Private/ChaosClothAsset/ClothEditorStyle.cpp
@@ -48,7 +48,7 @@ FChaosClothAssetEditorStyle::FChaosClothAssetEditorStyle()
 	PropertyNameString = "ChaosClothAssetEditor." + FChaosClothAssetEditorCommands::BeginAttributeEditorToolIdentifier;
 	Set(*PropertyNameString, new FSlateImageBrush(FChaosClothAssetEditorStyle::InContent("Icons/AttributeEditor_40x", ".png"), ToolbarIconSize));
 
-	PropertyNameString = "ChaosClothAssetEditor." + FChaosClothAssetEditorCommands::BeginWeightMapPaintToolIdentifier;
+	PropertyNameString = "ChaosClothAssetEditor." + FChaosClothAssetEditorCommands::AddWeightMapNodeIdentifier;
 	Set(*PropertyNameString, new FSlateImageBrush(FChaosClothAssetEditorStyle::InContent("Icons/ModelingAttributePaint_x40", ".png"), ToolbarIconSize));
 
 	// TODO: Get a permanent icon to replace this placeholder
@@ -56,7 +56,7 @@ FChaosClothAssetEditorStyle::FChaosClothAssetEditorStyle()
 	Set(*PropertyNameString, new FSlateImageBrush(FChaosClothAssetEditorStyle::InContent("Icons/ModelingAttributePaint_x40", ".png"), ToolbarIconSize));
 
 	// TODO: Get a permanent icon to replace this placeholder
-	PropertyNameString = "ChaosClothAssetEditor." + FChaosClothAssetEditorCommands::BeginTransferSkinWeightsToolIdentifier;
+	PropertyNameString = "ChaosClothAssetEditor." + FChaosClothAssetEditorCommands::AddTransferSkinWeightsNodeIdentifier;
 	Set(*PropertyNameString, new FSlateImageBrush(FChaosClothAssetEditorStyle::InContent("Icons/ModelingAttributePaint_x40", ".png"), ToolbarIconSize));
 
 	//
diff --git a/Engine/Plugins/Experimental/ChaosClothAssetEditor/Source/ChaosClothAssetEditor/Private/ChaosClothAsset/ClothEditorToolkit.cpp b/Engine/Plugins/Experimental/ChaosClothAssetEditor/Source/ChaosClothAssetEditor/Private/ChaosClothAsset/ClothEditorToolkit.cpp
--- a/Engine/Plugins/Experimental/ChaosClothAssetEditor/Source/ChaosClothAssetEditor/Private/ChaosClothAsset/ClothEditorToolkit.cpp
+++ b/Engine/Plugins/Experimental/ChaosClothAssetEditor/Source/ChaosClothAssetEditor/Private/ChaosClothAsset/ClothEditorToolkit.cpp
@@ -15,6 +15,10 @@
 #include "ChaosClothAsset/ClothEditorModeToolkit.h"
 #include "ChaosClothAsset/ClothCollection.h"
 #include "ChaosClothAsset/CollectionClothFacade.h"
+#include "ChaosClothAsset/ClothEditorCommands.h"
+#include "ChaosClothAsset/AddWeightMapNode.h"
+#include "ChaosClothAsset/TransferSkinWeightsNode.h"
+#include "EdModeInteractiveToolsContext.h"
 #include "Framework/Docking/LayoutExtender.h"
 #include "AssetEditorModeManager.h"
 #include "AdvancedPreviewScene.h"
@@ -1009,6 +1013,7 @@ void FChaosClothAssetEditorToolkit::OnNodeSelectionChanged(const TSet<UObject*>&
 	if (ClothMode)
 	{
 		ClothMode->SetSelectedClothCollection(Collection);
+		ClothMode->OnDataflowNodeSelectionChanged(NewSelection);
 	}
 
 	if (Outliner)
@@ -1033,6 +1038,7 @@ void FChaosClothAssetEditorToolkit::OnNodeDeleted(const TSet<UObject*>& DeletedN
 	if (ClothMode)
 	{
 		ClothMode->SetSelectedClothCollection(nullptr);
+		ClothMode->OnDataflowNodeDeleted(DeletedNodes);
 	}
 
 }
diff --git a/Engine/Plugins/Experimental/ChaosClothAssetEditor/Source/ChaosClothAssetEditorTools/Private/ChaosClothAsset/ClothEditorContextObject.cpp b/Engine/Plugins/Experimental/ChaosClothAssetEditor/Source/ChaosClothAssetEditorTools/Private/ChaosClothAsset/ClothEditorContextObject.cpp
--- a/Engine/Plugins/Experimental/ChaosClothAssetEditor/Source/ChaosClothAssetEditorTools/Private/ChaosClothAsset/ClothEditorContextObject.cpp
+++ b/Engine/Plugins/Experimental/ChaosClothAssetEditor/Source/ChaosClothAssetEditorTools/Private/ChaosClothAsset/ClothEditorContextObject.cpp
@@ -8,158 +8,13 @@
 
 #include UE_INLINE_GENERATED_CPP_BY_NAME(ClothEditorContextObject)
 
-void UClothEditorContextObject::Init(TWeakPtr<SDataflowGraphEditor> InDataflowGraphEditor, TObjectPtr<UDataflow> InDataflowGraph, UE::Chaos::ClothAsset::EClothPatternVertexType InConstructionViewMode, TWeakPtr<FManagedArrayCollection> InSelectedClothCollection)
+void UClothEditorContextObject::Init(TWeakPtr<SDataflowGraphEditor> InDataflowGraphEditor,UE::Chaos::ClothAsset::EClothPatternVertexType InConstructionViewMode, TWeakPtr<FManagedArrayCollection> InSelectedClothCollection)
 {
 	DataflowGraphEditor = InDataflowGraphEditor;
-	DataflowGraph = InDataflowGraph;
 	ConstructionViewMode = InConstructionViewMode;
 	SelectedClothCollection = InSelectedClothCollection;
 }
 
-TWeakPtr<SDataflowGraphEditor> UClothEditorContextObject::GetDataflowGraphEditor()
-{
-	return DataflowGraphEditor;
-}
-
-const TWeakPtr<const SDataflowGraphEditor> UClothEditorContextObject::GetDataflowGraphEditor() const
-{
-	return DataflowGraphEditor;
-}
-
-TObjectPtr<UDataflow> UClothEditorContextObject::GetDataflowGraph()
-{
-	return DataflowGraph;
-}
-
-const TObjectPtr<const UDataflow> UClothEditorContextObject::GetDataflowGraph() const
-{
-	return DataflowGraph;
-}
-
-UEdGraphNode* UClothEditorContextObject::GetSingleSelectedNode() const
-{
-	const TSharedPtr<const SDataflowGraphEditor> PinnedDataflowGraphEditor = DataflowGraphEditor.Pin();
-	if (!PinnedDataflowGraphEditor)
-	{
-		return nullptr;
-	}
-
-	return PinnedDataflowGraphEditor->GetSingleSelectedNode();
-}
-
-
-UEdGraphNode* UClothEditorContextObject::GetSingleSelectedNodeWithOutputType(const FName& SelectedNodeOutputTypeName) const
-{
-	const TSharedPtr<const SDataflowGraphEditor> PinnedDataflowGraphEditor = DataflowGraphEditor.Pin();
-	if (!PinnedDataflowGraphEditor)
-	{
-		return nullptr;
-	}
-
-	UEdGraphNode* const SelectedNode = PinnedDataflowGraphEditor->GetSingleSelectedNode();
-	if (!SelectedNode)
-	{
-		return nullptr;
-	}
-
-	const UDataflowEdNode* const SelectedDataflowEdNode = CastChecked<UDataflowEdNode>(SelectedNode);
-	const TSharedPtr<const FDataflowNode> SelectedDataflowNode = SelectedDataflowEdNode->GetDataflowNode();
-	for (const FDataflowOutput* const Output : SelectedDataflowNode->GetOutputs())
-	{
-		if (Output->GetType() == SelectedNodeOutputTypeName)
-		{
-			return SelectedNode;
-		}
-	}
-
-	return nullptr;
-}
-
-UEdGraphNode* UClothEditorContextObject::CreateNewNode(const FName& NewNodeTypeName)
-{
-	const TSharedPtr<const SDataflowGraphEditor> PinnedDataflowGraphEditor = DataflowGraphEditor.Pin();
-	if (!PinnedDataflowGraphEditor)
-	{
-		return nullptr;
-	}
-
-	const TSharedPtr<FAssetSchemaAction_Dataflow_CreateNode_DataflowEdNode> NodeAction =
-		FAssetSchemaAction_Dataflow_CreateNode_DataflowEdNode::CreateAction(DataflowGraph, NewNodeTypeName);
-	UEdGraphNode* const NewEdNode = NodeAction->PerformAction(DataflowGraph, nullptr, PinnedDataflowGraphEditor->GetPasteLocation(), true);
-
-	return NewEdNode;
-}
-
-UEdGraphNode* UClothEditorContextObject::CreateAndConnectNewNode(
-	const FName& NewNodeTypeName,
-	UEdGraphNode& UpstreamNode,
-	const FName& ConnectionTypeName)
-{
-	// First find the specified output of the upstream node, plus any pins it's connected to
-
-	UEdGraphPin* UpstreamNodeOutputPin = nullptr;
-	TArray<UEdGraphPin*> ExistingNodeInputPins;
-
-	const UDataflowEdNode* const UpstreamDataflowEdNode = CastChecked<UDataflowEdNode>(&UpstreamNode);
-	const TSharedPtr<const FDataflowNode> UpstreamDataflowNode = UpstreamDataflowEdNode->GetDataflowNode();
-
-	for (const FDataflowOutput* const Output : UpstreamDataflowNode->GetOutputs())
-	{
-		if (Output->GetType() == ConnectionTypeName)
-		{
-			UpstreamNodeOutputPin = UpstreamDataflowEdNode->FindPin(*Output->GetName().ToString(), EGPD_Output);
-			ExistingNodeInputPins = UpstreamNodeOutputPin->LinkedTo;
-			break;
-		}
-	}
-
-	// Add the new node 
-
-	UEdGraphNode* const NewEdNode = CreateNewNode(NewNodeTypeName);
-	checkf(NewEdNode, TEXT("Failed to create a new node in the DataflowGraph"));
-
-	UDataflowEdNode* const NewDataflowEdNode = CastChecked<UDataflowEdNode>(NewEdNode);
-	const TSharedPtr<FDataflowNode> NewDataflowNode = NewDataflowEdNode->GetDataflowNode();
-
-	// Re-wire the graph
-
-	if (UpstreamNodeOutputPin)
-	{
-		UEdGraphPin* NewNodeInputPin = nullptr;
-		for (const FDataflowInput* const NewNodeInput : NewDataflowNode->GetInputs())
-		{
-			if (NewNodeInput->GetType() == ConnectionTypeName)
-			{
-				NewNodeInputPin = NewDataflowEdNode->FindPin(*NewNodeInput->GetName().ToString(), EGPD_Input);
-			}
-		}
-
-		UEdGraphPin* NewNodeOutputPin = nullptr;
-		for (const FDataflowOutput* const NewNodeOutput : NewDataflowNode->GetOutputs())
-		{
-			if (NewNodeOutput->GetType() == ConnectionTypeName)
-			{
-				NewNodeOutputPin = NewDataflowEdNode->FindPin(*NewNodeOutput->GetName().ToString(), EGPD_Output);
-				break;
-			}
-		}
-
-		check(NewNodeInputPin);
-		check(NewNodeOutputPin);
-
-		DataflowGraph->GetSchema()->TryCreateConnection(UpstreamNodeOutputPin, NewNodeInputPin);
-
-		for (UEdGraphPin* DownstreamInputPin : ExistingNodeInputPins)
-		{
-			DataflowGraph->GetSchema()->TryCreateConnection(NewNodeOutputPin, DownstreamInputPin);
-		}
-	}
-
-	DataflowGraph->NotifyGraphChanged();
-
-	return NewEdNode;
-}
-
 void UClothEditorContextObject::SetClothCollection(UE::Chaos::ClothAsset::EClothPatternVertexType ViewMode, TWeakPtr<FManagedArrayCollection> ClothCollection)
 {
 	ConstructionViewMode = ViewMode;
diff --git a/Engine/Plugins/Experimental/ChaosClothAssetEditor/Source/ChaosClothAssetEditorTools/Private/ChaosClothAsset/ClothTransferSkinWeightsTool.cpp b/Engine/Plugins/Experimental/ChaosClothAssetEditor/Source/ChaosClothAssetEditorTools/Private/ChaosClothAsset/ClothTransferSkinWeightsTool.cpp
--- a/Engine/Plugins/Experimental/ChaosClothAssetEditor/Source/ChaosClothAssetEditorTools/Private/ChaosClothAsset/ClothTransferSkinWeightsTool.cpp
+++ b/Engine/Plugins/Experimental/ChaosClothAssetEditor/Source/ChaosClothAssetEditorTools/Private/ChaosClothAsset/ClothTransferSkinWeightsTool.cpp
@@ -108,7 +108,16 @@ void UClothTransferSkinWeightsTool::Setup()
 {
 	USingleSelectionMeshEditingTool::Setup();
 
+	TransferSkinWeightsNode = ClothEditorContextObject->GetSingleSelectedNodeOfType<FChaosClothAssetTransferSkinWeightsNode>();
+	checkf(TransferSkinWeightsNode, TEXT("No Transfer Skin Weights Node is currently selected, or more than one node is selected"));
+
 	ToolProperties = NewObject<UClothTransferSkinWeightsToolProperties>(this);
+
+	SetSRTPropertiesFromTransform(TransferSkinWeightsNode->Transform);
+	ToolProperties->SourceMesh = TransferSkinWeightsNode->SkeletalMesh;
+	ToolProperties->SourceMeshLOD = TransferSkinWeightsNode->SkeletalMeshLOD;
+	
+
 	AddToolPropertySource(ToolProperties);
 
 
@@ -214,7 +223,7 @@ void UClothTransferSkinWeightsTool::Setup()
 	ensure(GizmoManager);
 	SourceMeshTransformProxy = NewObject<UTransformProxy>(this);
 	ensure(SourceMeshTransformProxy);
-	SourceMeshTransformProxy->SetTransform(FTransform::Identity);
+	SourceMeshTransformProxy->SetTransform(TransferSkinWeightsNode->Transform);
 
 	SourceMeshTransformProxy->OnTransformChanged.AddWeakLambda(this, [this](UTransformProxy*, FTransform NewTransform)
 	{
@@ -256,7 +265,9 @@ void UClothTransferSkinWeightsTool::Shutdown(EToolShutdownType ShutdownType)
 
 	if (ShutdownType == EToolShutdownType::Accept)
 	{
-		AddNewNode();
+		TransferSkinWeightsNode->SkeletalMesh = ToolProperties->SourceMesh;
+		TransferSkinWeightsNode->SkeletalMeshLOD = ToolProperties->SourceMeshLOD;
+		TransferSkinWeightsNode->Transform = TransformFromProperties();
 	}
 
 	if (SourceMeshTransformProxy)
@@ -293,7 +304,14 @@ void UClothTransferSkinWeightsTool::Shutdown(EToolShutdownType ShutdownType)
 
 bool UClothTransferSkinWeightsTool::CanAccept() const
 {
-	return ToolProperties && ToolProperties->SourceMesh && ToolProperties->SourceMesh->IsValidLODIndex(ToolProperties->SourceMeshLOD);
+	const FTransform& TransformOnNode = TransferSkinWeightsNode->Transform;;
+
+	return (ToolProperties->SourceMesh != TransferSkinWeightsNode->SkeletalMesh) ||
+		(ToolProperties->SourceMeshLOD != TransferSkinWeightsNode->SkeletalMeshLOD) ||
+		(ToolProperties->SourceMeshRotation != TransformOnNode.Rotator().Euler()) ||
+		(ToolProperties->SourceMeshTranslation != TransformOnNode.GetTranslation()) ||
+		(ToolProperties->SourceMeshScale != TransformOnNode.GetScale3D());
+
 }
 
 void UClothTransferSkinWeightsTool::OnTick(float DeltaTime)
@@ -324,26 +342,6 @@ void UClothTransferSkinWeightsTool::SetClothEditorContextObject(TObjectPtr<UClot
 	ClothEditorContextObject = InClothEditorContextObject;
 }
 
-void UClothTransferSkinWeightsTool::AddNewNode()
-{
-	checkf(ClothEditorContextObject, TEXT("Expected a non-null ClothEditorContextObject. Check UClothTransferSkinWeightsToolBuilder::CreateNewTool"));
-
-	const FName ConnectionType(FManagedArrayCollection::StaticType());
-	UEdGraphNode* const CurrentlySelectedNode = ClothEditorContextObject->GetSingleSelectedNodeWithOutputType(ConnectionType);
-	checkf(CurrentlySelectedNode, TEXT("No node with FManagedArrayCollection output is currently selected in the Dataflow graph"));
-
-	const FName NewNodeType(FChaosClothAssetTransferSkinWeightsNode::StaticType());
-	UEdGraphNode* const NewNode = ClothEditorContextObject->CreateAndConnectNewNode(NewNodeType, *CurrentlySelectedNode, ConnectionType);
-	checkf(NewNode, TEXT("Unexpectedly failed to create a new FChaosClothAssetTransferSkinWeightsNode"));
-
-	UDataflowEdNode* const NewDataflowEdNode = CastChecked<UDataflowEdNode>(NewNode);
-	const TSharedPtr<FDataflowNode> NewDataflowNode = NewDataflowEdNode->GetDataflowNode();
-	FChaosClothAssetTransferSkinWeightsNode* const NewTransferNode = NewDataflowNode->AsType<FChaosClothAssetTransferSkinWeightsNode>();
-
-	NewTransferNode->SkeletalMesh = ToolProperties->SourceMesh;
-	NewTransferNode->SkeletalMeshLOD = ToolProperties->SourceMeshLOD;
-	NewTransferNode->Transform = TransformFromProperties();
-}
 
 void UClothTransferSkinWeightsTool::SetPreviewMeshColorFunction()
 {
@@ -398,6 +396,14 @@ FTransform UClothTransferSkinWeightsTool::TransformFromProperties() const
 	return FTransform(Rotation, ToolProperties->SourceMeshTranslation, ToolProperties->SourceMeshScale);
 }
 
+void UClothTransferSkinWeightsTool::SetSRTPropertiesFromTransform(const FTransform& Transform) const
+{
+	ToolProperties->SourceMeshRotation = Transform.Rotator().Euler();
+	ToolProperties->SourceMeshTranslation = Transform.GetTranslation();
+	ToolProperties->SourceMeshScale = Transform.GetScale3D();
+}
+
+
 void UClothTransferSkinWeightsTool::UpdateSourceMesh()
 {
 	checkf(ToolProperties, TEXT("ToolProperties is expected to be non-null. Be sure to run Setup() on this tool when it is created."));
@@ -415,7 +421,8 @@ void UClothTransferSkinWeightsTool::UpdateSourceMesh()
 		SourceMeshComponent->SetMaterial(0, ToolSetupUtil::GetTransparentSculptMaterial(GetToolManager(), FLinearColor::Red, 0.4, true));
 		SourceMeshComponent->SetVisibility(!ToolProperties->bHideSourceMesh);
 
-		SourceMeshTransformGizmo->SetNewGizmoTransform(SourceMeshParentActor->GetActorTransform());
+		// Use ReinitializeGizmoTransform rather than SetNewGizmoTransform to avoid having this on the undo stack
+		SourceMeshTransformGizmo->ReinitializeGizmoTransform(SourceMeshParentActor->GetActorTransform());
 		SourceMeshTransformGizmo->SetVisibility(!ToolProperties->bHideSourceMesh);
 		SourceMeshTransformGizmo->ActiveGizmoMode = EToolContextTransformGizmoMode::Combined;
 
diff --git a/Engine/Plugins/Experimental/ChaosClothAssetEditor/Source/ChaosClothAssetEditorTools/Private/ChaosClothAsset/ClothWeightMapPaintTool.cpp b/Engine/Plugins/Experimental/ChaosClothAssetEditor/Source/ChaosClothAssetEditorTools/Private/ChaosClothAsset/ClothWeightMapPaintTool.cpp
--- a/Engine/Plugins/Experimental/ChaosClothAssetEditor/Source/ChaosClothAssetEditorTools/Private/ChaosClothAsset/ClothWeightMapPaintTool.cpp
+++ b/Engine/Plugins/Experimental/ChaosClothAssetEditor/Source/ChaosClothAssetEditorTools/Private/ChaosClothAsset/ClothWeightMapPaintTool.cpp
@@ -83,7 +83,7 @@ UMeshSurfacePointTool* UClothEditorWeightMapPaintToolBuilder::CreateNewTool(cons
 /*
  * Properties
  */
-void UClothEditorWeightMapPaintToolActionPropertySet::PostAction(EClothEditorWeightMapPaintToolActions Action)
+void UClothEditorMeshWeightMapPaintToolActions::PostAction(EClothEditorWeightMapPaintToolActions Action)
 {
 	if (ParentTool.IsValid())
 	{
@@ -103,6 +103,10 @@ void UClothEditorWeightMapPaintTool::Setup()
 {
 	UMeshSculptToolBase::Setup();
 
+	// Get the selected weight map node
+	WeightMapNodeToUpdate = ClothEditorContextObject->GetSingleSelectedNodeOfType<FChaosClothAssetAddWeightMapNode>();
+	checkf(WeightMapNodeToUpdate, TEXT("No Weight Map Node is currently selected, or more than one node is selected"));
+
 	SetToolDisplayName(LOCTEXT("ToolName", "Paint Weight Maps"));
 
 	// create dynamic mesh component to use for live preview
@@ -182,14 +186,16 @@ void UClothEditorWeightMapPaintTool::Setup()
 	
 	MeshSpatial = MakeUnique<UE::Geometry::FDynamicMeshAABBTree3>(DynamicMeshComponent->GetMesh(), bAutoBuild);
 	PolygonSelectionMechanic->Initialize(DynamicMeshComponent, GradientSelectionTopology.Get(), [this]() { return MeshSpatial.Get(); });
-	
-
-	AddWeightMapProperties = NewObject<UClothEditorAddWeightMapProperties>(this);
-	AddWeightMapProperties->Initialize(this);
-	AddToolPropertySource(AddWeightMapProperties);
 
 	UpdateWeightMapProperties = NewObject<UClothEditorUpdateWeightMapProperties>(this);
-	UpdateWeightMapProperties->Initialize(this);
+	UpdateWeightMapProperties->Name = WeightMapNodeToUpdate->Name;
+	UpdateWeightMapProperties->WatchProperty(UpdateWeightMapProperties->Name, [this, OriginalName = WeightMapNodeToUpdate->Name](const FString& NewName)
+	{
+		if (NewName != OriginalName)
+		{
+			bAnyChangeMade = true;
+		}
+	});
 	AddToolPropertySource(UpdateWeightMapProperties);
 
 	// initialize other properties
@@ -247,7 +253,7 @@ void UClothEditorWeightMapPaintTool::Setup()
 	UpdateBrushType(FilterProperties->PrimaryBrushType);
 	SetActiveSecondaryBrushType((int32)EClothEditorWeightMapPaintBrushType::Erase);
 	
-	ActionsProps = NewObject<UMeshWeightMapPaintToolActions>(this);
+	ActionsProps = NewObject<UClothEditorMeshWeightMapPaintToolActions>(this);
 	ActionsProps->Initialize(this);
 	AddToolPropertySource(ActionsProps);
 
@@ -317,46 +323,31 @@ void UClothEditorWeightMapPaintTool::Setup()
 	ActiveWeightMap = Mesh->Attributes()->GetWeightLayer(NumAttributeLayers);
 	ActiveWeightMap->SetName(FName("PaintLayer"));
 
-	// Check if there is a weight map node selected
-	FString CurrentMapName;
-	TArray<float> CurrentWeights;
-	const bool bWeightMapNodeIsSelected = GetSelectedNodeInfo(CurrentMapName, CurrentWeights);
-	const int32 ExpectedWeightsCount = bHaveDynamicMeshToWeightConversion ? WeightToDynamicMesh.Num() : Mesh->VertexCount();
-
-	bShouldUpdateExistingNode = bWeightMapNodeIsSelected && (ExpectedWeightsCount == CurrentWeights.Num());
-
-	if (bShouldUpdateExistingNode)
+	// Get weights from selected node
+	const TArray<float>& CurrentWeights = WeightMapNodeToUpdate->VertexWeights;
+	
+	if (bHaveDynamicMeshToWeightConversion)
 	{
-		if (bHaveDynamicMeshToWeightConversion)
-		{
-			check(WeightToDynamicMesh.Num() == CurrentWeights.Num());
-			for (int32 WeightID = 0; WeightID < CurrentWeights.Num(); ++WeightID)
-			{
-				for (const int32 VertexID : WeightToDynamicMesh[WeightID])
-				{
-					ActiveWeightMap->SetValue(VertexID, &CurrentWeights[WeightID]);
-				}
-			}
-		}
-		else
+		for (int32 WeightID = 0; WeightID < CurrentWeights.Num(); ++WeightID)
 		{
-			check(Mesh->VertexCount() == CurrentWeights.Num());
-			for (int32 VertexID = 0; VertexID < Mesh->VertexCount(); ++VertexID)
+			for (const int32 VertexID : WeightToDynamicMesh[WeightID])
 			{
-				ActiveWeightMap->SetValue(VertexID, &CurrentWeights[VertexID]);
+				ActiveWeightMap->SetValue(VertexID, &CurrentWeights[WeightID]);
 			}
 		}
-
-		UpdateWeightMapProperties->Name = CurrentMapName;
-		SetToolPropertySourceEnabled(AddWeightMapProperties, false);
-		SetToolPropertySourceEnabled(UpdateWeightMapProperties, true);
 	}
 	else
 	{
-		SetToolPropertySourceEnabled(AddWeightMapProperties, true);
-		SetToolPropertySourceEnabled(UpdateWeightMapProperties, false);
+		for (int32 VertexID = 0; VertexID < CurrentWeights.Num(); ++VertexID)
+		{
+			ActiveWeightMap->SetValue(VertexID, &CurrentWeights[VertexID]);
+		}
 	}
 
+	UpdateWeightMapProperties->Name = WeightMapNodeToUpdate->Name;
+	SetToolPropertySourceEnabled(UpdateWeightMapProperties, true);
+
+
 	// update colors
 	DynamicMeshComponent->FastNotifyVertexAttributesUpdated(EMeshRenderAttributeFlags::VertexColors);
 	GetToolManager()->PostInvalidation();
@@ -403,14 +394,7 @@ void UClothEditorWeightMapPaintTool::CommitResult(UBaseDynamicMeshComponent* Com
 {
 	GetToolManager()->BeginUndoTransaction(LOCTEXT("WeightPaintToolTransactionName", "Paint Weights"));
 
-	if (bShouldUpdateExistingNode)
-	{
-		UpdateSelectedNode();
-	}
-	else
-	{
-		AddNewNode(AddWeightMapProperties->Name);
-	}
+	UpdateSelectedNode();
 
 	GetToolManager()->EndUndoTransaction();
 }
@@ -1497,14 +1481,7 @@ void UClothEditorWeightMapPaintTool::OnTick(float DeltaTime)
 
 bool UClothEditorWeightMapPaintTool::CanAccept() const
 {
-	if (bShouldUpdateExistingNode)
-	{
-		return true;
-	}
-	else
-	{
-		return !(AddWeightMapProperties->Name.IsEmpty());
-	}
+	return bAnyChangeMade;
 }
 
 
@@ -1615,90 +1592,27 @@ void UClothEditorWeightMapPaintTool::ClearAllWeightsAction()
 }
 
 
-bool UClothEditorWeightMapPaintTool::GetSelectedNodeInfo(FString& OutMapName, TArray<float>& OutWeights)
-{
-	const TSharedPtr<const SDataflowGraphEditor> GraphEditor = ClothEditorContextObject->GetDataflowGraphEditor().Pin();
-	if (const UEdGraphNode* const SingleSelectedNode = GraphEditor->GetSingleSelectedNode())
-	{
-		const UDataflowEdNode* const SelectedDataflowEdNode = Cast<UDataflowEdNode>(SingleSelectedNode);
-		const TSharedPtr<const FDataflowNode> DataflowNode = SelectedDataflowEdNode->GetDataflowNode();
-		if (const FChaosClothAssetAddWeightMapNode* const WeightMapNode = DataflowNode->AsType<FChaosClothAssetAddWeightMapNode>())
-		{
-			OutMapName = WeightMapNode->Name;
-			OutWeights = WeightMapNode->VertexWeights;
-			return true;
-		}
-	}
-
-	// No single node selected, or the selected node is not FChaosClothAssetAddWeightMapNode
-	return false;
-}
-
-
-void UClothEditorWeightMapPaintTool::AddNewNode(const FString& NewMapName)
-{
-	check(ActiveWeightMap);
-	TArray<float> CurrentWeights;
-	GetCurrentWeightMap(CurrentWeights);
-
-	// TODO: If we are operating on a single pattern, expand the weight map to cover the whole LOD
-
-	const FName ConnectionType(FManagedArrayCollection::StaticType());
-	UEdGraphNode* const CurrentlySelectedNode = ClothEditorContextObject->GetSingleSelectedNodeWithOutputType(ConnectionType);
-	checkf(CurrentlySelectedNode, TEXT("No node with FManagedArrayCollection output is currently selected in the Dataflow graph"));
-
-	const FName NewNodeType(FChaosClothAssetAddWeightMapNode::StaticType());
-	UEdGraphNode* const NewNode = ClothEditorContextObject->CreateAndConnectNewNode(NewNodeType, *CurrentlySelectedNode, ConnectionType);
-	checkf(NewNode, TEXT("Unexpectedly failed to create a new FChaosClothAssetAddWeightMapNode"));
-
-	UDataflowEdNode* const NewDataflowEdNode = CastChecked<UDataflowEdNode>(NewNode);
-	const TSharedPtr<FDataflowNode> NewDataflowNode = NewDataflowEdNode->GetDataflowNode();
-	FChaosClothAssetAddWeightMapNode* const NewWeightMapNode = NewDataflowNode->AsType<FChaosClothAssetAddWeightMapNode>();
-
-	if (bHaveDynamicMeshToWeightConversion)
-	{
-		NewWeightMapNode->VertexWeights.Init(0.f, WeightToDynamicMesh.Num());
-		for (int32 DynamicMeshIdx = 0; DynamicMeshIdx < CurrentWeights.Num(); ++DynamicMeshIdx)
-		{
-			NewWeightMapNode->VertexWeights[DynamicMeshToWeight[DynamicMeshIdx]] = CurrentWeights[DynamicMeshIdx];
-		}
-	}
-	else
-	{
-		NewWeightMapNode->VertexWeights = CurrentWeights;
-	}
-	NewWeightMapNode->Name = AddWeightMapProperties->Name;
-	// NewWeightMapNode->SetName(FName(AddWeightMapProperties->NewMapName));   TODO: This doesn't seem to work to set the Node name
-
-}
-
-
 void UClothEditorWeightMapPaintTool::UpdateSelectedNode()
 {
 	check(ActiveWeightMap);
 	TArray<float> CurrentWeights;
 	GetCurrentWeightMap(CurrentWeights);
 
-	const TSharedPtr<SGraphEditor> GraphEditor = ClothEditorContextObject->GetDataflowGraphEditor().Pin();
-	UEdGraphNode* const SingleSelectedNode = GraphEditor->GetSingleSelectedNode();
-	checkf(SingleSelectedNode, TEXT("There should be a single selected node in the Dataflow graph"));
+	check(WeightMapNodeToUpdate);
 
-	UDataflowEdNode* const SelectedDataflowEdNode = CastChecked<UDataflowEdNode>(SingleSelectedNode);
-	const TSharedPtr<FDataflowNode> DataflowNode = SelectedDataflowEdNode->GetDataflowNode();
-	FChaosClothAssetAddWeightMapNode* const WeightMapNode = DataflowNode->AsType<FChaosClothAssetAddWeightMapNode>();
 	if (bHaveDynamicMeshToWeightConversion)
 	{
-		WeightMapNode->VertexWeights.Init(0.f, WeightToDynamicMesh.Num());
+		WeightMapNodeToUpdate->VertexWeights.Init(0.f, WeightToDynamicMesh.Num());
 		for (int32 DynamicMeshIdx = 0; DynamicMeshIdx < CurrentWeights.Num(); ++DynamicMeshIdx)
 		{
-			WeightMapNode->VertexWeights[DynamicMeshToWeight[DynamicMeshIdx]] = CurrentWeights[DynamicMeshIdx];
+			WeightMapNodeToUpdate->VertexWeights[DynamicMeshToWeight[DynamicMeshIdx]] = CurrentWeights[DynamicMeshIdx];
 		}
 	}
 	else
 	{
-		WeightMapNode->VertexWeights = CurrentWeights;
+		WeightMapNodeToUpdate->VertexWeights = CurrentWeights;
 	}
-	WeightMapNode->Invalidate();
+	WeightMapNodeToUpdate->Invalidate();
 }
 
 
@@ -1764,7 +1678,10 @@ void UClothEditorWeightMapPaintTool::EndChange()
 {
 	check(ActiveWeightEditChangeTracker);
 
-	TUniquePtr < FDynamicMeshChange > EditResult = ActiveWeightEditChangeTracker->EndChange();
+	bAnyChangeMade = true;
+
+	TUniquePtr<FDynamicMeshChange> EditResult = ActiveWeightEditChangeTracker->EndChange();
+
 	TUniquePtr<ClothWeightPaintLocals::FClothWeightPaintMeshChange> ClothWeightPaintMeshChange =
 		MakeUnique<ClothWeightPaintLocals::FClothWeightPaintMeshChange>(DynamicMeshComponent.Get(), MoveTemp(EditResult));
 	ActiveWeightEditChangeTracker = nullptr;
diff --git a/Engine/Plugins/Experimental/ChaosClothAssetEditor/Source/ChaosClothAssetEditorTools/Public/ChaosClothAsset/ClothEditorContextObject.h b/Engine/Plugins/Experimental/ChaosClothAssetEditor/Source/ChaosClothAssetEditorTools/Public/ChaosClothAsset/ClothEditorContextObject.h
--- a/Engine/Plugins/Experimental/ChaosClothAssetEditor/Source/ChaosClothAssetEditorTools/Public/ChaosClothAsset/ClothEditorContextObject.h
+++ b/Engine/Plugins/Experimental/ChaosClothAssetEditor/Source/ChaosClothAssetEditorTools/Public/ChaosClothAsset/ClothEditorContextObject.h
@@ -3,9 +3,10 @@
 #pragma once
 
 #include "UObject/Object.h"
+#include "Dataflow/DataflowEdNode.h"
+#include "Dataflow/DataflowGraphEditor.h"
 #include "ClothEditorContextObject.generated.h"
 
-class SDataflowGraphEditor;
 class UDataflow;
 class UEdGraphNode;
 
@@ -23,52 +24,28 @@ class CHAOSCLOTHASSETEDITORTOOLS_API UClothEditorContextObject : public UObject
 
 public:
 
-	void Init(TWeakPtr<SDataflowGraphEditor> DataflowGraphEditor, TObjectPtr<UDataflow> DataflowGraph, UE::Chaos::ClothAsset::EClothPatternVertexType InConstructionViewMode, TWeakPtr<FManagedArrayCollection> SelectedClothCollection);
+	void Init(TWeakPtr<SDataflowGraphEditor> DataflowGraphEditor, UE::Chaos::ClothAsset::EClothPatternVertexType InConstructionViewMode, TWeakPtr<FManagedArrayCollection> SelectedClothCollection);
 
-	TWeakPtr<SDataflowGraphEditor> GetDataflowGraphEditor();
-	const TWeakPtr<const SDataflowGraphEditor> GetDataflowGraphEditor() const;
-
-	TObjectPtr<UDataflow> GetDataflowGraph();
-	const TObjectPtr<const UDataflow> GetDataflowGraph() const;
-
-	/** 
-	 * Return the single selected node in the Dataflow Graph Editor, or nullptr if multiple or no nodes are selected 
-	 */
-	UEdGraphNode* GetSingleSelectedNode() const;
-
-	/** 
-	* Return the single selected node in the Dataflow Graph Editor only if it has an output of the specified type
-	* If there is not a single node selected, or if it does not have the specified output, return null 
+	/**
+	* Get a single selected node of the specified type. Return nullptr if the specified node is not selected, or if multiple nodes are selected
 	*/
-	UEdGraphNode* GetSingleSelectedNodeWithOutputType(const FName& SelectedNodeOutputTypeName) const;
-
-	/** 
-	 * Create a node with the specified type in the graph 
-	 */
-	UEdGraphNode* CreateNewNode(const FName& NewNodeTypeName);
-
-
-	/** Create a node with the specified type, then connect it to the output of the specified UpstreamNode.
-	* If the specified output of the upstream node is already connected to another node downstream, we first break
-	* that connecttion, then insert the new node along the previous connection.
-	* We want to turn this:
-	* 
-	* [UpstreamNode] ----> [DownstreamNode(s)]
-	* 
-	* to this:
-	*
-	* [UpstreamNode] ----> [NewNode] ----> [DownstreamNode(s)]
-	*
-	* 
-	* @param NewNodeTypeName The type of node to create, by name
-	* @param UpstreamNode Node to connect the new node to
-	* @param ConnectionTypeName The type of output of the upstream node to connect our new node to
-	* @return The newly-created node
-	*/
-	UEdGraphNode* CreateAndConnectNewNode(
-		const FName& NewNodeTypeName,
-		UEdGraphNode& UpstreamNode,
-		const FName& ConnectionTypeName);
+	template<typename NodeType>
+	NodeType* GetSingleSelectedNodeOfType() const
+	{
+		const TSharedPtr<SDataflowGraphEditor> GraphEditor = DataflowGraphEditor.Pin();
+
+		if (UEdGraphNode* const SingleSelectedNode = GraphEditor->GetSingleSelectedNode())
+		{
+			UDataflowEdNode* const SelectedDataflowEdNode = CastChecked<UDataflowEdNode>(SingleSelectedNode);
+			const TSharedPtr<FDataflowNode> DataflowNode = SelectedDataflowEdNode->GetDataflowNode();
+			if (NodeType* const NodeTypeNode = DataflowNode->AsType<NodeType>())
+			{
+				return NodeTypeNode;
+			}
+		}
+
+		return nullptr;
+	}
 
 	void SetClothCollection(UE::Chaos::ClothAsset::EClothPatternVertexType ViewMode, TWeakPtr<FManagedArrayCollection> ClothCollection);
 
@@ -78,9 +55,6 @@ private:
 
 	TWeakPtr<SDataflowGraphEditor> DataflowGraphEditor;
 
-	UPROPERTY()
-	TObjectPtr<UDataflow> DataflowGraph = nullptr;
-
 	UE::Chaos::ClothAsset::EClothPatternVertexType ConstructionViewMode;
 	TWeakPtr<const FManagedArrayCollection> SelectedClothCollection;
 };
diff --git a/Engine/Plugins/Experimental/ChaosClothAssetEditor/Source/ChaosClothAssetEditorTools/Public/ChaosClothAsset/ClothTransferSkinWeightsTool.h b/Engine/Plugins/Experimental/ChaosClothAssetEditor/Source/ChaosClothAssetEditorTools/Public/ChaosClothAsset/ClothTransferSkinWeightsTool.h
--- a/Engine/Plugins/Experimental/ChaosClothAssetEditor/Source/ChaosClothAssetEditorTools/Public/ChaosClothAsset/ClothTransferSkinWeightsTool.h
+++ b/Engine/Plugins/Experimental/ChaosClothAssetEditor/Source/ChaosClothAssetEditorTools/Public/ChaosClothAsset/ClothTransferSkinWeightsTool.h
@@ -19,6 +19,7 @@ class UMeshOpPreviewWithBackgroundCompute;
 class AInternalToolFrameworkActor;
 class UDynamicMeshComponent;
 class FTransformGizmoDataBinder;
+struct FChaosClothAssetTransferSkinWeightsNode;
 
 UCLASS()
 class CHAOSCLOTHASSETEDITORTOOLS_API UClothTransferSkinWeightsToolProperties : public UInteractiveToolPropertySet
@@ -94,11 +95,10 @@ private:
 	
 	void SetClothEditorContextObject(TObjectPtr<UClothEditorContextObject> InClothEditorContextObject);
 
-	void AddNewNode();
-
 	void SetPreviewMeshColorFunction();
 
 	FTransform TransformFromProperties() const;
+	void SetSRTPropertiesFromTransform(const FTransform& Transform) const;
 
 	void UpdateSourceMesh();
 
@@ -134,6 +134,8 @@ private:
 	// Used to lookup the index of the currently selected-by-name bone
 	TMap<FName, FBoneIndexType> TargetMeshBoneNameToIndex;
 
+	FChaosClothAssetTransferSkinWeightsNode* TransferSkinWeightsNode = nullptr;
+
 	bool bHasInvalidLODWarning = false;
 	bool bHasOpFailedWarning = false;
 };
diff --git a/Engine/Plugins/Experimental/ChaosClothAssetEditor/Source/ChaosClothAssetEditorTools/Public/ChaosClothAsset/ClothWeightMapPaintTool.h b/Engine/Plugins/Experimental/ChaosClothAssetEditor/Source/ChaosClothAssetEditorTools/Public/ChaosClothAsset/ClothWeightMapPaintTool.h
--- a/Engine/Plugins/Experimental/ChaosClothAssetEditor/Source/ChaosClothAssetEditorTools/Public/ChaosClothAsset/ClothWeightMapPaintTool.h
+++ b/Engine/Plugins/Experimental/ChaosClothAssetEditor/Source/ChaosClothAssetEditorTools/Public/ChaosClothAsset/ClothWeightMapPaintTool.h
@@ -31,6 +31,7 @@ class UWeightMapPaintBrushOpProps;
 class UWeightMapSmoothBrushOpProps;
 class UClothEditorContextObject;
 class UPolygonSelectionMechanic;
+struct FChaosClothAssetAddWeightMapNode;
 
 DECLARE_STATS_GROUP(TEXT("WeightMapPaintTool"), STATGROUP_WeightMapPaintTool, STATCAT_Advanced);
 DECLARE_CYCLE_STAT(TEXT("WeightMapPaintTool_UpdateROI"), WeightMapPaintTool_UpdateROI, STATGROUP_WeightMapPaintTool);
@@ -186,24 +187,17 @@ enum class EClothEditorWeightMapPaintToolActions
 
 
 UCLASS()
-class CHAOSCLOTHASSETEDITORTOOLS_API UClothEditorWeightMapPaintToolActionPropertySet : public UInteractiveToolPropertySet
+class CHAOSCLOTHASSETEDITORTOOLS_API UClothEditorMeshWeightMapPaintToolActions : public UInteractiveToolPropertySet
 {
 	GENERATED_BODY()
 
 public:
+
 	TWeakObjectPtr<UClothEditorWeightMapPaintTool> ParentTool;
 
 	void Initialize(UClothEditorWeightMapPaintTool* ParentToolIn) { ParentTool = ParentToolIn; }
 
 	void PostAction(EClothEditorWeightMapPaintToolActions Action);
-};
-
-UCLASS()
-class CHAOSCLOTHASSETEDITORTOOLS_API UMeshWeightMapPaintToolActions : public UClothEditorWeightMapPaintToolActionPropertySet
-{
-	GENERATED_BODY()
-
-public:
 
 	UFUNCTION(CallInEditor, Category = Operations, meta = (DisplayPriority = 10))
 	void ClearAll()
@@ -219,27 +213,14 @@ public:
 
 };
 
-
 UCLASS()
-class CHAOSCLOTHASSETEDITORTOOLS_API UClothEditorAddWeightMapProperties : public UClothEditorWeightMapPaintToolActionPropertySet
+class CHAOSCLOTHASSETEDITORTOOLS_API UClothEditorUpdateWeightMapProperties : public UInteractiveToolPropertySet
 {
 	GENERATED_BODY()
 
 public:
 
-	UPROPERTY(EditAnywhere, Category = AddNode, meta = (DisplayName = "Name"))
-	FString Name = "WeightMap";
-};
-
-
-UCLASS()
-class CHAOSCLOTHASSETEDITORTOOLS_API UClothEditorUpdateWeightMapProperties : public UClothEditorWeightMapPaintToolActionPropertySet
-{
-	GENERATED_BODY()
-
-public:
-
-	UPROPERTY(VisibleAnywhere, Category = UpdateNode, meta = (DisplayName = "Name"))
+	UPROPERTY(EditAnywhere, Category = UpdateNode, meta = (DisplayName = "Name"))
 	FString Name;
 };
 
@@ -337,10 +318,7 @@ public:
 	virtual void RequestAction(EClothEditorWeightMapPaintToolActions ActionType);
 	
 	UPROPERTY()
-	TObjectPtr<UMeshWeightMapPaintToolActions> ActionsProps;
-
-	UPROPERTY()
-	TObjectPtr<UClothEditorAddWeightMapProperties> AddWeightMapProperties;
+	TObjectPtr<UClothEditorMeshWeightMapPaintToolActions> ActionsProps;
 
 	UPROPERTY()
 	TObjectPtr<UClothEditorUpdateWeightMapProperties> UpdateWeightMapProperties;
@@ -471,11 +449,13 @@ protected:
 
 	friend class UClothEditorWeightMapPaintToolBuilder;
 
+	bool bAnyChangeMade = false;
+
 	// Node graph editor support
-	bool GetSelectedNodeInfo(FString& OutMapName, TArray<float>& OutWeights);
-	void AddNewNode(const FString& NewMapName);
+
+	FChaosClothAssetAddWeightMapNode* WeightMapNodeToUpdate = nullptr;
+
 	void UpdateSelectedNode();
-	bool bShouldUpdateExistingNode = true;
 };
 
 
diff --git a/Engine/Plugins/Experimental/ChaosClothAssetEditor/Source/ChaosClothAssetTools/Private/ChaosClothAsset/ClothPatternToDynamicMesh.cpp b/Engine/Plugins/Experimental/ChaosClothAssetEditor/Source/ChaosClothAssetTools/Private/ChaosClothAsset/ClothPatternToDynamicMesh.cpp
--- a/Engine/Plugins/Experimental/ChaosClothAssetEditor/Source/ChaosClothAssetTools/Private/ChaosClothAsset/ClothPatternToDynamicMesh.cpp
+++ b/Engine/Plugins/Experimental/ChaosClothAssetEditor/Source/ChaosClothAssetTools/Private/ChaosClothAsset/ClothPatternToDynamicMesh.cpp
@@ -173,7 +173,7 @@ public:
 
 	int32 NumWeightMapLayers() const
 	{
-		return WeightMapNames.Num();
+		return (VertexDataType == EClothPatternVertexType::Render) ? 0 : WeightMapNames.Num(); // No weight maps for render mesh
 	}
 
 	FName GetWeightMapName(int32 LayerIndex) const
