Chaos Cloth - Cache Datasmith import.

#rnx

diff --git a/Engine/Plugins/Experimental/ChaosClothAssetEditor/Source/ChaosClothAssetDataflowNodes/Private/ChaosClothAsset/DatasmithImportNode.cpp b/Engine/Plugins/Experimental/ChaosClothAssetEditor/Source/ChaosClothAssetDataflowNodes/Private/ChaosClothAsset/DatasmithImportNode.cpp
--- a/Engine/Plugins/Experimental/ChaosClothAssetEditor/Source/ChaosClothAssetDataflowNodes/Private/ChaosClothAsset/DatasmithImportNode.cpp
+++ b/Engine/Plugins/Experimental/ChaosClothAssetEditor/Source/ChaosClothAssetDataflowNodes/Private/ChaosClothAsset/DatasmithImportNode.cpp
@@ -17,17 +17,21 @@
 FChaosClothAssetDatasmithImportNode::FChaosClothAssetDatasmithImportNode(const Dataflow::FNodeParameters& InParam, FGuid InGuid)
 	: FDataflowNode(InParam, InGuid)
 {
-	RegisterInputConnection(&DatasmithFile);
 	RegisterOutputConnection(&Collection);
 }
 
+void FChaosClothAssetDatasmithImportNode::Serialize(FArchive& Archive)
+{
+	::Chaos::FChaosArchive ChaosArchive(Archive);
+	ImportCache.Serialize(ChaosArchive);
+	Archive << ImportHash;
+}
+
 bool FChaosClothAssetDatasmithImportNode::EvaluateImpl(Dataflow::FContext& Context, FManagedArrayCollection& OutCollection) const
 {
 	using namespace UE::DatasmithImporter;
 
-	const FFilePath& InFilePath = GetValue<FFilePath>(Context, &DatasmithFile);
-
-	const FSourceUri SourceUri = FSourceUri::FromFilePath(InFilePath.FilePath);
+	const FSourceUri SourceUri = FSourceUri::FromFilePath(ImportFile.FilePath);
 	const TSharedPtr<FExternalSource> ExternalSource = IExternalSourceModule::GetOrCreateExternalSource(SourceUri);
 	if (!ExternalSource)
 	{
@@ -92,23 +96,34 @@ bool FChaosClothAssetDatasmithImportNode::EvaluateImpl(Dataflow::FContext& Conte
 
 void FChaosClothAssetDatasmithImportNode::Evaluate(Dataflow::FContext& Context, const FDataflowOutput* Out) const
 {
-	FManagedArrayCollection OutCollection;
-	const bool bSuccess = EvaluateImpl(Context, OutCollection);
-
-	if (bSuccess)
-	{
-		SetValue<FManagedArrayCollection>(Context, OutCollection, &Collection);
-	}
-	else
+	if (Out->IsA<FManagedArrayCollection>(&Collection))
 	{
-		using namespace UE::Chaos::ClothAsset;
+		if (FPaths::FileExists(ImportFile.FilePath))
+		{
+			const FMD5Hash FileHash = FMD5Hash::HashFile(*ImportFile.FilePath);
+			if (FileHash != ImportHash)
+			{
+				ImportHash = FileHash;
+
+				// The file has changed, reimport
+				const bool bSuccess = EvaluateImpl(Context, ImportCache);
+
+				if (!bSuccess)
+				{
+					using namespace UE::Chaos::ClothAsset;
 
-		// Init with an empty cloth collection
-		const TSharedRef<FManagedArrayCollection> ClothCollection = MakeShared<FManagedArrayCollection>();
-		FCollectionClothFacade CollectionClothFacade(ClothCollection);
-		CollectionClothFacade.DefineSchema();
+					// The import has failed, reset the cache to an empty cloth collection
+					ImportCache.Reset();
 
-		SetValue<FManagedArrayCollection>(Context, *ClothCollection, &Collection);
+					const TSharedRef<FManagedArrayCollection> ClothCollection = MakeShared<FManagedArrayCollection>(MoveTemp(ImportCache));
+					FCollectionClothFacade CollectionClothFacade(ClothCollection);
+					CollectionClothFacade.DefineSchema();
+
+					ImportCache = MoveTemp(*ClothCollection);
+				}
+			}
+		}
+		SetValue<FManagedArrayCollection>(Context, ImportCache, &Collection);
 	}
 }
 
diff --git a/Engine/Plugins/Experimental/ChaosClothAssetEditor/Source/ChaosClothAssetDataflowNodes/Public/ChaosClothAsset/DatasmithImportNode.h b/Engine/Plugins/Experimental/ChaosClothAssetEditor/Source/ChaosClothAssetDataflowNodes/Public/ChaosClothAsset/DatasmithImportNode.h
--- a/Engine/Plugins/Experimental/ChaosClothAssetEditor/Source/ChaosClothAssetDataflowNodes/Public/ChaosClothAsset/DatasmithImportNode.h
+++ b/Engine/Plugins/Experimental/ChaosClothAssetEditor/Source/ChaosClothAssetDataflowNodes/Public/ChaosClothAsset/DatasmithImportNode.h
@@ -4,6 +4,7 @@
 
 #include "Dataflow/DataflowNode.h"
 #include "GeometryCollection/ManagedArrayCollection.h"
+#include "Misc/SecureHash.h"
 #include "DatasmithImportNode.generated.h"
 
 USTRUCT(meta = (DataflowCloth))
@@ -11,22 +12,25 @@ struct FChaosClothAssetDatasmithImportNode : public FDataflowNode
 {
 	GENERATED_USTRUCT_BODY()
 	DATAFLOW_NODE_DEFINE_INTERNAL(FChaosClothAssetDatasmithImportNode, "DatasmithImport", "Cloth", "Cloth Datasmith Import")
-	//DATAFLOW_NODE_RENDER_TYPE(FGeometryCollection::StaticType(), "Collection")
 
 public:
-	
-	/** Datasmith file to read from. */
-	UPROPERTY(EditAnywhere, Category = "Datasmith Importer")
-	FFilePath DatasmithFile;
-
 	UPROPERTY(meta = (DataflowOutput, DisplayName = "Collection"))
 	FManagedArrayCollection Collection;
 
-	FChaosClothAssetDatasmithImportNode(const Dataflow::FNodeParameters& InParam, FGuid InGuid = FGuid::NewGuid());
+	/** Path of the file to import using any available Datasmith cloth translator. */
+	UPROPERTY(EditAnywhere, Category = "Datasmith Import")
+	FFilePath ImportFile;
 
-	virtual void Evaluate(Dataflow::FContext& Context, const FDataflowOutput* Out) const override;
+	FChaosClothAssetDatasmithImportNode(const Dataflow::FNodeParameters& InParam, FGuid InGuid = FGuid::NewGuid());
 
 private:
+	//~ Begin FDataflowNode interface
+	virtual void Evaluate(Dataflow::FContext& Context, const FDataflowOutput* Out) const override;
+	virtual void Serialize(FArchive& Archive) override;
+	//~ End FDataflowNode interface
 
 	bool EvaluateImpl(Dataflow::FContext& Context, FManagedArrayCollection& OutCollection) const;
+
+	mutable FManagedArrayCollection ImportCache;
+	mutable FMD5Hash ImportHash;
 };
