Chaos Cloth - Fixed crash on AddWeighMap node after change in schema.

#rnx

diff --git a/Engine/Plugins/Experimental/ChaosClothAssetEditor/Source/ChaosClothAssetDataflowNodes/Private/ChaosClothAsset/AddWeightMapNode.cpp b/Engine/Plugins/Experimental/ChaosClothAssetEditor/Source/ChaosClothAssetDataflowNodes/Private/ChaosClothAsset/AddWeightMapNode.cpp
--- a/Engine/Plugins/Experimental/ChaosClothAssetEditor/Source/ChaosClothAssetDataflowNodes/Private/ChaosClothAsset/AddWeightMapNode.cpp
+++ b/Engine/Plugins/Experimental/ChaosClothAssetEditor/Source/ChaosClothAssetDataflowNodes/Private/ChaosClothAsset/AddWeightMapNode.cpp
@@ -25,29 +25,31 @@ void FChaosClothAssetAddWeightMapNode::Evaluate(Dataflow::FContext& Context, con
 	if (Out->IsA<FManagedArrayCollection>(&Collection))
 	{
 		// Evaluate in collection
+		const FString InNameString = GetValue<FString>(Context, &Name);
+
 		FManagedArrayCollection InCollection = GetValue<FManagedArrayCollection>(Context, &Collection);
 		const TSharedRef<FManagedArrayCollection> ClothCollection = MakeShared<FManagedArrayCollection>(MoveTemp(InCollection));
 		FCollectionClothFacade ClothFacade(ClothCollection);
-
-		const FString InNameString = GetValue<FString>(Context, &Name);
-		const FName InName(InNameString);
-		ClothFacade.AddWeightMap(InName);		// Does nothing if weight map already exists
-
-		TArrayView<float> ClothWeights = ClothFacade.GetWeightMap(InName);
-		const int32 MaxWeightIndex = FMath::Min(VertexWeights.Num(), ClothWeights.Num());
-		if (VertexWeights.Num() > 0 && VertexWeights.Num() != ClothWeights.Num())
+		if (ClothFacade.IsValid())  // Can only act on the collection if it is a valid cloth collection
 		{
-			DataflowNodes::LogAndToastWarning(
-				FText::Format(LOCTEXT("WeightMapSize", "FChaosClothAssetAddWeightMapNode: Vertex count mismatch: vertex weights in the node: {0}; 3D vertices in cloth: {1}"),
-					VertexWeights.Num(),
-					ClothWeights.Num()));
+			const FName InName(InNameString);
+			ClothFacade.AddWeightMap(InName);		// Does nothing if weight map already exists
+
+			TArrayView<float> ClothWeights = ClothFacade.GetWeightMap(InName);
+			const int32 MaxWeightIndex = FMath::Min(VertexWeights.Num(), ClothWeights.Num());
+			if (VertexWeights.Num() > 0 && VertexWeights.Num() != ClothWeights.Num())
+			{
+				DataflowNodes::LogAndToastWarning(
+					FText::Format(LOCTEXT("WeightMapSize", "FChaosClothAssetAddWeightMapNode: Vertex count mismatch: vertex weights in the node: {0}; 3D vertices in cloth: {1}"),
+						VertexWeights.Num(),
+						ClothWeights.Num()));
+			}
+
+			for (int32 VertexID = 0; VertexID < MaxWeightIndex; ++VertexID)
+			{
+				ClothWeights[VertexID] = VertexWeights[VertexID];
+			}
 		}
-
-		for (int32 VertexID = 0; VertexID < MaxWeightIndex; ++VertexID)
-		{
-			ClothWeights[VertexID] = VertexWeights[VertexID];
-		}
-
 		SetValue<FManagedArrayCollection>(Context, *ClothCollection, &Collection);
 		SetValue<FString>(Context, InNameString, &Name);
 	}
