Cloth Editor weight map paint tool: fix weight map name not being saved on the node by the Tool

diff --git a/Engine/Plugins/Experimental/ChaosClothAssetEditor/Source/ChaosClothAssetEditorTools/Private/ChaosClothAsset/ClothWeightMapPaintTool.cpp b/Engine/Plugins/Experimental/ChaosClothAssetEditor/Source/ChaosClothAssetEditorTools/Private/ChaosClothAsset/ClothWeightMapPaintTool.cpp
--- a/Engine/Plugins/Experimental/ChaosClothAssetEditor/Source/ChaosClothAssetEditorTools/Private/ChaosClothAsset/ClothWeightMapPaintTool.cpp
+++ b/Engine/Plugins/Experimental/ChaosClothAssetEditor/Source/ChaosClothAssetEditorTools/Private/ChaosClothAsset/ClothWeightMapPaintTool.cpp
@@ -1612,6 +1612,9 @@ void UClothEditorWeightMapPaintTool::UpdateSelectedNode()
 	{
 		WeightMapNodeToUpdate->VertexWeights = CurrentWeights;
 	}
+	
+	WeightMapNodeToUpdate->Name = UpdateWeightMapProperties->Name;
+
 	WeightMapNodeToUpdate->Invalidate();
 }
 
