Chaos Cloth Editor - Prevent a crash when selecting an unconnected or a non cloth collection type node that appeared after the schema rework.

#rnx

diff --git a/Engine/Plugins/Experimental/ChaosClothAssetEditor/Source/ChaosClothAssetEditor/Private/ChaosClothAsset/ClothEditorToolkit.cpp b/Engine/Plugins/Experimental/ChaosClothAssetEditor/Source/ChaosClothAssetEditor/Private/ChaosClothAsset/ClothEditorToolkit.cpp
--- a/Engine/Plugins/Experimental/ChaosClothAssetEditor/Source/ChaosClothAssetEditor/Private/ChaosClothAsset/ClothEditorToolkit.cpp
+++ b/Engine/Plugins/Experimental/ChaosClothAssetEditor/Source/ChaosClothAssetEditor/Private/ChaosClothAsset/ClothEditorToolkit.cpp
@@ -973,7 +973,8 @@ void FChaosClothAssetEditorToolkit::OnNodeSelectionChanged(const TSet<UObject*>&
 						return Collection;
 					}
 
-					return MakeShared<FManagedArrayCollection>();
+					// The cloth collection schema must be applied to prevent the dynamic mesh conversion and tools from crashing trying to access invalid facades
+					break;
 				}
 			}
 		}
