ClothEditor: don't restart the preview animation or reset the preview camera when the graph re-evaluates.
Also, a fix for the preview viewport sometimes not immediately focusing on the cloth component when the Editor opens

diff --git a/Engine/Plugins/Experimental/ChaosClothAssetEditor/Source/ChaosClothAssetEditor/Private/ChaosClothAsset/ClothEditorMode.cpp b/Engine/Plugins/Experimental/ChaosClothAssetEditor/Source/ChaosClothAssetEditor/Private/ChaosClothAsset/ClothEditorMode.cpp
--- a/Engine/Plugins/Experimental/ChaosClothAssetEditor/Source/ChaosClothAssetEditor/Private/ChaosClothAsset/ClothEditorMode.cpp
+++ b/Engine/Plugins/Experimental/ChaosClothAssetEditor/Source/ChaosClothAssetEditor/Private/ChaosClothAsset/ClothEditorMode.cpp
@@ -61,6 +61,7 @@
 #include "Materials/Material.h"
 #include "MaterialDomain.h"
 #include "Dataflow/DataflowSNode.h"
+#include "Components/SkeletalMeshComponent.h"
 
 #include UE_INLINE_GENERATED_CPP_BY_NAME(ClothEditorMode)
 
@@ -618,17 +619,6 @@ void UChaosClothAssetEditorMode::InitializeTargets(const TArray<TObjectPtr<UObje
 
 	UBaseCharacterFXEditorMode::InitializeTargets(AssetsIn);
 
-	for (TObjectPtr<UObject> Asset : AssetsIn)
-	{
-		if (UChaosClothAsset* ChaosClothAsset = Cast<UChaosClothAsset>(Asset))
-		{
-			PreviewScene->SetClothAsset(ChaosClothAsset);
-			break;
-		}
-	}
-
-	ReinitializeDynamicMeshComponents();
-
 	DataflowComponent = NewObject<UDataflowComponent>();
 	DataflowComponent->RegisterComponentWithWorld(PreviewScene->GetWorld());
 }
@@ -803,13 +793,19 @@ FBox UChaosClothAssetEditorMode::SelectionBoundingBox() const
 
 FBox UChaosClothAssetEditorMode::PreviewBoundingBox() const
 {
-	if (PreviewScene->GetClothComponent())
+	FBox Bounds(ForceInit);
+
+	if (const UChaosClothComponent* const Cloth = PreviewScene->GetClothComponent())
+	{
+		Bounds += Cloth->Bounds.GetBox();
+	}
+
+	if (const USkeletalMeshComponent* const SkeletalMesh = PreviewScene->GetSkeletalMeshComponent())
 	{
-		FBoxSphereBounds Bounds = PreviewScene->GetClothComponent()->Bounds;
-		return Bounds.GetBox();
+		Bounds += SkeletalMesh->Bounds.GetBox();
 	}
 
-	return FBox(ForceInitToZero);
+	return Bounds;
 }
 
 void UChaosClothAssetEditorMode::SetConstructionViewMode(UE::Chaos::ClothAsset::EClothPatternVertexType InMode)
diff --git a/Engine/Plugins/Experimental/ChaosClothAssetEditor/Source/ChaosClothAssetEditor/Private/ChaosClothAsset/ClothEditorPreviewScene.cpp b/Engine/Plugins/Experimental/ChaosClothAssetEditor/Source/ChaosClothAssetEditor/Private/ChaosClothAsset/ClothEditorPreviewScene.cpp
--- a/Engine/Plugins/Experimental/ChaosClothAssetEditor/Source/ChaosClothAssetEditor/Private/ChaosClothAsset/ClothEditorPreviewScene.cpp
+++ b/Engine/Plugins/Experimental/ChaosClothAssetEditor/Source/ChaosClothAssetEditor/Private/ChaosClothAsset/ClothEditorPreviewScene.cpp
@@ -8,6 +8,8 @@
 #include "Elements/Framework/EngineElementsLibrary.h"
 #include "Animation/AnimSingleNodeInstance.h"
 #include "ComponentReregisterContext.h"
+#include "ChaosClothAsset/ClothAsset.h"
+#include "SkinnedAssetCompiler.h"
 
 #define LOCTEXT_NAMESPACE "UChaosClothEditorPreviewScene"
 
@@ -38,7 +40,7 @@ FChaosClothPreviewScene::FChaosClothPreviewScene(FPreviewScene::ConstructionValu
 	if (PreviewSceneDescription->SkeletalMeshAsset)
 	{
 		CreateSkeletalMeshComponent();
-		ReattachSkeletalMeshAndAnimation();
+		UpdateSkeletalMeshAnimation();
 	}
 }
 
@@ -70,7 +72,7 @@ void FChaosClothPreviewScene::AddReferencedObjects(FReferenceCollector& Collecto
 }
 
 
-void FChaosClothPreviewScene::ReattachSkeletalMeshAndAnimation()
+void FChaosClothPreviewScene::UpdateSkeletalMeshAnimation()
 {
 	if (SkeletalMeshComponent)
 	{
@@ -94,24 +96,14 @@ void FChaosClothPreviewScene::ReattachSkeletalMeshAndAnimation()
 			SkeletalMeshComponent->AnimScriptInstance = nullptr;
 		}
 	}
+}
 
-	if (ClothComponent && SkeletalMeshComponent)
+void FChaosClothPreviewScene::UpdateClothComponentAttachment()
+{
+	if (ClothComponent && SkeletalMeshComponent && !ClothComponent->IsAttachedTo(SkeletalMeshComponent))
 	{
-		const bool bWasSimming = !ClothComponent->IsSimulationSuspended();
-		ClothComponent->SuspendSimulation();
-
 		ClothComponent->AttachToComponent(SkeletalMeshComponent, FAttachmentTransformRules::SnapToTargetNotIncludingScale);
-
-		// hard reset and resume simulation if it was running prior
-		{
-			const FComponentReregisterContext Context(ClothComponent);
-		}
-		if (bWasSimming)
-		{
-			ClothComponent->ResumeSimulation();
-		}
 	}
-
 }
 
 void FChaosClothPreviewScene::SceneDescriptionPropertyChanged(struct FPropertyChangedEvent& PropertyChangedEvent)
@@ -126,6 +118,8 @@ void FChaosClothPreviewScene::SceneDescriptionPropertyChanged(struct FPropertyCh
 		{
 			DeleteSkeletalMeshComponent();
 		}
+		UpdateSkeletalMeshAnimation();
+		UpdateClothComponentAttachment();
 	}
 
 	if (PropertyChangedEvent.Property->GetFName() == GET_MEMBER_NAME_CHECKED(UChaosClothPreviewSceneDescription, SkeletalMeshTransform))
@@ -142,9 +136,9 @@ void FChaosClothPreviewScene::SceneDescriptionPropertyChanged(struct FPropertyCh
 		{
 			PreviewAnimInstance = nullptr;
 		}
+		UpdateSkeletalMeshAnimation();
 	}
-
-	ReattachSkeletalMeshAndAnimation();
+	
 }
 
 UChaosClothComponent* FChaosClothPreviewScene::GetClothComponent()
@@ -253,37 +247,32 @@ void FChaosClothPreviewScene::SetClothAsset(UChaosClothAsset* Asset)
 	check(Asset);
 	check(SceneActor);
 
-	bool bShouldClothBeRootComponent = true;
 	if (ClothComponent)
 	{
-		if (SceneActor->GetRootComponent() == ClothComponent)
+		ClothComponent->SetClothAsset(Asset);
+	}
+	else
+	{
+		ClothComponent = NewObject<UChaosClothComponent>(SceneActor);
+		ClothComponent->SetClothAsset(Asset);
+		ClothComponent->SelectionOverrideDelegate = UPrimitiveComponent::FSelectionOverride::CreateRaw(this, &FChaosClothPreviewScene::IsComponentSelected);
+
+		if (SkeletalMeshComponent)
 		{
-			SceneActor->SetRootComponent(nullptr);
+			check(SceneActor->GetRootComponent() == SkeletalMeshComponent);
 		}
 		else
 		{
-			// SkeletalMesh is the RootComponent, so don't change it
-			bShouldClothBeRootComponent = false;
+			SceneActor->SetRootComponent(ClothComponent);
 		}
 
-		ClothComponent->SelectionOverrideDelegate.Unbind();
-		SceneActor->RemoveOwnedComponent(ClothComponent);
-		ClothComponent->DestroyComponent();
-	}
-
-	ClothComponent = NewObject<UChaosClothComponent>(SceneActor);
-	ClothComponent->SetClothAsset(Asset);
-
-	ClothComponent->SelectionOverrideDelegate = UPrimitiveComponent::FSelectionOverride::CreateRaw(this, &FChaosClothPreviewScene::IsComponentSelected);
-
-	if (bShouldClothBeRootComponent)
-	{
-		SceneActor->SetRootComponent(ClothComponent);
+		SceneActor->RegisterAllComponents();
 	}
 
-	SceneActor->RegisterAllComponents();
+	UpdateClothComponentAttachment();
 
-	ReattachSkeletalMeshAndAnimation();
+	FSkinnedAssetCompilingManager::Get().FinishCompilation(TArrayView<USkinnedAsset* const>{Asset});
+	ClothComponent->UpdateBounds();
 }
 
 UAnimSingleNodeInstance* FChaosClothPreviewScene::GetPreviewAnimInstance()
diff --git a/Engine/Plugins/Experimental/ChaosClothAssetEditor/Source/ChaosClothAssetEditor/Private/ChaosClothAsset/ClothEditorPreviewScene.h b/Engine/Plugins/Experimental/ChaosClothAssetEditor/Source/ChaosClothAssetEditor/Private/ChaosClothAsset/ClothEditorPreviewScene.h
--- a/Engine/Plugins/Experimental/ChaosClothAssetEditor/Source/ChaosClothAssetEditor/Private/ChaosClothAsset/ClothEditorPreviewScene.h
+++ b/Engine/Plugins/Experimental/ChaosClothAssetEditor/Source/ChaosClothAssetEditor/Private/ChaosClothAsset/ClothEditorPreviewScene.h
@@ -83,9 +83,12 @@ public:
 
 private:
 
-	// (Re)attach the cloth component to the skeletal mesh component, if it exists. 
-	// Create the PreviewAnimationInstance if the AnimationAsset and SkeletalMesh both exist.
-	void ReattachSkeletalMeshAndAnimation();
+	// Create the PreviewAnimationInstance if the AnimationAsset and SkeletalMesh both exist, and set the animation to run on the SkeletalMeshComponent
+	void UpdateSkeletalMeshAnimation();
+
+	// Attach the cloth component to the skeletal mesh component, if it exists
+	void UpdateClothComponentAttachment();
+
 
 	void SkeletalMeshTransformChanged(USceneComponent* UpdatedComponent, EUpdateTransformFlags UpdateTransformFlags, ETeleportType Teleport);
 
diff --git a/Engine/Plugins/Experimental/ChaosClothAssetEditor/Source/ChaosClothAssetEditor/Private/ChaosClothAsset/ClothEditorToolkit.cpp b/Engine/Plugins/Experimental/ChaosClothAssetEditor/Source/ChaosClothAssetEditor/Private/ChaosClothAsset/ClothEditorToolkit.cpp
--- a/Engine/Plugins/Experimental/ChaosClothAssetEditor/Source/ChaosClothAssetEditor/Private/ChaosClothAsset/ClothEditorToolkit.cpp
+++ b/Engine/Plugins/Experimental/ChaosClothAssetEditor/Source/ChaosClothAssetEditor/Private/ChaosClothAsset/ClothEditorToolkit.cpp
@@ -1048,8 +1048,6 @@ void FChaosClothAssetEditorToolkit::OnNodeDeleted(const TSet<UObject*>& DeletedN
 
 void FChaosClothAssetEditorToolkit::OnClothAssetChanged()
 {
-	ClothPreviewViewportClient->ClearSelectedComponents();
-
 	TArray<TObjectPtr<UObject>> ObjectsToEdit;
 	OwningAssetEditor->GetObjectsToEdit(MutableView(ObjectsToEdit));
 
@@ -1061,12 +1059,18 @@ void FChaosClothAssetEditorToolkit::OnClothAssetChanged()
 
 	if (UChaosClothAsset* const ClothAsset = Cast<UChaosClothAsset>(ObjectsToEdit[0]))
 	{
+		const bool bHadClothComponent = (ClothPreviewScene->GetClothComponent() != nullptr);
+
 		ClothPreviewScene->SetClothAsset(ClothAsset);
 
 		ensure(ClothAsset->HasAnyFlags(RF_Transactional));		// Ensure all objects are transactable for undo/redo in the details panel
 		SetEditingObject(ClothAsset);
 
-		ClothPreviewViewportClient->FocusViewportOnBox(ClothMode->PreviewBoundingBox());
+		if (!bHadClothComponent)
+		{
+			// Focus on the cloth component if this is the first time adding one
+			ClothPreviewViewportClient->FocusViewportOnBox(ClothMode->PreviewBoundingBox());
+		}
 
 		PopulateOutliner();
 	}
