Cloth Editor: Handle DynamicMesh changing vertex indices from cloth collection.

diff --git a/Engine/Plugins/Experimental/ChaosClothAssetEditor/Source/ChaosClothAssetEditorTools/Private/ChaosClothAsset/ClothWeightMapPaintTool.cpp b/Engine/Plugins/Experimental/ChaosClothAssetEditor/Source/ChaosClothAssetEditorTools/Private/ChaosClothAsset/ClothWeightMapPaintTool.cpp
--- a/Engine/Plugins/Experimental/ChaosClothAssetEditor/Source/ChaosClothAssetEditorTools/Private/ChaosClothAsset/ClothWeightMapPaintTool.cpp
+++ b/Engine/Plugins/Experimental/ChaosClothAssetEditor/Source/ChaosClothAssetEditorTools/Private/ChaosClothAsset/ClothWeightMapPaintTool.cpp
@@ -12,6 +12,7 @@
 #include "MeshWeights.h"
 #include "DynamicMesh/MeshNormals.h"
 #include "DynamicMesh/MeshIndexUtil.h"
+#include "DynamicMesh/NonManifoldMappingSupport.h"
 #include "Util/BufferUtil.h"
 #include "Util/ColorConstants.h"
 #include "Selections/MeshConnectedComponents.h"
@@ -302,11 +303,34 @@ void UClothEditorWeightMapPaintTool::Setup()
 		if (TSharedPtr<const FManagedArrayCollection> ClothCollection = ClothEditorContextObject->GetSelectedClothCollection().Pin())
 		{
 			using namespace UE::Chaos::ClothAsset;
-			bHaveDynamicMeshToWeightConversion = ClothEditorContextObject->GetConstructionViewMode() == EClothPatternVertexType::Sim2D;
-			if (bHaveDynamicMeshToWeightConversion)
+			const FNonManifoldMappingSupport NonManifoldMapping(*Mesh);
+
+			const bool bHasNonManifoldMapping = NonManifoldMapping.IsNonManifoldVertexInSource();
+			const bool bHas2D3DConversion = ClothEditorContextObject->GetConstructionViewMode() == EClothPatternVertexType::Sim2D;
+
+			bHaveDynamicMeshToWeightConversion = bHasNonManifoldMapping || bHas2D3DConversion;
+
+			FCollectionClothConstFacade Cloth(ClothCollection);
+			check(Cloth.IsValid());
+			if (bHasNonManifoldMapping)
+			{
+				const TConstArrayView<int32> SimVertex3DLookup = Cloth.GetSimVertex3DLookup();
+
+				DynamicMeshToWeight.SetNumUninitialized(Mesh->VertexCount());
+				WeightToDynamicMesh.Reset();
+				WeightToDynamicMesh.SetNum(Cloth.GetNumSimVertices3D());
+				for (int32 DynamicMeshVert = 0; DynamicMeshVert < Mesh->VertexCount(); ++DynamicMeshVert)
+				{
+					DynamicMeshToWeight[DynamicMeshVert] = NonManifoldMapping.GetOriginalNonManifoldVertexID(DynamicMeshVert);
+					if (bHas2D3DConversion)
+					{
+						DynamicMeshToWeight[DynamicMeshVert] = SimVertex3DLookup[DynamicMeshToWeight[DynamicMeshVert]];
+					}
+					WeightToDynamicMesh[DynamicMeshToWeight[DynamicMeshVert]].Add(DynamicMeshVert);
+				}
+			}
+			else if (bHas2D3DConversion)
 			{
-				FCollectionClothConstFacade Cloth(ClothCollection);
-				check(Cloth.IsValid());
 				DynamicMeshToWeight = Cloth.GetSimVertex3DLookup();
 				WeightToDynamicMesh = Cloth.GetSimVertex2DLookup();
 			}
diff --git a/Engine/Plugins/Experimental/ChaosClothAssetEditor/Source/ChaosClothAssetTools/Private/ChaosClothAsset/ClothPatternToDynamicMesh.cpp b/Engine/Plugins/Experimental/ChaosClothAssetEditor/Source/ChaosClothAssetTools/Private/ChaosClothAsset/ClothPatternToDynamicMesh.cpp
--- a/Engine/Plugins/Experimental/ChaosClothAssetEditor/Source/ChaosClothAssetTools/Private/ChaosClothAsset/ClothPatternToDynamicMesh.cpp
+++ b/Engine/Plugins/Experimental/ChaosClothAssetEditor/Source/ChaosClothAssetTools/Private/ChaosClothAsset/ClothPatternToDynamicMesh.cpp
@@ -8,6 +8,7 @@
 #include "ChaosClothAsset/CollectionClothFacade.h"
 #include "ToDynamicMesh.h"
 #include "DynamicMesh/MeshNormals.h"
+#include "DynamicMesh/NonManifoldMappingSupport.h"
 #include "SkeletalMeshAttributes.h"
 #include "Animation/Skeleton.h"
 
@@ -33,8 +34,6 @@ public:
 		VertexDataType(VertexDataType),
 		Cloth(ClothFacade)
 	{
-		int32 NumFaces = 0;
-		int32 FaceOffset = 0;
 		if (PatternIndex == INDEX_NONE)
 		{
 			// All patterns in one dynamic mesh
@@ -50,21 +49,49 @@ public:
 				}
 				NormalIDs = VertIDs;
 
-				NumFaces = Cloth.GetNumRenderFaces();
-				FaceOffset = 0;
+				const TConstArrayView<FIntVector3> Indices = Cloth.GetRenderIndices();
+				const int32 NumFaces = Cloth.GetNumRenderFaces();
+				TriIDs.Reserve(NumFaces);
+				for (int32 TriIndex = 0; TriIndex < NumFaces; ++TriIndex)
+				{
+					if (Indices[TriIndex][0] != INDEX_NONE &&
+						Indices[TriIndex][1] != INDEX_NONE &&
+						Indices[TriIndex][2] != INDEX_NONE)
+					{
+						TriIDs.Add(TriIndex);
+					}
+				}
 			} break;
 			case EClothPatternVertexType::Sim2D:
 			{
 				const int32 NumVertices = Cloth.GetNumSimVertices2D();
-				VertIDs.SetNum(NumVertices);
+				const TConstArrayView<int32> SimVertex3DLookup = Cloth.GetSimVertex3DLookup();
+				VertIDs.Reserve(NumVertices);
+				NormalIDs.Reserve(NumVertices);
 				for (int32 VtxIndex = 0; VtxIndex < NumVertices; ++VtxIndex)
 				{
-					VertIDs[VtxIndex] = VtxIndex;
+					if (SimVertex3DLookup[VtxIndex] != INDEX_NONE)
+					{
+						VertIDs.Add(VtxIndex);
+						NormalIDs.Add(SimVertex3DLookup[VtxIndex]);
+					}
 				}
-				NormalIDs = Cloth.GetSimVertex3DLookup();
 
-				NumFaces = Cloth.GetNumSimFaces();
-				FaceOffset = 0;
+				const TConstArrayView<FIntVector3> Indices = Cloth.GetSimIndices2D();
+				const int32 NumFaces = Cloth.GetNumSimFaces();
+				TriIDs.Reserve(NumFaces);
+				for (int32 TriIndex = 0; TriIndex < NumFaces; ++TriIndex)
+				{
+					if (Indices[TriIndex][0] != INDEX_NONE &&
+						Indices[TriIndex][1] != INDEX_NONE &&
+						Indices[TriIndex][2] != INDEX_NONE &&
+						SimVertex3DLookup[Indices[TriIndex][0]] != INDEX_NONE &&
+						SimVertex3DLookup[Indices[TriIndex][1]] != INDEX_NONE &&
+						SimVertex3DLookup[Indices[TriIndex][2]] != INDEX_NONE)
+					{
+						TriIDs.Add(TriIndex);
+					}
+				}
 			}
 			break;
 			case EClothPatternVertexType::Sim3D:
@@ -77,8 +104,18 @@ public:
 				}
 				NormalIDs = VertIDs;
 
-				NumFaces = Cloth.GetNumSimFaces();
-				FaceOffset = 0;
+				const TConstArrayView<FIntVector3> Indices = Cloth.GetSimIndices3D();
+				const int32 NumFaces = Cloth.GetNumSimFaces();
+				TriIDs.Reserve(NumFaces);
+				for (int32 TriIndex = 0; TriIndex < NumFaces; ++TriIndex)
+				{
+					if (Indices[TriIndex][0] != INDEX_NONE &&
+						Indices[TriIndex][1] != INDEX_NONE &&
+						Indices[TriIndex][2] != INDEX_NONE)
+					{
+						TriIDs.Add(TriIndex);
+					}
+				}
 			}break;
 			default:
 				checkNoEntry();
@@ -100,23 +137,53 @@ public:
 				}
 				NormalIDs = VertIDs;
 
-				NumFaces = Pattern.GetNumRenderFaces();
-				FaceOffset = Pattern.GetRenderFacesOffset();
+				const TConstArrayView<FIntVector3> Indices = Pattern.GetRenderIndices();
+				const int32 NumFaces = Pattern.GetNumRenderFaces();
+				const int32 FaceOffset = Pattern.GetRenderFacesOffset();
+				TriIDs.Reserve(NumFaces);
+				for (int32 TriIndex = 0; TriIndex < NumFaces; ++TriIndex)
+				{
+					if (Indices[TriIndex][0] != INDEX_NONE &&
+						Indices[TriIndex][1] != INDEX_NONE &&
+						Indices[TriIndex][2] != INDEX_NONE)
+					{
+						TriIDs.Add(TriIndex + FaceOffset);
+					}
+				}
 			} break;
 			case EClothPatternVertexType::Sim2D:
 			{
 				FCollectionClothSimPatternConstFacade Pattern = Cloth.GetSimPattern(PatternIndex);
 				const int32 NumVertices = Pattern.GetNumSimVertices2D();
+				const TConstArrayView<int32> SimVertex3DLookup = Pattern.GetSimVertex3DLookup();
 				const int32 VertexOffset = Pattern.GetSimVertices2DOffset();
-				VertIDs.SetNum(NumVertices);
+				VertIDs.Reserve(NumVertices);
+				NormalIDs.Reserve(NumVertices);
 				for (int32 VtxIndex = 0; VtxIndex < NumVertices; ++VtxIndex)
 				{
-					VertIDs[VtxIndex] = VtxIndex + VertexOffset;
+					if (SimVertex3DLookup[VtxIndex] != INDEX_NONE)
+					{
+						VertIDs.Add(VtxIndex + VertexOffset);
+						NormalIDs.Add(SimVertex3DLookup[VtxIndex]);
+					}
 				}
-				NormalIDs = Pattern.GetSimVertex3DLookup();
 
-				NumFaces = Pattern.GetNumSimFaces();
-				FaceOffset = Pattern.GetSimFacesOffset();
+				const TConstArrayView<FIntVector3> Indices = Pattern.GetSimIndices2D();
+				const int32 NumFaces = Pattern.GetNumSimFaces();
+				const int32 FaceOffset = Pattern.GetSimFacesOffset();
+				TriIDs.Reserve(NumFaces);
+				for (int32 TriIndex = 0; TriIndex < NumFaces; ++TriIndex)
+				{
+					if (Indices[TriIndex][0] != INDEX_NONE &&
+						Indices[TriIndex][1] != INDEX_NONE &&
+						Indices[TriIndex][2] != INDEX_NONE &&
+						SimVertex3DLookup[Indices[TriIndex][0] - VertexOffset] != INDEX_NONE &&
+						SimVertex3DLookup[Indices[TriIndex][1] - VertexOffset] != INDEX_NONE &&
+						SimVertex3DLookup[Indices[TriIndex][2] - VertexOffset] != INDEX_NONE)
+					{
+						TriIDs.Add(TriIndex + FaceOffset);
+					}
+				}
 			} break;
 			case EClothPatternVertexType::Sim3D:
 			{
@@ -124,20 +191,25 @@ public:
 				VertIDs = Pattern.GetSimVertex3DLookup();
 				NormalIDs = Pattern.GetSimVertex3DLookup();
 
-				NumFaces = Pattern.GetNumSimFaces();
-				FaceOffset = Pattern.GetSimFacesOffset();
+				const TConstArrayView<FIntVector3> Indices = Pattern.GetSimIndices3D();
+				const int32 NumFaces = Pattern.GetNumSimFaces();
+				const int32 FaceOffset = Pattern.GetSimFacesOffset();
+				TriIDs.Reserve(NumFaces);
+				for (int32 TriIndex = 0; TriIndex < NumFaces; ++TriIndex)
+				{
+					if (Indices[TriIndex][0] != INDEX_NONE &&
+						Indices[TriIndex][1] != INDEX_NONE &&
+						Indices[TriIndex][2] != INDEX_NONE)
+					{
+						TriIDs.Add(TriIndex + FaceOffset);
+					}
+				}
 			} break;
 			default:
 				checkNoEntry();
 			}
 		}
 
-		TriIDs.SetNum(NumFaces);
-		for (int32 TriIndex = 0; TriIndex < NumFaces; ++TriIndex)
-		{
-			TriIDs[TriIndex] = TriIndex + FaceOffset;;
-		}
-
 		if(VertexDataType != EClothPatternVertexType::Render)
 		{
 			//
@@ -519,41 +591,50 @@ void FClothPatternToDynamicMesh::Convert(const TSharedPtr<const FManagedArrayCol
 	const FCollectionClothConstFacade ClothFacade(ClothCollection);
 
 	// Actual conversion
-	UE::Geometry::TToDynamicMesh<FClothPatternWrapper> PatternToDynamicMesh; 
+	UE::Geometry::TToDynamicMesh<FClothPatternWrapper> PatternToDynamicMesh;
 	FClothPatternWrapper PatternWrapper(ClothFacade, PatternIndex, VertexDataType);
 
-	const bool bDisableAttributes = false;
 	auto TriangleToGroupFunction = [](FClothPatternWrapper::TriIDType) { return 0; };
-
-	if (bDisableAttributes)
-	{
-		MeshOut.DiscardAttributes();
-		PatternToDynamicMesh.ConvertWOAttributes(MeshOut, PatternWrapper, TriangleToGroupFunction);
-	}
-	else
+	const bool bIsRenderType = VertexDataType == EClothPatternVertexType::Render;
+	auto TriangleToMaterialFunction = [PatternIndex, bIsRenderType, &ClothFacade](FClothPatternWrapper::TriIDType TriID)->int32
 	{
-		MeshOut.EnableAttributes();
-
-		const bool bIsRenderType = VertexDataType == EClothPatternVertexType::Render;
-		auto TriangleToMaterialFunction = [PatternIndex, bIsRenderType, &ClothFacade](FClothPatternWrapper::TriIDType TriID)->int32
+		if (bIsRenderType)
 		{
-			if (bIsRenderType)
+			if (PatternIndex != INDEX_NONE)
 			{
-				if (PatternIndex != INDEX_NONE)
-				{
-					return PatternIndex;
-				}
-				const int32 FoundPattern = ClothFacade.FindRenderPatternByFaceIndex(TriID);
-				check(FoundPattern != INDEX_NONE);
-				return FoundPattern;
+				return PatternIndex;
 			}
-			
-			// Sim meshes will have a default material with MaterialID zero applied
-			return 0;
-		};
+			const int32 FoundPattern = ClothFacade.FindRenderPatternByFaceIndex(TriID);
+			check(FoundPattern != INDEX_NONE);
+			return FoundPattern;
+		}
 
-		constexpr bool bCopyTangents = false;
-		PatternToDynamicMesh.Convert(MeshOut, PatternWrapper, TriangleToGroupFunction, TriangleToMaterialFunction, bCopyTangents);
+		// Sim meshes will have a default material with MaterialID zero applied
+		return 0;
+	};
+
+	constexpr bool bCopyTangents = false;
+	PatternToDynamicMesh.Convert(MeshOut, PatternWrapper, TriangleToGroupFunction, TriangleToMaterialFunction, bCopyTangents);
+
+	// Add non-manifold mapping data if DynamicMesh indices don't match orig indices
+	const int32 NumVertices = VertexDataType == EClothPatternVertexType::Render ? ClothFacade.GetNumRenderVertices() :
+		VertexDataType == EClothPatternVertexType::Sim2D ? ClothFacade.GetNumSimVertices2D() : ClothFacade.GetNumSimVertices3D();
+	bool bVertexIndicesMatch = false;
+	if (PatternToDynamicMesh.ToSrcVertIDMap.Num() == NumVertices)
+	{
+		bVertexIndicesMatch = true;
+		for (int32 VertId = 0; VertId < NumVertices; ++VertId)
+		{
+			if (PatternToDynamicMesh.ToSrcVertIDMap[VertId] != VertId)
+			{
+				bVertexIndicesMatch = false;
+				break;
+			}
+		}
+	}
+	if (!bVertexIndicesMatch)
+	{
+		UE::Geometry::FNonManifoldMappingSupport::AttachNonManifoldVertexMappingData(PatternToDynamicMesh.ToSrcVertIDMap, MeshOut);
 	}
 
 }
