Cloth Editor: Outliner changes
- populate the Outliner based on the output FManagedArray of the currently selected node
- clear the Outliner if no node is selected or if the selected node has no FManagedArray
- move the group name combo box into the Outliner class
- remember the last group name selected when deselecting a node -- when a valid node is then selected, repopulate the selected group name from the saved group name

diff --git a/Engine/Plugins/Experimental/ChaosClothAssetEditor/Source/ChaosClothAssetEditor/Private/ChaosClothAsset/ClothEditorToolkit.cpp b/Engine/Plugins/Experimental/ChaosClothAssetEditor/Source/ChaosClothAssetEditor/Private/ChaosClothAsset/ClothEditorToolkit.cpp
--- a/Engine/Plugins/Experimental/ChaosClothAssetEditor/Source/ChaosClothAssetEditor/Private/ChaosClothAsset/ClothEditorToolkit.cpp
+++ b/Engine/Plugins/Experimental/ChaosClothAssetEditor/Source/ChaosClothAssetEditor/Private/ChaosClothAsset/ClothEditorToolkit.cpp
@@ -632,46 +632,9 @@ TSharedRef<SDockTab> FChaosClothAssetEditorToolkit::SpawnTab_ClothPreview(const
 
 TSharedRef<SDockTab> FChaosClothAssetEditorToolkit::SpawnTab_Outliner(const FSpawnTabArgs& Args)
 {
-
 	TSharedRef<SDockTab> DockableTab = SNew(SDockTab)
 	[
 		SNew(SVerticalBox)
-
-		+ SVerticalBox::Slot()
-		.AutoHeight()
-		[
-			SAssignNew(SelectedGroupNameComboBox, SComboBox<FName>)
-			.OptionsSource(&ClothCollectionGroupNames)
-			.OnSelectionChanged(SComboBox<FName>::FOnSelectionChanged::CreateLambda(
-				[this](FName SelectedName, ESelectInfo::Type)
-				{
-					if (Outliner)
-					{
-						Outliner->SetSelectedGroupName(SelectedName);	// this will also rebuild the table
-					}
-				}))
-			.OnGenerateWidget(SComboBox<FName>::FOnGenerateWidget::CreateLambda(
-				[](FName Item)
-				{
-					return SNew(STextBlock)
-						.Text(FText::FromName(Item));
-				}))
-			[
-				SNew(STextBlock)
-				.Text_Lambda([this]()
-				{
-					if (Outliner)
-					{
-						return FText::FromName(Outliner->GetSelectedGroupName());
-					}
-					else
-					{
-						return LOCTEXT("NoSelectedGroup", "");
-					}
-				})
-			]
-		]
-
 		+ SVerticalBox::Slot()
 		[
 			SAssignNew(Outliner, SClothCollectionOutliner)
@@ -747,37 +710,8 @@ void FChaosClothAssetEditorToolkit::InitDetailsViewPanel()
 		PreviewSceneDockTab->SetContent(AdvancedPreviewSettingsWidget.ToSharedRef());
 	}
 
-	PopulateOutliner();
 }
 
-void FChaosClothAssetEditorToolkit::PopulateOutliner()
-{
-	TSharedPtr<FManagedArrayCollection> ClothCollection;
-	if (const TObjectPtr<UChaosClothComponent> ClothComponent = ClothPreviewScene->GetClothComponent())
-	{
-		if (UChaosClothAsset* const ClothAsset = ClothComponent->GetClothAsset())
-		{
-			// TODO: handle multiple LODs
-			if (ClothAsset->GetClothCollections().Num())
-			{
-				ClothCollection = ClothAsset->GetClothCollections()[0];
-			}
-		}
-	}
-	
-	if (Outliner.IsValid() && ClothCollection.IsValid())
-	{
-		ClothCollectionGroupNames = ClothCollection->GroupNames();
-		if (ClothCollectionGroupNames.Num())
-		{
-			const FName& SelectedGroupName = ClothCollectionGroupNames[0];
-	
-			Outliner->SetClothCollection(ClothCollection);
-			Outliner->SetSelectedGroupName(SelectedGroupName);
-		}
-
-	}
-}
 
 void FChaosClothAssetEditorToolkit::OnFinishedChangingAssetProperties(const FPropertyChangedEvent& Event)
 {
@@ -1071,8 +1005,6 @@ void FChaosClothAssetEditorToolkit::OnClothAssetChanged()
 			// Focus on the cloth component if this is the first time adding one
 			ClothPreviewViewportClient->FocusViewportOnBox(ClothMode->PreviewBoundingBox());
 		}
-
-		PopulateOutliner();
 	}
 
 	if (bWasSimulationSuspended)
diff --git a/Engine/Plugins/Experimental/ChaosClothAssetEditor/Source/ChaosClothAssetEditor/Private/ChaosClothAsset/ClothEditorToolkit.h b/Engine/Plugins/Experimental/ChaosClothAssetEditor/Source/ChaosClothAssetEditor/Private/ChaosClothAsset/ClothEditorToolkit.h
--- a/Engine/Plugins/Experimental/ChaosClothAssetEditor/Source/ChaosClothAssetEditor/Private/ChaosClothAsset/ClothEditorToolkit.h
+++ b/Engine/Plugins/Experimental/ChaosClothAssetEditor/Source/ChaosClothAssetEditor/Private/ChaosClothAsset/ClothEditorToolkit.h
@@ -109,7 +109,6 @@ private:
 	void InitDetailsViewPanel();
 	void OnFinishedChangingAssetProperties(const FPropertyChangedEvent&);
 
-	void PopulateOutliner();
 	void OnClothAssetChanged();
 	void InvalidateViews();
 
@@ -143,10 +142,6 @@ private:
 
 	TSharedPtr<SClothCollectionOutliner> Outliner;
 
-	TSharedPtr<SComboBox<FName>> SelectedGroupNameComboBox;
-	TArray<FName> ClothCollectionGroupNames;		// Data source for SelectedGroupNameComboBox
-
-
 	// Dataflow
 	UDataflow* Dataflow = nullptr;
 	FString DataflowTerminalPath = "";
diff --git a/Engine/Plugins/Experimental/ChaosClothAssetEditor/Source/ChaosClothAssetEditor/Private/ChaosClothAsset/SClothCollectionOutliner.cpp b/Engine/Plugins/Experimental/ChaosClothAssetEditor/Source/ChaosClothAssetEditor/Private/ChaosClothAsset/SClothCollectionOutliner.cpp
--- a/Engine/Plugins/Experimental/ChaosClothAssetEditor/Source/ChaosClothAssetEditor/Private/ChaosClothAsset/SClothCollectionOutliner.cpp
+++ b/Engine/Plugins/Experimental/ChaosClothAssetEditor/Source/ChaosClothAssetEditor/Private/ChaosClothAsset/SClothCollectionOutliner.cpp
@@ -61,6 +61,32 @@ void SClothCollectionOutliner::Construct(const FArguments& InArgs)
 	ChildSlot
 	[
 		SNew(SVerticalBox)
+
+		+ SVerticalBox::Slot()
+		.AutoHeight()
+		[
+			SAssignNew(SelectedGroupNameComboBox, SComboBox<FName>)
+			.OptionsSource(&ClothCollectionGroupNames)
+			.OnSelectionChanged(SComboBox<FName>::FOnSelectionChanged::CreateLambda(
+				[this](FName SelectedName, ESelectInfo::Type)
+				{
+					SetSelectedGroupName(SelectedName);
+				}))
+			.OnGenerateWidget(SComboBox<FName>::FOnGenerateWidget::CreateLambda(
+				[](FName Item)
+				{
+					return SNew(STextBlock)
+						.Text(FText::FromName(Item));
+				}))
+			[
+				SNew(STextBlock)
+				.Text_Lambda([this]()
+				{
+					return FText::FromName(GetSelectedGroupName());
+				})
+			]
+		]
+
 		+ SVerticalBox::Slot()
 		.Padding(FMargin(0.0f, 3.f))
 		[
@@ -77,15 +103,35 @@ void SClothCollectionOutliner::Construct(const FArguments& InArgs)
 	];
 }
 
-void SClothCollectionOutliner::SetClothCollection(TWeakPtr<FManagedArrayCollection> InClothCollection)
+void SClothCollectionOutliner::SetClothCollection(TSharedPtr<FManagedArrayCollection> InClothCollection)
 {
-	if (ClothCollection != InClothCollection)
+	if (!InClothCollection.IsValid() && ClothCollection.IsValid())
+	{
+		// Going from valid node selected to no node selected
+		SavedLastValidGroupName = SelectedGroupName;
+	}
+	else if (InClothCollection.IsValid() && !ClothCollection.IsValid())
 	{
-		ClothCollection = InClothCollection;
+		// Going from no node selected to a valid node selected
+		SelectedGroupName = SavedLastValidGroupName;
+	}
+
+	ClothCollection = InClothCollection;
 
+	if (ClothCollection.IsValid())
+	{
+		ClothCollectionGroupNames = ClothCollection->GroupNames();
 		RegenerateHeader();
 		RepopulateListView();
 	}
+	else
+	{
+		ClothCollectionGroupNames.Reset();
+		HeaderRowWidget->ClearColumns();
+		HeaderData->AttributeNames.Empty();
+		ListItems.Empty();
+		SelectedGroupName = FName();
+	}
 }
 
 void SClothCollectionOutliner::SetSelectedGroupName(const FName& InSelectedGroupName)
@@ -105,9 +151,7 @@ void SClothCollectionOutliner::RegenerateHeader()
 {
 	constexpr float CustomFillWidth = 2.0f;
 
-	TSharedPtr<FManagedArrayCollection> PinnedClothCollection = ClothCollection.Pin();
-
-	if (!PinnedClothCollection.IsValid())
+	if (!ClothCollection.IsValid())
 	{
 		return;
 	}
@@ -116,7 +160,7 @@ void SClothCollectionOutliner::RegenerateHeader()
 
 	HeaderData = MakeShared<FClothCollectionHeaderData>();
 	HeaderData->AttributeNames.Add(FClothCollectionHeaderData::ColumnZeroName);
-	HeaderData->AttributeNames.Append(PinnedClothCollection->AttributeNames(SelectedGroupName));
+	HeaderData->AttributeNames.Append(ClothCollection->AttributeNames(SelectedGroupName));
 
 	for (int32 AttributeNameIndex = 0; AttributeNameIndex < HeaderData->AttributeNames.Num(); ++AttributeNameIndex)
 	{
@@ -243,14 +287,12 @@ void SClothCollectionOutliner::RepopulateListView()
 {
 	ListItems.Empty();
 
-	const TSharedPtr<const FManagedArrayCollection> PinnedClothCollection = ClothCollection.Pin();
-
-	if (!PinnedClothCollection.IsValid())
+	if (!ClothCollection.IsValid())
 	{
 		return;
 	}
 
-	const int32 NumElements = PinnedClothCollection->NumElements(SelectedGroupName);
+	const int32 NumElements = ClothCollection->NumElements(SelectedGroupName);
 
 	for (int32 ElementIndex = 0; ElementIndex < NumElements; ++ElementIndex)
 	{
@@ -266,7 +308,7 @@ void SClothCollectionOutliner::RepopulateListView()
 			}
 			else
 			{
-				NewItem->AttributeValues[AttributeNameIndex] = ClothCollectionOutlinerHelpers::AttributeValueToString(*PinnedClothCollection, AttributeName, SelectedGroupName, ElementIndex);
+				NewItem->AttributeValues[AttributeNameIndex] = ClothCollectionOutlinerHelpers::AttributeValueToString(*ClothCollection, AttributeName, SelectedGroupName, ElementIndex);
 			}
 		}
 
diff --git a/Engine/Plugins/Experimental/ChaosClothAssetEditor/Source/ChaosClothAssetEditor/Private/ChaosClothAsset/SClothCollectionOutliner.h b/Engine/Plugins/Experimental/ChaosClothAssetEditor/Source/ChaosClothAssetEditor/Private/ChaosClothAsset/SClothCollectionOutliner.h
--- a/Engine/Plugins/Experimental/ChaosClothAssetEditor/Source/ChaosClothAssetEditor/Private/ChaosClothAsset/SClothCollectionOutliner.h
+++ b/Engine/Plugins/Experimental/ChaosClothAssetEditor/Source/ChaosClothAssetEditor/Private/ChaosClothAsset/SClothCollectionOutliner.h
@@ -4,6 +4,7 @@
 
 #include "Widgets/SCompoundWidget.h"
 #include "Widgets/Views/SListView.h"
+#include "Widgets/Input/SComboBox.h"
 
 struct FManagedArrayCollection;
 
@@ -45,13 +46,13 @@ class SClothCollectionOutliner : public SCompoundWidget
 public:
 
 	SLATE_BEGIN_ARGS(SClothCollectionOutliner) {}
-		SLATE_ARGUMENT(TWeakPtr<FManagedArrayCollection>, ClothCollection)
+		SLATE_ARGUMENT(TSharedPtr<FManagedArrayCollection>, ClothCollection)
 		SLATE_ARGUMENT(FName, SelectedGroupName)
 	SLATE_END_ARGS()
 
 	void Construct(const FArguments& InArgs);
 
-	void SetClothCollection(TWeakPtr<FManagedArrayCollection> ClothCollection);
+	void SetClothCollection(TSharedPtr<FManagedArrayCollection> ClothCollection);
 
 	///  Changing the selected group name will also automatically rebuild the header and all rows
 	void SetSelectedGroupName(const FName& SelectedGroupName);
@@ -62,8 +63,13 @@ public:
 
 private:
 
-	TWeakPtr<FManagedArrayCollection> ClothCollection;
+	TSharedPtr<FManagedArrayCollection> ClothCollection;
+
+	TSharedPtr<SComboBox<FName>> SelectedGroupNameComboBox;
+	TArray<FName> ClothCollectionGroupNames;		// Data source for SelectedGroupNameComboBox
+
 	FName SelectedGroupName;
+	FName SavedLastValidGroupName;
 
 	TSharedPtr<SListView<TSharedPtr<const FClothCollectionItem>>> ListView;
 	TArray<TSharedPtr<const FClothCollectionItem>> ListItems;
