Cloth Editor: Disable standard mouse controls when in 2D construction view. The only mouse controls that will work are right-drag for moving the camera in the xy plane, and scroll wheel for moving the camera in the z direction. In particular this disables camera rotation with LMB drag and alt+RMB drag.

diff --git a/Engine/Plugins/Experimental/ChaosClothAssetEditor/Source/ChaosClothAssetEditor/Private/ChaosClothAsset/ClothEditorRestSpaceViewportClient.cpp b/Engine/Plugins/Experimental/ChaosClothAssetEditor/Source/ChaosClothAssetEditor/Private/ChaosClothAsset/ClothEditorRestSpaceViewportClient.cpp
--- a/Engine/Plugins/Experimental/ChaosClothAssetEditor/Source/ChaosClothAssetEditor/Private/ChaosClothAsset/ClothEditorRestSpaceViewportClient.cpp
+++ b/Engine/Plugins/Experimental/ChaosClothAssetEditor/Source/ChaosClothAssetEditor/Private/ChaosClothAsset/ClothEditorRestSpaceViewportClient.cpp
@@ -111,23 +111,11 @@ bool FChaosClothEditorRestSpaceViewportClient::ShouldOrbitCamera() const
 
 bool FChaosClothEditorRestSpaceViewportClient::InputKey(const FInputKeyEventArgs& EventArgs)
 {
-	// See if any tool commands want to handle the key event
-	const TSharedPtr<FUICommandList> PinnedToolCommandList = ToolCommandList.Pin();
-	if (EventArgs.Event != IE_Released && PinnedToolCommandList.IsValid())
-	{
-		const FModifierKeysState KeyState = FSlateApplication::Get().GetModifierKeys();
-		if (PinnedToolCommandList->ProcessCommandBindings(EventArgs.Key, KeyState, (EventArgs.Event == IE_Repeat)))
-		{
-			return true;
-		}
-	}
-
 	if (ConstructionViewMode != EClothPatternVertexType::Sim2D)
 	{
 		return FEditorViewportClient::InputKey(EventArgs);
 	}
 
-
 	// We'll support disabling input like our base class, even if it does not end up being used.
 	if (bDisableInput)
 	{
@@ -135,66 +123,7 @@ bool FChaosClothEditorRestSpaceViewportClient::InputKey(const FInputKeyEventArgs
 	}
 
 	// Our viewport manipulation is placed in the input router that ModeTools manages
-	if (ModeTools->InputKey(this, EventArgs.Viewport, EventArgs.Key, EventArgs.Event))
-	{
-		return true;
-	}
-	
-	//
-	// Recreate the click handling from FEditorViewportClient::Internal_InputKey so that we can still have 
-	// mouse click, keyboard, etc., but skip the code dealing with camera controls
-	//
-
-	bool bHandled = false;
-
-	const FInputEventState InputState(EventArgs.Viewport, EventArgs.Key, EventArgs.Event);
-	const FViewport* const InViewport = EventArgs.Viewport;
-	const EInputEvent Event = EventArgs.Event;
-
-	const bool bWasCursorVisible = InViewport->IsCursorVisible();
-	const bool bWasSoftwareCursorVisible = InViewport->IsSoftwareCursorVisible();
-
-	RequiredCursorVisibiltyAndAppearance.bDontResetCursor = false;
-	UpdateRequiredCursorVisibility();
-
-	if (bWasCursorVisible != RequiredCursorVisibiltyAndAppearance.bHardwareCursorVisible || bWasSoftwareCursorVisible != RequiredCursorVisibiltyAndAppearance.bSoftwareCursorVisible)
-	{
-		bHandled = true;
-	}
-
-	// Compute a view.
-	FSceneViewFamilyContext ViewFamily(FSceneViewFamily::ConstructionValues(
-		InViewport,
-		GetScene(),
-		EngineShowFlags)
-		.SetRealtimeUpdate(IsRealtime()));
-	FSceneView* const View = CalcSceneView(&ViewFamily);
-
-	// Start tracking if any mouse button is down and it was a tracking event (MouseButton/Ctrl/Shift/Alt):
-	if (InputState.IsAnyMouseButtonDown()
-		&& (Event == IE_Pressed || Event == IE_Released)
-		&& (InputState.IsMouseButtonEvent() || InputState.IsCtrlButtonEvent() || InputState.IsAltButtonEvent() || InputState.IsShiftButtonEvent()))
-	{
-		StartTrackingDueToInput(InputState, *View);
-		return true;
-	}
-
-	// If we are tracking and no mouse button is down and this input event released the mouse button stop tracking and process any clicks if necessary
-	if (bIsTracking && !InputState.IsAnyMouseButtonDown() && InputState.IsMouseButtonEvent())
-	{
-		// Handle possible mouse click viewport
-		ProcessClickInViewport(InputState, *View);
-
-		// Stop tracking if no mouse button is down
-		StopTracking();
-
-		bHandled |= true;
-	}
-
-	// apply the visibility and set the cursor positions
-	ApplyRequiredCursorVisibility(true);
-
-	return bHandled;
+	return ModeTools->InputKey(this, EventArgs.Viewport, EventArgs.Key, EventArgs.Event);
 }
 
 void FChaosClothEditorRestSpaceViewportClient::SetEditorViewportWidget(TWeakPtr<SEditorViewport> InEditorViewportWidget)
