Chaos Cloth Editor- Fixed crash when viewing an empty weight map in the outliner.

#rnx

diff --git a/Engine/Plugins/Experimental/ChaosClothAssetEditor/Source/ChaosClothAssetEditor/Private/ChaosClothAsset/SClothCollectionOutliner.cpp b/Engine/Plugins/Experimental/ChaosClothAssetEditor/Source/ChaosClothAssetEditor/Private/ChaosClothAsset/SClothCollectionOutliner.cpp
--- a/Engine/Plugins/Experimental/ChaosClothAssetEditor/Source/ChaosClothAssetEditor/Private/ChaosClothAsset/SClothCollectionOutliner.cpp
+++ b/Engine/Plugins/Experimental/ChaosClothAssetEditor/Source/ChaosClothAssetEditor/Private/ChaosClothAsset/SClothCollectionOutliner.cpp
@@ -162,9 +162,14 @@ void SClothCollectionOutliner::RegenerateHeader()
 	HeaderData->AttributeNames.Add(FClothCollectionHeaderData::ColumnZeroName);
 	HeaderData->AttributeNames.Append(ClothCollection->AttributeNames(SelectedGroupName));
 
+	int32 UnnamedCount = 0;
 	for (int32 AttributeNameIndex = 0; AttributeNameIndex < HeaderData->AttributeNames.Num(); ++AttributeNameIndex)
 	{
-		const FName& AttrName = HeaderData->AttributeNames[AttributeNameIndex];
+		FName AttrName = HeaderData->AttributeNames[AttributeNameIndex];
+		if (AttrName == NAME_None)  // AddColumn needs a name, otherwise it crashes
+		{
+			AttrName = FName(FString::Format(TEXT("Unnamed{0}"), { UnnamedCount++ }));
+		}
 
 		HeaderRowWidget->AddColumn(
 			SHeaderRow::Column(AttrName)
