Cloth Weight Map Paint Tool: prevent out of bounds error if the node has more weights than the mesh has vertices, which can happen if the mesh changes in between tool invocations

diff --git a/Engine/Plugins/Experimental/ChaosClothAssetEditor/Source/ChaosClothAssetEditorTools/Private/ChaosClothAsset/ClothWeightMapPaintTool.cpp b/Engine/Plugins/Experimental/ChaosClothAssetEditor/Source/ChaosClothAssetEditorTools/Private/ChaosClothAsset/ClothWeightMapPaintTool.cpp
--- a/Engine/Plugins/Experimental/ChaosClothAssetEditor/Source/ChaosClothAssetEditorTools/Private/ChaosClothAsset/ClothWeightMapPaintTool.cpp
+++ b/Engine/Plugins/Experimental/ChaosClothAssetEditor/Source/ChaosClothAssetEditorTools/Private/ChaosClothAsset/ClothWeightMapPaintTool.cpp
@@ -347,24 +347,30 @@ void UClothEditorWeightMapPaintTool::Setup()
 	ActiveWeightMap = Mesh->Attributes()->GetWeightLayer(NumAttributeLayers);
 	ActiveWeightMap->SetName(FName("PaintLayer"));
 
-	// Get weights from selected node
+	// Copy weights from selected node to the preview mesh
 	const TArray<float>& CurrentWeights = WeightMapNodeToUpdate->VertexWeights;
 	
 	if (bHaveDynamicMeshToWeightConversion)
 	{
-		for (int32 WeightID = 0; WeightID < CurrentWeights.Num(); ++WeightID)
+		if (WeightToDynamicMesh.Num() == CurrentWeights.Num())	// Only copy node weights if they match the number of mesh vertices
 		{
-			for (const int32 VertexID : WeightToDynamicMesh[WeightID])
+			for (int32 WeightID = 0; WeightID < CurrentWeights.Num(); ++WeightID)
 			{
-				ActiveWeightMap->SetValue(VertexID, &CurrentWeights[WeightID]);
+				for (const int32 VertexID : WeightToDynamicMesh[WeightID])
+				{
+					ActiveWeightMap->SetValue(VertexID, &CurrentWeights[WeightID]);
+				}
 			}
 		}
 	}
 	else
 	{
-		for (int32 VertexID = 0; VertexID < CurrentWeights.Num(); ++VertexID)
+		if (Mesh->MaxVertexID() == CurrentWeights.Num())	// Only copy node weights if they match the number of mesh vertices
 		{
-			ActiveWeightMap->SetValue(VertexID, &CurrentWeights[VertexID]);
+			for (int32 VertexID = 0; VertexID < CurrentWeights.Num(); ++VertexID)
+			{
+				ActiveWeightMap->SetValue(VertexID, &CurrentWeights[VertexID]);
+			}
 		}
 	}
 
