Chaos Cloth Editor - Fixed Outliner crash when selecting a non valid cloth collection node.

#rnx

diff --git a/Engine/Plugins/Experimental/ChaosClothAssetEditor/Source/ChaosClothAssetEditor/Private/ChaosClothAsset/SClothCollectionOutliner.cpp b/Engine/Plugins/Experimental/ChaosClothAssetEditor/Source/ChaosClothAssetEditor/Private/ChaosClothAsset/SClothCollectionOutliner.cpp
--- a/Engine/Plugins/Experimental/ChaosClothAssetEditor/Source/ChaosClothAssetEditor/Private/ChaosClothAsset/SClothCollectionOutliner.cpp
+++ b/Engine/Plugins/Experimental/ChaosClothAssetEditor/Source/ChaosClothAssetEditor/Private/ChaosClothAsset/SClothCollectionOutliner.cpp
@@ -128,7 +128,10 @@ void SClothCollectionOutliner::SetClothCollection(TSharedPtr<FManagedArrayCollec
 	{
 		ClothCollectionGroupNames.Reset();
 		HeaderRowWidget->ClearColumns();
-		HeaderData->AttributeNames.Empty();
+		if (HeaderData)
+		{
+			HeaderData->AttributeNames.Empty();
+		}
 		ListItems.Empty();
 		SelectedGroupName = FName();
 	}
