Add Cloth stats to Cloth Editor display. Only show stats if "Simulation Statistics" menu item is checked in the debug visualization.

diff --git a/Engine/Plugins/Experimental/ChaosClothAsset/Source/ChaosClothAssetEngine/Private/ChaosClothAsset/ClothSimulationProxy.cpp b/Engine/Plugins/Experimental/ChaosClothAsset/Source/ChaosClothAssetEngine/Private/ChaosClothAsset/ClothSimulationProxy.cpp
--- a/Engine/Plugins/Experimental/ChaosClothAsset/Source/ChaosClothAssetEngine/Private/ChaosClothAsset/ClothSimulationProxy.cpp
+++ b/Engine/Plugins/Experimental/ChaosClothAsset/Source/ChaosClothAssetEngine/Private/ChaosClothAsset/ClothSimulationProxy.cpp
@@ -148,6 +148,11 @@ PRAGMA_ENABLE_DEPRECATION_WARNINGS
 		Solver->AddCloth(Cloths[ClothIndex].Get());
 		Cloths[ClothIndex]->Reset();
 
+		// Update cloth stats
+		NumCloths = 1;
+		NumKinematicParticles = Cloths[ClothIndex]->GetNumActiveKinematicParticles();
+		NumDynamicParticles = Cloths[ClothIndex]->GetNumActiveDynamicParticles();
+
 		// Create solver config simulation thread object
 		const int32 SolverConfigIndex = Configs.Emplace(MakeUnique<FClothingSimulationConfig>(ClothComponent.GetPropertyCollection()));  // TODO: Use a separate solver config for outfits
 		Solver->SetConfig(Configs[SolverConfigIndex].Get());
@@ -244,6 +249,7 @@ PRAGMA_ENABLE_DEPRECATION_WARNINGS
 
 		// Keep the actual used number of iterations for the stats
 		NumIterations = Solver->GetNumUsedIterations();
+		NumSubsteps = Solver->GetNumSubsteps();
 
 		// Update simulation time in ms (and provide an instant average instead of the value in real-time)
 		const float CurrSimulationTime = (float)((FPlatformTime::Seconds() - StartTime) * 1000.);
diff --git a/Engine/Plugins/Experimental/ChaosClothAsset/Source/ChaosClothAssetEngine/Public/ChaosClothAsset/ClothSimulationProxy.h b/Engine/Plugins/Experimental/ChaosClothAsset/Source/ChaosClothAssetEngine/Public/ChaosClothAsset/ClothSimulationProxy.h
--- a/Engine/Plugins/Experimental/ChaosClothAsset/Source/ChaosClothAssetEngine/Public/ChaosClothAsset/ClothSimulationProxy.h
+++ b/Engine/Plugins/Experimental/ChaosClothAsset/Source/ChaosClothAssetEngine/Public/ChaosClothAsset/ClothSimulationProxy.h
@@ -61,6 +61,14 @@ namespace UE::Chaos::ClothAsset
 
 		const ::Chaos::FClothVisualization* GetVisualization() const { return Visualization.Get(); }
 
+		int32 GetNumCloths() const { return NumCloths; }
+		int32 GetNumKinematicParticles() const { return NumKinematicParticles; }
+		int32 GetNumDynamicParticles() const { return NumDynamicParticles; }
+		int32 GetNumIterations() const { return NumIterations; }
+		int32 GetNumSubsteps() const { return NumSubsteps; }
+		float GetSimulationTime() const { return SimulationTime; }
+		bool IsTeleported() const { return bIsTeleported; }
+
 	protected:
 		void Tick();
 		void WriteSimulationData();
diff --git a/Engine/Plugins/Experimental/ChaosClothAssetEditor/Source/ChaosClothAssetEditor/Private/ChaosClothAsset/ClothEditorSimulationVisualization.cpp b/Engine/Plugins/Experimental/ChaosClothAssetEditor/Source/ChaosClothAssetEditor/Private/ChaosClothAsset/ClothEditorSimulationVisualization.cpp
--- a/Engine/Plugins/Experimental/ChaosClothAssetEditor/Source/ChaosClothAssetEditor/Private/ChaosClothAsset/ClothEditorSimulationVisualization.cpp
+++ b/Engine/Plugins/Experimental/ChaosClothAssetEditor/Source/ChaosClothAssetEditor/Private/ChaosClothAsset/ClothEditorSimulationVisualization.cpp
@@ -15,6 +15,57 @@ namespace UE::Chaos::ClothAsset
 {
 namespace Private
 {
+FText ConcatenateLine(const FText& InText, const FText& InNewLine)
+{
+	if (InText.IsEmpty())
+	{
+		return InNewLine;
+	}
+	return FText::Format(LOCTEXT("ViewportTextNewlineFormatter", "{0}\n{1}"), InText, InNewLine);
+}
+
+static FText GetSimulationStatisticsString(const FClothSimulationProxy* SimProxy)
+{
+	FText TextValue;
+	if (SimProxy)
+	{
+		// Cloth stats
+		if (const int32 NumActiveCloths = SimProxy->GetNumCloths())
+		{
+			TextValue = ConcatenateLine(TextValue, FText::Format(LOCTEXT("NumActiveCloths", "Active Cloths: {0}"), NumActiveCloths));
+		}
+		if (const int32 NumKinematicParticles = SimProxy->GetNumKinematicParticles())
+		{
+			TextValue = ConcatenateLine(TextValue, FText::Format(LOCTEXT("NumKinematicParticles", "Kinematic Particles: {0}"), NumKinematicParticles));
+		}
+		if (const int32 NumDynamicParticles = SimProxy->GetNumDynamicParticles())
+		{
+			TextValue = ConcatenateLine(TextValue, FText::Format(LOCTEXT("NumDynamicParticles", "Dynamic Particles: {0}"), NumDynamicParticles));
+		}
+		if (const int32 NumIterations = SimProxy->GetNumIterations())
+		{
+			TextValue = ConcatenateLine(TextValue, FText::Format(LOCTEXT("NumIterations", "Iterations: {0}"), NumIterations));
+		}
+		if (const float SimulationTime = SimProxy->GetSimulationTime())
+		{
+			FNumberFormattingOptions NumberFormatOptions;
+			NumberFormatOptions.AlwaysSign = false;
+			NumberFormatOptions.UseGrouping = false;
+			NumberFormatOptions.RoundingMode = ERoundingMode::HalfFromZero;
+			NumberFormatOptions.MinimumIntegralDigits = 1;
+			NumberFormatOptions.MaximumIntegralDigits = 6;
+			NumberFormatOptions.MinimumFractionalDigits = 2;
+			NumberFormatOptions.MaximumFractionalDigits = 2;
+			TextValue = ConcatenateLine(TextValue, FText::Format(LOCTEXT("SimulationTime", "Simulation Time: {0}ms"), FText::AsNumber(SimulationTime, &NumberFormatOptions)));
+		}
+		if (SimProxy->IsTeleported())
+		{
+			TextValue = ConcatenateLine(TextValue, LOCTEXT("IsTeleported", "Simulation Teleport Activated"));
+		}
+	}
+	return TextValue;
+}
+
 struct FVisualizationOption
 {
 	// Actual option entries
@@ -22,10 +73,12 @@ struct FVisualizationOption
 	static const uint32 Count;
 
 	// Chaos debug draw function
-	typedef void (::Chaos::FClothVisualization::* FDebugDrawFunction)(FPrimitiveDrawInterface*) const;
-	typedef void (::Chaos::FClothVisualization::* FDebugDrawTextsFunction)(FCanvas*, const FSceneView*) const;
-	FDebugDrawFunction DebugDrawFunction;
-	FDebugDrawTextsFunction DebugDrawTextsFunction;
+	typedef void (::Chaos::FClothVisualization::* FClothVisualizationDebugDrawFunction)(FPrimitiveDrawInterface*) const;
+	typedef void (::Chaos::FClothVisualization::* FClothVisualizationDebugDrawTextsFunction)(FCanvas*, const FSceneView*) const;
+	typedef FText (*FLocalDebugDisplayStringFunction)(const FClothSimulationProxy*);
+	FClothVisualizationDebugDrawFunction ClothVisualizationDebugDrawFunction;
+	FClothVisualizationDebugDrawTextsFunction ClothVisualizationDebugDrawTextsFunction;
+	FLocalDebugDisplayStringFunction LocalDebugDisplayStringFunction;
 
 	FText DisplayName;         // Text for menu entries.
 	FText ToolTip;             // Text for menu tooltips.
@@ -35,8 +88,8 @@ struct FVisualizationOption
 	// Console override
 	IConsoleVariable* const ConsoleVariable;
 
-	FVisualizationOption(FDebugDrawFunction InDebugDrawFunction, const TCHAR* ConsoleName, const FText& InDisplayName, const FText& InToolTip, bool bInDisablesSimulation = false, bool bInHidesClothSections = false)
-		: DebugDrawFunction(InDebugDrawFunction)
+	FVisualizationOption(FClothVisualizationDebugDrawFunction InDebugDrawFunction, const TCHAR* ConsoleName, const FText& InDisplayName, const FText& InToolTip, bool bInDisablesSimulation = false, bool bInHidesClothSections = false)
+		: ClothVisualizationDebugDrawFunction(InDebugDrawFunction)
 		, DisplayName(InDisplayName)
 		, ToolTip(InToolTip)
 		, bDisablesSimulation(bInDisablesSimulation)
@@ -44,8 +97,17 @@ struct FVisualizationOption
 		, ConsoleVariable(IConsoleManager::Get().RegisterConsoleVariable(ConsoleName, false, *InToolTip.ToString(), 0))
 	{}
 
-	FVisualizationOption(FDebugDrawTextsFunction InDebugDrawTextsFunction, const TCHAR* ConsoleName, const FText& InDisplayName, const FText& InToolTip, bool bInDisablesSimulation = false, bool bInHidesClothSections = false)
-		: DebugDrawTextsFunction(InDebugDrawTextsFunction)
+	FVisualizationOption(FLocalDebugDisplayStringFunction InLocalDebugDisplayStringFunction, const TCHAR* ConsoleName, const FText& InDisplayName, const FText& InToolTip, bool bInDisablesSimulation = false, bool bInHidesClothSections = false)
+		: LocalDebugDisplayStringFunction(InLocalDebugDisplayStringFunction)
+		, DisplayName(InDisplayName)
+		, ToolTip(InToolTip)
+		, bDisablesSimulation(bInDisablesSimulation)
+		, bHidesClothSections(bInHidesClothSections)
+		, ConsoleVariable(IConsoleManager::Get().RegisterConsoleVariable(ConsoleName, false, *InToolTip.ToString(), 0))
+	{}
+
+	FVisualizationOption(FClothVisualizationDebugDrawTextsFunction InDebugDrawTextsFunction, const TCHAR * ConsoleName, const FText & InDisplayName, const FText & InToolTip, bool bInDisablesSimulation = false, bool bInHidesClothSections = false)
+		: ClothVisualizationDebugDrawTextsFunction(InDebugDrawTextsFunction)
 		, DisplayName(InDisplayName)
 		, ToolTip(InToolTip)
 		, bDisablesSimulation(bInDisablesSimulation)
@@ -63,30 +125,33 @@ struct FVisualizationOption
 
 const FVisualizationOption FVisualizationOption::OptionData[] =
 {
-	FVisualizationOption(&::Chaos::FClothVisualization::DrawPhysMeshShaded      , TEXT("p.ChaosClothEditor.DebugDrawPhysMeshShaded"), LOCTEXT("ChaosVisName_PhysMesh"            , "Physical Mesh (Flat Shaded)"), LOCTEXT("ChaosVisName_PhysMeshShaded_ToolTip"      , "Draws the current physical result as a doubled sided flat shaded mesh"), /*bDisablesSimulation =*/false, /*bHidesClothSections=*/true),
-	FVisualizationOption(&::Chaos::FClothVisualization::DrawPhysMeshWired       , TEXT("p.ChaosClothEditor.DebugDrawPhysMeshWired"), LOCTEXT("ChaosVisName_PhysMeshWire"        , "Physical Mesh (Wireframe)"), LOCTEXT("ChaosVisName_PhysMeshWired_ToolTip"       , "Draws the current physical mesh result in wireframe")),
-	FVisualizationOption(&::Chaos::FClothVisualization::DrawAnimMeshWired       , TEXT("p.ChaosClothEditor.DebugDrawAnimMeshWired"), LOCTEXT("ChaosVisName_AnimMeshWire"        , "Animated Mesh (Wireframe)"), LOCTEXT("ChaosVisName_AnimMeshWired_ToolTip"       , "Draws the current animated mesh input in wireframe")),
-	FVisualizationOption(&::Chaos::FClothVisualization::DrawParticleIndices     , TEXT("p.ChaosClothEditor.DebugDrawParticleIndices"), LOCTEXT("ChaosVisName_ParticleIndices"     , "Particle Indices"), LOCTEXT("ChaosVisName_ParticleIndices_ToolTip"     , "Draws the particle indices as instantiated by the solver")),
-	FVisualizationOption(&::Chaos::FClothVisualization::DrawElementIndices      , TEXT("p.ChaosClothEditor.DebugDrawElementIndices"), LOCTEXT("ChaosVisName_ElementIndices"      , "Element Indices"), LOCTEXT("ChaosVisName_ElementIndices_ToolTip"      , "Draws the element's (triangle or other) indices as instantiated by the solver")),
-	FVisualizationOption(&::Chaos::FClothVisualization::DrawPointNormals        , TEXT("p.ChaosClothEditor.DebugDrawPointNormals"), LOCTEXT("ChaosVisName_PointNormals"        , "Physical Mesh Normals"), LOCTEXT("ChaosVisName_PointNormals_ToolTip"        , "Draws the current point normals for the simulation mesh")),
-	FVisualizationOption(&::Chaos::FClothVisualization::DrawPointVelocities     , TEXT("p.ChaosClothEditor.DebugDrawPointVelocities"), LOCTEXT("ChaosVisName_PointVelocities"     , "Point Velocities"), LOCTEXT("ChaosVisName_PointVelocities_ToolTip"     , "Draws the current point velocities for the simulation mesh")),
-	FVisualizationOption(&::Chaos::FClothVisualization::DrawAnimNormals         , TEXT("p.ChaosClothEditor.DebugDrawAnimNormals"), LOCTEXT("ChaosVisName_AnimNormals"         , "Animated Mesh Normals"), LOCTEXT("ChaosVisName_AnimNormals_ToolTip"         , "Draws the current point normals for the animated mesh")),
-	FVisualizationOption(&::Chaos::FClothVisualization::DrawCollision           , TEXT("p.ChaosClothEditor.DebugDrawCollision"), LOCTEXT("ChaosVisName_Collision"           , "Collisions"), LOCTEXT("ChaosVisName_Collision_ToolTip"           , "Draws the collision bodies the simulation is currently using")),
-	FVisualizationOption(&::Chaos::FClothVisualization::DrawBackstops           , TEXT("p.ChaosClothEditor.DebugDrawBackstops"), LOCTEXT("ChaosVisName_Backstop"            , "Backstops"), LOCTEXT("ChaosVisName_Backstop_ToolTip"            , "Draws the backstop radius and position for each simulation particle")),
-	FVisualizationOption(&::Chaos::FClothVisualization::DrawBackstopDistances   , TEXT("p.ChaosClothEditor.DebugDrawBackstopDistances"), LOCTEXT("ChaosVisName_BackstopDistance"    , "Backstop Distances"), LOCTEXT("ChaosVisName_BackstopDistance_ToolTip"    , "Draws the backstop distance offset for each simulation particle"), /*bDisablesSimulation =*/true),
-	FVisualizationOption(&::Chaos::FClothVisualization::DrawMaxDistances        , TEXT("p.ChaosClothEditor.DebugDrawMaxDistances"), LOCTEXT("ChaosVisName_MaxDistance"         , "Max Distances"), LOCTEXT("ChaosVisName_MaxDistance_ToolTip"         , "Draws the current max distances for the sim particles as a line along its normal"), /*bDisablesSimulation =*/true),
-	FVisualizationOption(&::Chaos::FClothVisualization::DrawMaxDistanceValues   , TEXT("p.ChaosClothEditor.DebugDrawMaxDistanceValues"), LOCTEXT("ChaosVisName_MaxDistanceValue"    , "Max Distances As Numbers"), LOCTEXT("ChaosVisName_MaxDistanceValue_ToolTip"    , "Draws the current max distances as numbers")),
-	FVisualizationOption(&::Chaos::FClothVisualization::DrawAnimDrive           , TEXT("p.ChaosClothEditor.DebugDrawAnimDrive"), LOCTEXT("ChaosVisName_AnimDrive"           , "Anim Drive"), LOCTEXT("ChaosVisName_AnimDrive_Tooltip"           , "Draws the current skinned reference mesh for the simulation which anim drive will attempt to reach if enabled")),
-	FVisualizationOption(&::Chaos::FClothVisualization::DrawEdgeConstraint      , TEXT("p.ChaosClothEditor.DebugDrawEdgeConstraint"), LOCTEXT("ChaosVisName_EdgeConstraint"      , "Edge Constraint"), LOCTEXT("ChaosVisName_EdgeConstraint_Tooltip"      , "Draws the edge spring constraints")),
-	FVisualizationOption(&::Chaos::FClothVisualization::DrawBendingConstraint   , TEXT("p.ChaosClothEditor.DebugDrawBendingConstraint"), LOCTEXT("ChaosVisName_BendingConstraint"   , "Bending Constraint"), LOCTEXT("ChaosVisName_BendingConstraint_Tooltip"   , "Draws the bending spring constraints")),
-	FVisualizationOption(&::Chaos::FClothVisualization::DrawLongRangeConstraint , TEXT("p.ChaosClothEditor.DebugDrawLongRangeConstraint"), LOCTEXT("ChaosVisName_LongRangeConstraint" , "Long Range Constraint"), LOCTEXT("ChaosVisName_LongRangeConstraint_Tooltip" , "Draws the long range attachment constraint distances")),
-	FVisualizationOption(&::Chaos::FClothVisualization::DrawWindAndPressureForces, TEXT("p.ChaosClothEditor.DebugDrawWindAndPressureForces"), LOCTEXT("ChaosVisName_WindAndPressureForces", "Wind Aerodynamic And Pressure Forces"), LOCTEXT("ChaosVisName_WindAndPressure_Tooltip", "Draws the Wind drag and lift and pressure forces")),
-	FVisualizationOption(&::Chaos::FClothVisualization::DrawLocalSpace          , TEXT("p.ChaosClothEditor.DebugDrawLocalSpace"), LOCTEXT("ChaosVisName_LocalSpace"          , "Local Space Reference Bone"), LOCTEXT("ChaosVisName_LocalSpace_Tooltip"          , "Draws the local space reference bone")),
-	FVisualizationOption(&::Chaos::FClothVisualization::DrawSelfCollision       , TEXT("p.ChaosClothEditor.DebugDrawSelfCollision"), LOCTEXT("ChaosVisName_SelfCollision"       , "Self Collision"), LOCTEXT("ChaosVisName_SelfCollision_Tooltip"       , "Draws the self collision thickness/debugging information")),
-	FVisualizationOption(&::Chaos::FClothVisualization::DrawSelfIntersection    , TEXT("p.ChaosClothEditor.DebugDrawSelfIntersection"), LOCTEXT("ChaosVisName_SelfIntersection"    , "Self Intersection"), LOCTEXT("ChaosVisName_SelfIntersection_Tooltip"    , "Draws the self intersection contour/region information")),
+	FVisualizationOption(&::Chaos::FClothVisualization::DrawPhysMeshShaded      , TEXT("p.ChaosClothAssetEditor.DebugDrawPhysMeshShaded"), LOCTEXT("ChaosVisName_PhysMesh"            , "Physical Mesh (Flat Shaded)"), LOCTEXT("ChaosVisName_PhysMeshShaded_ToolTip"      , "Draws the current physical result as a doubled sided flat shaded mesh"), /*bDisablesSimulation =*/false, /*bHidesClothSections=*/true),
+	FVisualizationOption(&::Chaos::FClothVisualization::DrawPhysMeshWired       , TEXT("p.ChaosClothAssetEditor.DebugDrawPhysMeshWired"), LOCTEXT("ChaosVisName_PhysMeshWire"        , "Physical Mesh (Wireframe)"), LOCTEXT("ChaosVisName_PhysMeshWired_ToolTip"       , "Draws the current physical mesh result in wireframe")),
+	FVisualizationOption(&::Chaos::FClothVisualization::DrawAnimMeshWired       , TEXT("p.ChaosClothAssetEditor.DebugDrawAnimMeshWired"), LOCTEXT("ChaosVisName_AnimMeshWire"        , "Animated Mesh (Wireframe)"), LOCTEXT("ChaosVisName_AnimMeshWired_ToolTip"       , "Draws the current animated mesh input in wireframe")),
+	FVisualizationOption(&::Chaos::FClothVisualization::DrawParticleIndices     , TEXT("p.ChaosClothAssetEditor.DebugDrawParticleIndices"), LOCTEXT("ChaosVisName_ParticleIndices"     , "Particle Indices"), LOCTEXT("ChaosVisName_ParticleIndices_ToolTip"     , "Draws the particle indices as instantiated by the solver")),
+	FVisualizationOption(&::Chaos::FClothVisualization::DrawElementIndices      , TEXT("p.ChaosClothAssetEditor.DebugDrawElementIndices"), LOCTEXT("ChaosVisName_ElementIndices"      , "Element Indices"), LOCTEXT("ChaosVisName_ElementIndices_ToolTip"      , "Draws the element's (triangle or other) indices as instantiated by the solver")),
+	FVisualizationOption(&::Chaos::FClothVisualization::DrawPointNormals        , TEXT("p.ChaosClothAssetEditor.DebugDrawPointNormals"), LOCTEXT("ChaosVisName_PointNormals"        , "Physical Mesh Normals"), LOCTEXT("ChaosVisName_PointNormals_ToolTip"        , "Draws the current point normals for the simulation mesh")),
+	FVisualizationOption(&::Chaos::FClothVisualization::DrawPointVelocities     , TEXT("p.ChaosClothAssetEditor.DebugDrawPointVelocities"), LOCTEXT("ChaosVisName_PointVelocities"     , "Point Velocities"), LOCTEXT("ChaosVisName_PointVelocities_ToolTip"     , "Draws the current point velocities for the simulation mesh")),
+	FVisualizationOption(&::Chaos::FClothVisualization::DrawAnimNormals         , TEXT("p.ChaosClothAssetEditor.DebugDrawAnimNormals"), LOCTEXT("ChaosVisName_AnimNormals"         , "Animated Mesh Normals"), LOCTEXT("ChaosVisName_AnimNormals_ToolTip"         , "Draws the current point normals for the animated mesh")),
+	FVisualizationOption(&::Chaos::FClothVisualization::DrawCollision           , TEXT("p.ChaosClothAssetEditor.DebugDrawCollision"), LOCTEXT("ChaosVisName_Collision"           , "Collisions"), LOCTEXT("ChaosVisName_Collision_ToolTip"           , "Draws the collision bodies the simulation is currently using")),
+	FVisualizationOption(&::Chaos::FClothVisualization::DrawBackstops           , TEXT("p.ChaosClothAssetEditor.DebugDrawBackstops"), LOCTEXT("ChaosVisName_Backstop"            , "Backstops"), LOCTEXT("ChaosVisName_Backstop_ToolTip"            , "Draws the backstop radius and position for each simulation particle")),
+	FVisualizationOption(&::Chaos::FClothVisualization::DrawBackstopDistances   , TEXT("p.ChaosClothAssetEditor.DebugDrawBackstopDistances"), LOCTEXT("ChaosVisName_BackstopDistance"    , "Backstop Distances"), LOCTEXT("ChaosVisName_BackstopDistance_ToolTip"    , "Draws the backstop distance offset for each simulation particle"), /*bDisablesSimulation =*/true),
+	FVisualizationOption(&::Chaos::FClothVisualization::DrawMaxDistances        , TEXT("p.ChaosClothAssetEditor.DebugDrawMaxDistances"), LOCTEXT("ChaosVisName_MaxDistance"         , "Max Distances"), LOCTEXT("ChaosVisName_MaxDistance_ToolTip"         , "Draws the current max distances for the sim particles as a line along its normal"), /*bDisablesSimulation =*/true),
+	FVisualizationOption(&::Chaos::FClothVisualization::DrawMaxDistanceValues   , TEXT("p.ChaosClothAssetEditor.DebugDrawMaxDistanceValues"), LOCTEXT("ChaosVisName_MaxDistanceValue"    , "Max Distances As Numbers"), LOCTEXT("ChaosVisName_MaxDistanceValue_ToolTip"    , "Draws the current max distances as numbers")),
+	FVisualizationOption(&::Chaos::FClothVisualization::DrawAnimDrive           , TEXT("p.ChaosClothAssetEditor.DebugDrawAnimDrive"), LOCTEXT("ChaosVisName_AnimDrive"           , "Anim Drive"), LOCTEXT("ChaosVisName_AnimDrive_Tooltip"           , "Draws the current skinned reference mesh for the simulation which anim drive will attempt to reach if enabled")),
+	FVisualizationOption(&::Chaos::FClothVisualization::DrawEdgeConstraint      , TEXT("p.ChaosClothAssetEditor.DebugDrawEdgeConstraint"), LOCTEXT("ChaosVisName_EdgeConstraint"      , "Edge Constraint"), LOCTEXT("ChaosVisName_EdgeConstraint_Tooltip"      , "Draws the edge spring constraints")),
+	FVisualizationOption(&::Chaos::FClothVisualization::DrawBendingConstraint   , TEXT("p.ChaosClothAssetEditor.DebugDrawBendingConstraint"), LOCTEXT("ChaosVisName_BendingConstraint"   , "Bending Constraint"), LOCTEXT("ChaosVisName_BendingConstraint_Tooltip"   , "Draws the bending spring constraints")),
+	FVisualizationOption(&::Chaos::FClothVisualization::DrawLongRangeConstraint , TEXT("p.ChaosClothAssetEditor.DebugDrawLongRangeConstraint"), LOCTEXT("ChaosVisName_LongRangeConstraint" , "Long Range Constraint"), LOCTEXT("ChaosVisName_LongRangeConstraint_Tooltip" , "Draws the long range attachment constraint distances")),
+	FVisualizationOption(&::Chaos::FClothVisualization::DrawWindAndPressureForces, TEXT("p.ChaosClothAssetEditor.DebugDrawWindAndPressureForces"), LOCTEXT("ChaosVisName_WindAndPressureForces", "Wind Aerodynamic And Pressure Forces"), LOCTEXT("ChaosVisName_WindAndPressure_Tooltip", "Draws the Wind drag and lift and pressure forces")),
+	FVisualizationOption(&::Chaos::FClothVisualization::DrawLocalSpace          , TEXT("p.ChaosClothAssetEditor.DebugDrawLocalSpace"), LOCTEXT("ChaosVisName_LocalSpace"          , "Local Space Reference Bone"), LOCTEXT("ChaosVisName_LocalSpace_Tooltip"          , "Draws the local space reference bone")),
+	FVisualizationOption(&::Chaos::FClothVisualization::DrawSelfCollision       , TEXT("p.ChaosClothAssetEditor.DebugDrawSelfCollision"), LOCTEXT("ChaosVisName_SelfCollision"       , "Self Collision"), LOCTEXT("ChaosVisName_SelfCollision_Tooltip"       , "Draws the self collision thickness/debugging information")),
+	FVisualizationOption(&::Chaos::FClothVisualization::DrawSelfIntersection    , TEXT("p.ChaosClothAssetEditor.DebugDrawSelfIntersection"), LOCTEXT("ChaosVisName_SelfIntersection"    , "Self Intersection"), LOCTEXT("ChaosVisName_SelfIntersection_Tooltip"    , "Draws the self intersection contour/region information")),
+	FVisualizationOption(&GetSimulationStatisticsString    , TEXT("p.ChaosClothAssetEditor.DebugDrawSimulationStatistics"), LOCTEXT("ChaosVisName_SimulationStatistics"    , "Simulation Statistics"), LOCTEXT("ChaosVisName_SimulationStatistics_Tooltip"    , "Displays simulation statistics")),
 };
 const uint32 FVisualizationOption::Count = sizeof(OptionData) / sizeof(FVisualizationOption);
 
+
+
 } // namespace Private
 
 FClothEditorSimulationVisualization::FClothEditorSimulationVisualization()
@@ -105,22 +170,22 @@ void FClothEditorSimulationVisualization::ExtendViewportShowMenu(FMenuBuilder& M
 			{
 				Flags[OptionIndex] = !Flags[OptionIndex];
 
-			if (UChaosClothComponent* const ClothComponent = ViewportClient->GetPreviewClothComponent())
-			{
-				// If we need to toggle the disabled or visibility states, handle it
-				// Disable simulation
-				const bool bShouldDisableSimulation = ShouldDisableSimulation();
-				if (ClothComponent->IsSimulationEnabled() == bShouldDisableSimulation)
-				{
-					ClothComponent->SetEnableSimulation(!bShouldDisableSimulation);
-				}
-				// Hide cloth section
-				if (Private::FVisualizationOption::OptionData[OptionIndex].bHidesClothSections)
+				if (UChaosClothComponent* const ClothComponent = ViewportClient->GetPreviewClothComponent())
 				{
-					const bool bIsClothSectionsVisible = !Flags[OptionIndex];
-					ShowClothSections(ClothComponent, bIsClothSectionsVisible);
+					// If we need to toggle the disabled or visibility states, handle it
+					// Disable simulation
+					const bool bShouldDisableSimulation = ShouldDisableSimulation();
+					if (ClothComponent->IsSimulationEnabled() == bShouldDisableSimulation)
+					{
+						ClothComponent->SetEnableSimulation(!bShouldDisableSimulation);
+					}
+					// Hide cloth section
+					if (Private::FVisualizationOption::OptionData[OptionIndex].bHidesClothSections)
+					{
+						const bool bIsClothSectionsVisible = !Flags[OptionIndex];
+						ShowClothSections(ClothComponent, bIsClothSectionsVisible);
+					}
 				}
-			}
 			});
 
 			// Checkstate function for visualization entries
@@ -147,10 +212,10 @@ void FClothEditorSimulationVisualization::DebugDrawSimulation(const UChaosClothC
 
 	for (int32 OptionIndex = 0; OptionIndex < Private::FVisualizationOption::Count; ++OptionIndex)
 	{
-		if (Private::FVisualizationOption::OptionData[OptionIndex].DebugDrawFunction &&
+		if (Private::FVisualizationOption::OptionData[OptionIndex].ClothVisualizationDebugDrawFunction &&
 			(Private::FVisualizationOption::OptionData[OptionIndex].IsConsoleVariableEnabled() || Flags[OptionIndex]))
 		{
-			(Visualization->*(Private::FVisualizationOption::OptionData[OptionIndex].DebugDrawFunction))(PDI);
+			(Visualization->*(Private::FVisualizationOption::OptionData[OptionIndex].ClothVisualizationDebugDrawFunction))(PDI);
 		}
 	}
 }
@@ -164,12 +229,32 @@ void FClothEditorSimulationVisualization::DebugDrawSimulationTexts(const UChaosC
 	}
 	for (int32 OptionIndex = 0; OptionIndex < Private::FVisualizationOption::Count; ++OptionIndex)
 	{
-		if (Private::FVisualizationOption::OptionData[OptionIndex].DebugDrawTextsFunction &&
+		if (Private::FVisualizationOption::OptionData[OptionIndex].ClothVisualizationDebugDrawTextsFunction &&
+			(Private::FVisualizationOption::OptionData[OptionIndex].IsConsoleVariableEnabled() || Flags[OptionIndex]))
+		{
+			(Visualization->*(Private::FVisualizationOption::OptionData[OptionIndex].ClothVisualizationDebugDrawTextsFunction))(Canvas, SceneView);
+		}
+	}
+}
+
+FText FClothEditorSimulationVisualization::GetDisplayString(const UChaosClothComponent* ClothComponent) const
+{
+	const FClothSimulationProxy* const SimProxy = ClothComponent ? ClothComponent->GetClothSimulationProxy() : nullptr;
+	if (!ensure(SimProxy))
+	{
+		return FText();
+	}
+
+	FText DisplayString;
+	for (int32 OptionIndex = 0; OptionIndex < Private::FVisualizationOption::Count; ++OptionIndex)
+	{
+		if (Private::FVisualizationOption::OptionData[OptionIndex].LocalDebugDisplayStringFunction &&
 			(Private::FVisualizationOption::OptionData[OptionIndex].IsConsoleVariableEnabled() || Flags[OptionIndex]))
 		{
-			(Visualization->*(Private::FVisualizationOption::OptionData[OptionIndex].DebugDrawTextsFunction))(Canvas, SceneView);
+			DisplayString = Private::ConcatenateLine(DisplayString, (Private::FVisualizationOption::OptionData[OptionIndex].LocalDebugDisplayStringFunction)(SimProxy));
 		}
 	}
+	return DisplayString;
 }
 
 bool FClothEditorSimulationVisualization::ShouldDisableSimulation() const
diff --git a/Engine/Plugins/Experimental/ChaosClothAssetEditor/Source/ChaosClothAssetEditor/Private/ChaosClothAsset/ClothEditorSimulationVisualization.h b/Engine/Plugins/Experimental/ChaosClothAssetEditor/Source/ChaosClothAssetEditor/Private/ChaosClothAsset/ClothEditorSimulationVisualization.h
--- a/Engine/Plugins/Experimental/ChaosClothAssetEditor/Source/ChaosClothAssetEditor/Private/ChaosClothAsset/ClothEditorSimulationVisualization.h
+++ b/Engine/Plugins/Experimental/ChaosClothAssetEditor/Source/ChaosClothAssetEditor/Private/ChaosClothAsset/ClothEditorSimulationVisualization.h
@@ -23,6 +23,7 @@ public:
 	void ExtendViewportShowMenu(FMenuBuilder& MenuBuilder, TSharedRef<FChaosClothAssetEditor3DViewportClient> ViewportClient);
 	void DebugDrawSimulation(const UChaosClothComponent* ClothComponent, FPrimitiveDrawInterface* PDI);
 	void DebugDrawSimulationTexts(const UChaosClothComponent* ClothComponent, FCanvas* Canvas, const FSceneView* SceneView);
+	FText GetDisplayString(const UChaosClothComponent* ClothComponent) const;
 
 private:
 	/** Return whether or not - given the current enabled options - the simulation should be disabled. */
diff --git a/Engine/Plugins/Experimental/ChaosClothAssetEditor/Source/ChaosClothAssetEditor/Private/ChaosClothAsset/SClothEditor3DViewportToolBar.cpp b/Engine/Plugins/Experimental/ChaosClothAssetEditor/Source/ChaosClothAssetEditor/Private/ChaosClothAsset/SClothEditor3DViewportToolBar.cpp
--- a/Engine/Plugins/Experimental/ChaosClothAssetEditor/Source/ChaosClothAssetEditor/Private/ChaosClothAsset/SClothEditor3DViewportToolBar.cpp
+++ b/Engine/Plugins/Experimental/ChaosClothAssetEditor/Source/ChaosClothAssetEditor/Private/ChaosClothAsset/SClothEditor3DViewportToolBar.cpp
@@ -11,6 +11,7 @@
 #include "Framework/MultiBox/MultiBoxBuilder.h"
 #include "Widgets/Layout/SBorder.h"
 #include "Widgets/SBoxPanel.h"
+#include "Widgets/Text/SRichTextBlock.h"
 #include "Editor/UnrealEd/Public/SEditorViewportToolBarMenu.h"
 #include "SEditorViewportViewMenu.h"
 
@@ -33,7 +34,21 @@ void SChaosClothAssetEditor3DViewportToolBar::Construct(const FArguments& InArgs
 		.BorderImage(FAppStyle::Get().GetBrush("EditorViewportToolBar.Background"))
 		.Cursor(EMouseCursor::Default)
 		[
-			SAssignNew( MainBoxPtr, SHorizontalBox )
+			SNew(SVerticalBox)
+			+SVerticalBox::Slot()
+			.AutoHeight()
+			[
+				SAssignNew( MainBoxPtr, SHorizontalBox )
+			]
+			+SVerticalBox::Slot()
+			.Padding(FMargin(4.0f, 3.0f, 0.0f, 0.0f))
+			[
+				// Display text (e.g., item being previewed)
+				SNew(SRichTextBlock)
+				.DecoratorStyleSet(&FAppStyle::Get())
+				.Text(this, &SChaosClothAssetEditor3DViewportToolBar::GetDisplayString)
+				.TextStyle(&FAppStyle::Get().GetWidgetStyle<FTextBlockStyle>("AnimViewport.MessageText"))
+			]
 		]
 	];
 
@@ -139,6 +154,17 @@ TSharedRef<SWidget> SChaosClothAssetEditor3DViewportToolBar::MakeToolBar(const T
 	return ToolbarBuilder.MakeWidget();
 }
 
+FText SChaosClothAssetEditor3DViewportToolBar::GetDisplayString() const
+{
+	using namespace UE::Chaos::ClothAsset;
+	TSharedRef<FChaosClothAssetEditor3DViewportClient> ViewportClient = StaticCastSharedPtr<FChaosClothAssetEditor3DViewportClient>(ChaosClothAssetEditor3DViewportPtr.Pin()->GetViewportClient()).ToSharedRef();
+	if (FClothEditorSimulationVisualization* const Visualization = ViewportClient->GetSimulationVisualization().Pin().Get())
+	{
+		return Visualization->GetDisplayString(ViewportClient->GetPreviewClothComponent());
+	}
+	return FText();
+}
+
 void SChaosClothAssetEditor3DViewportToolBar::ExtendOptionsMenu(FMenuBuilder& OptionsMenuBuilder) const
 {
 	using namespace UE::Chaos::ClothAsset;
diff --git a/Engine/Plugins/Experimental/ChaosClothAssetEditor/Source/ChaosClothAssetEditor/Private/ChaosClothAsset/SClothEditor3DViewportToolBar.h b/Engine/Plugins/Experimental/ChaosClothAssetEditor/Source/ChaosClothAssetEditor/Private/ChaosClothAsset/SClothEditor3DViewportToolBar.h
--- a/Engine/Plugins/Experimental/ChaosClothAssetEditor/Source/ChaosClothAssetEditor/Private/ChaosClothAsset/SClothEditor3DViewportToolBar.h
+++ b/Engine/Plugins/Experimental/ChaosClothAssetEditor/Source/ChaosClothAssetEditor/Private/ChaosClothAsset/SClothEditor3DViewportToolBar.h
@@ -25,6 +25,7 @@ private:
 	TSharedRef<SWidget> MakeOptionsMenu();
 	TSharedRef<SWidget> MakeDisplayToolBar(const TSharedPtr<FExtender> InExtenders);
 	TSharedRef<SWidget> MakeToolBar(const TSharedPtr<FExtender> InExtenders);
+	FText GetDisplayString() const;
 
 	virtual void ExtendOptionsMenu(FMenuBuilder& OptionsMenuBuilder) const override;
 
