ClothAssetEngine: remove ClothTrainingTool

diff --git a/Engine/Plugins/Experimental/ChaosClothAssetEditor/ChaosClothAssetEditor.uplugin b/Engine/Plugins/Experimental/ChaosClothAssetEditor/ChaosClothAssetEditor.uplugin
--- a/Engine/Plugins/Experimental/ChaosClothAssetEditor/ChaosClothAssetEditor.uplugin
+++ b/Engine/Plugins/Experimental/ChaosClothAssetEditor/ChaosClothAssetEditor.uplugin
@@ -76,10 +76,6 @@
 			"SupportedTargetPlatforms": [
 				"Win64"
 			]
-		},
-		{
-			"Name": "ChaosCaching",
-			"Enabled": true
 		}
 	]
 }
\ No newline at end of file
diff --git a/Engine/Plugins/Experimental/ChaosClothAssetEditor/Source/ChaosClothAssetEditor/Private/ChaosClothAsset/ClothEditorCommands.cpp b/Engine/Plugins/Experimental/ChaosClothAssetEditor/Source/ChaosClothAssetEditor/Private/ChaosClothAsset/ClothEditorCommands.cpp
--- a/Engine/Plugins/Experimental/ChaosClothAssetEditor/Source/ChaosClothAssetEditor/Private/ChaosClothAsset/ClothEditorCommands.cpp
+++ b/Engine/Plugins/Experimental/ChaosClothAssetEditor/Source/ChaosClothAssetEditor/Private/ChaosClothAsset/ClothEditorCommands.cpp
@@ -13,7 +13,6 @@ const FString FChaosClothAssetEditorCommands::BeginRemeshToolIdentifier = TEXT("
 const FString FChaosClothAssetEditorCommands::BeginAttributeEditorToolIdentifier = TEXT("BeginAttributeEditorTool");
 const FString FChaosClothAssetEditorCommands::BeginWeightMapPaintToolIdentifier = TEXT("BeginWeightMapPaintTool");
 const FString FChaosClothAssetEditorCommands::AddWeightMapNodeIdentifier = TEXT("AddWeightMapNode");
-const FString FChaosClothAssetEditorCommands::BeginClothTrainingToolIdentifier = TEXT("BeginClothTrainingTool");
 const FString FChaosClothAssetEditorCommands::BeginTransferSkinWeightsToolIdentifier = TEXT("BeginTransferSkinWeightsTool");
 const FString FChaosClothAssetEditorCommands::AddTransferSkinWeightsNodeIdentifier = TEXT("AddTransferSkinWeightsNode");
 const FString FChaosClothAssetEditorCommands::ToggleSimulationSuspendedIdentifier = TEXT("ToggleSimulationSuspended");
@@ -43,8 +42,6 @@ void FChaosClothAssetEditorCommands::RegisterCommands()
 	UI_COMMAND(BeginWeightMapPaintTool, "MapPnt", "Paint weight maps on the mesh", EUserInterfaceActionType::None, FInputChord());
 	UI_COMMAND(AddWeightMapNode, "MapPnt", "Paint weight maps on the mesh", EUserInterfaceActionType::Button, FInputChord());
 
-	UI_COMMAND(BeginClothTrainingTool, "Generate Train Data", "Generate cloth training data", EUserInterfaceActionType::Button, FInputChord());
-
 	UI_COMMAND(BeginTransferSkinWeightsTool, "TransferSkinWeights", "Transfer skinning weights from a SkeletalMesh", EUserInterfaceActionType::None, FInputChord());
 	UI_COMMAND(AddTransferSkinWeightsNode, "TransferSkinWeights", "Transfer skinning weights from a SkeletalMesh", EUserInterfaceActionType::Button, FInputChord());
 
diff --git a/Engine/Plugins/Experimental/ChaosClothAssetEditor/Source/ChaosClothAssetEditor/Private/ChaosClothAsset/ClothEditorCommands.h b/Engine/Plugins/Experimental/ChaosClothAssetEditor/Source/ChaosClothAssetEditor/Private/ChaosClothAsset/ClothEditorCommands.h
--- a/Engine/Plugins/Experimental/ChaosClothAssetEditor/Source/ChaosClothAssetEditor/Private/ChaosClothAsset/ClothEditorCommands.h
+++ b/Engine/Plugins/Experimental/ChaosClothAssetEditor/Source/ChaosClothAssetEditor/Private/ChaosClothAsset/ClothEditorCommands.h
@@ -41,9 +41,6 @@ public:
 	const static FString BeginAttributeEditorToolIdentifier;
 	TSharedPtr<FUICommandInfo> BeginAttributeEditorTool;
 
-	const static FString BeginClothTrainingToolIdentifier;
-	TSharedPtr<FUICommandInfo> BeginClothTrainingTool;
-
 	const static FString BeginTransferSkinWeightsToolIdentifier;
 	TSharedPtr<FUICommandInfo> BeginTransferSkinWeightsTool;
 	const static FString AddTransferSkinWeightsNodeIdentifier;
diff --git a/Engine/Plugins/Experimental/ChaosClothAssetEditor/Source/ChaosClothAssetEditor/Private/ChaosClothAsset/ClothEditorMode.cpp b/Engine/Plugins/Experimental/ChaosClothAssetEditor/Source/ChaosClothAssetEditor/Private/ChaosClothAsset/ClothEditorMode.cpp
--- a/Engine/Plugins/Experimental/ChaosClothAssetEditor/Source/ChaosClothAssetEditor/Private/ChaosClothAsset/ClothEditorMode.cpp
+++ b/Engine/Plugins/Experimental/ChaosClothAssetEditor/Source/ChaosClothAssetEditor/Private/ChaosClothAsset/ClothEditorMode.cpp
@@ -36,7 +36,6 @@
 #include "AttributeEditorTool.h"
 #include "MeshAttributePaintTool.h"
 #include "ChaosClothAsset/ClothWeightMapPaintTool.h"
-#include "ChaosClothAsset/ClothTrainingTool.h"
 #include "ChaosClothAsset/ClothTransferSkinWeightsTool.h"
 #include "DynamicMesh/DynamicMeshAttributeSet.h"
 #include "DynamicMesh/DynamicVertexSkinWeightsAttribute.h"
@@ -201,12 +200,6 @@ void UChaosClothAssetEditorMode::RegisterAddNodeCommand(TSharedPtr<FUICommandInf
 
 void UChaosClothAssetEditorMode::RegisterPreviewTools()
 {
-	using namespace UE::Chaos::ClothAsset;
-
-	UEditorInteractiveToolsContext* const PreviewSceneToolsContext = PreviewScene->GetClothPreviewEditorModeManager()->GetInteractiveToolsContext();
-
-	const FChaosClothAssetEditorCommands& CommandInfos = FChaosClothAssetEditorCommands::Get();
-	RegisterClothTool(CommandInfos.BeginClothTrainingTool, FChaosClothAssetEditorCommands::BeginClothTrainingToolIdentifier, NewObject<UClothTrainingToolBuilder>(), PreviewSceneToolsContext);
 }
 
 void UChaosClothAssetEditorMode::RegisterTools()
diff --git a/Engine/Plugins/Experimental/ChaosClothAssetEditor/Source/ChaosClothAssetEditor/Private/ChaosClothAsset/ClothEditorModeToolkit.cpp b/Engine/Plugins/Experimental/ChaosClothAssetEditor/Source/ChaosClothAssetEditor/Private/ChaosClothAsset/ClothEditorModeToolkit.cpp
--- a/Engine/Plugins/Experimental/ChaosClothAssetEditor/Source/ChaosClothAssetEditor/Private/ChaosClothAsset/ClothEditorModeToolkit.cpp
+++ b/Engine/Plugins/Experimental/ChaosClothAssetEditor/Source/ChaosClothAssetEditor/Private/ChaosClothAsset/ClothEditorModeToolkit.cpp
@@ -148,7 +148,6 @@ void FChaosClothAssetEditorModeToolkit::BuildToolPalette(FName PaletteIndex, cla
 	if (PaletteIndex == FBaseCharacterFXEditorModeToolkit::ToolsTabName)
 	{
 		ToolbarBuilder.AddToolBarButton(Commands.AddWeightMapNode);
-		ToolbarBuilder.AddToolBarButton(Commands.BeginClothTrainingTool);
 		ToolbarBuilder.AddToolBarButton(Commands.AddTransferSkinWeightsNode);
 	}
 }
diff --git a/Engine/Plugins/Experimental/ChaosClothAssetEditor/Source/ChaosClothAssetEditor/Private/ChaosClothAsset/ClothEditorStyle.cpp b/Engine/Plugins/Experimental/ChaosClothAssetEditor/Source/ChaosClothAssetEditor/Private/ChaosClothAsset/ClothEditorStyle.cpp
--- a/Engine/Plugins/Experimental/ChaosClothAssetEditor/Source/ChaosClothAssetEditor/Private/ChaosClothAsset/ClothEditorStyle.cpp
+++ b/Engine/Plugins/Experimental/ChaosClothAssetEditor/Source/ChaosClothAssetEditor/Private/ChaosClothAsset/ClothEditorStyle.cpp
@@ -51,10 +51,6 @@ FChaosClothAssetEditorStyle::FChaosClothAssetEditorStyle()
 	PropertyNameString = "ChaosClothAssetEditor." + FChaosClothAssetEditorCommands::AddWeightMapNodeIdentifier;
 	Set(*PropertyNameString, new FSlateImageBrush(FChaosClothAssetEditorStyle::InContent("Icons/ModelingAttributePaint_x40", ".png"), ToolbarIconSize));
 
-	// TODO: Get a permanent icon to replace this placeholder
-	PropertyNameString = "ChaosClothAssetEditor." + FChaosClothAssetEditorCommands::BeginClothTrainingToolIdentifier;
-	Set(*PropertyNameString, new FSlateImageBrush(FChaosClothAssetEditorStyle::InContent("Icons/ModelingAttributePaint_x40", ".png"), ToolbarIconSize));
-
 	// TODO: Get a permanent icon to replace this placeholder
 	PropertyNameString = "ChaosClothAssetEditor." + FChaosClothAssetEditorCommands::AddTransferSkinWeightsNodeIdentifier;
 	Set(*PropertyNameString, new FSlateImageBrush(FChaosClothAssetEditorStyle::InContent("Icons/ModelingAttributePaint_x40", ".png"), ToolbarIconSize));
diff --git a/Engine/Plugins/Experimental/ChaosClothAssetEditor/Source/ChaosClothAssetEditorTools/ChaosClothAssetEditorTools.Build.cs b/Engine/Plugins/Experimental/ChaosClothAssetEditor/Source/ChaosClothAssetEditorTools/ChaosClothAssetEditorTools.Build.cs
--- a/Engine/Plugins/Experimental/ChaosClothAssetEditor/Source/ChaosClothAssetEditorTools/ChaosClothAssetEditorTools.Build.cs
+++ b/Engine/Plugins/Experimental/ChaosClothAssetEditor/Source/ChaosClothAssetEditorTools/ChaosClothAssetEditorTools.Build.cs
@@ -37,7 +37,6 @@ public class ChaosClothAssetEditorTools : ModuleRules
 				"MeshModelingToolsExp",
 				"EditorStyle",
 				"EditorSubsystem",
-				"GeometryCache",
 				"GeometryCore",
 				"GeometryFramework",
 				"LevelEditor",
diff --git a/Engine/Plugins/Experimental/ChaosClothAssetEditor/Source/ChaosClothAssetEditorTools/Private/ChaosClothAsset/ClothDataGenerationComponent.cpp b/Engine/Plugins/Experimental/ChaosClothAssetEditor/Source/ChaosClothAssetEditorTools/Private/ChaosClothAsset/ClothDataGenerationComponent.cpp
deleted file mode 100644
diff --git a/Engine/Plugins/Experimental/ChaosClothAssetEditor/Source/ChaosClothAssetEditorTools/Private/ChaosClothAsset/ClothDataGenerationComponent.h b/Engine/Plugins/Experimental/ChaosClothAssetEditor/Source/ChaosClothAssetEditorTools/Private/ChaosClothAsset/ClothDataGenerationComponent.h
deleted file mode 100644
diff --git a/Engine/Plugins/Experimental/ChaosClothAssetEditor/Source/ChaosClothAssetEditorTools/Private/ChaosClothAsset/ClothTrainingTool.cpp b/Engine/Plugins/Experimental/ChaosClothAssetEditor/Source/ChaosClothAssetEditorTools/Private/ChaosClothAsset/ClothTrainingTool.cpp
deleted file mode 100644
diff --git a/Engine/Plugins/Experimental/ChaosClothAssetEditor/Source/ChaosClothAssetEditorTools/Private/ChaosClothAsset/ClothTrainingToolProperties.cpp b/Engine/Plugins/Experimental/ChaosClothAssetEditor/Source/ChaosClothAssetEditorTools/Private/ChaosClothAsset/ClothTrainingToolProperties.cpp
deleted file mode 100644
diff --git a/Engine/Plugins/Experimental/ChaosClothAssetEditor/Source/ChaosClothAssetEditorTools/Private/ChaosClothAsset/ClothTrainingToolProperties.h b/Engine/Plugins/Experimental/ChaosClothAssetEditor/Source/ChaosClothAssetEditorTools/Private/ChaosClothAsset/ClothTrainingToolProperties.h
deleted file mode 100644
diff --git a/Engine/Plugins/Experimental/ChaosClothAssetEditor/Source/ChaosClothAssetEditorTools/Public/ChaosClothAsset/ClothTrainingTool.h b/Engine/Plugins/Experimental/ChaosClothAssetEditor/Source/ChaosClothAssetEditorTools/Public/ChaosClothAsset/ClothTrainingTool.h
deleted file mode 100644
