Chaos Cloth:
- Updated all nodes to prevent them from crashing by running on empty or unformatted managed array collections.
- Added nodes missing descriptions.
- Added nodes parameters missing ducumentation.
- Formated all nodes' error messages to use the same layout.
- Tidied up and remove unnecessary/unused inputs.
- Tidied up utility files and includes, by merging DataFlowNodes.cpp/h into ClothDataflowTools.cpp/h and moving the nodes registration to ChaosClothAssetDataflowNodesModule.cpp.

#rnx

diff --git a/Engine/Plugins/Experimental/ChaosClothAssetEditor/Source/ChaosClothAssetDataflowNodes/Private/ChaosClothAsset/AddWeightMapNode.cpp b/Engine/Plugins/Experimental/ChaosClothAssetEditor/Source/ChaosClothAssetDataflowNodes/Private/ChaosClothAsset/AddWeightMapNode.cpp
--- a/Engine/Plugins/Experimental/ChaosClothAssetEditor/Source/ChaosClothAssetDataflowNodes/Private/ChaosClothAsset/AddWeightMapNode.cpp
+++ b/Engine/Plugins/Experimental/ChaosClothAssetEditor/Source/ChaosClothAssetDataflowNodes/Private/ChaosClothAsset/AddWeightMapNode.cpp
@@ -1,9 +1,10 @@
 // Copyright Epic Games, Inc. All Rights Reserved.
 
 #include "ChaosClothAsset/AddWeightMapNode.h"
-#include "ChaosClothAsset/DataflowNodes.h"
 #include "ChaosClothAsset/CollectionClothFacade.h"
+#include "ChaosClothAsset/ClothDataflowTools.h"
 #include "ChaosClothAsset/ClothGeometryTools.h"
+#include "Dataflow/DataflowInputOutput.h"
 
 #include UE_INLINE_GENERATED_CPP_BY_NAME(AddWeightMapNode)
 
@@ -14,8 +15,7 @@ FChaosClothAssetAddWeightMapNode::FChaosClothAssetAddWeightMapNode(const Dataflo
 {
 	RegisterInputConnection(&Collection);
 	RegisterOutputConnection(&Collection, &Collection);
-	RegisterInputConnection(&Name);
-	RegisterOutputConnection(&Name, &Name);
+	RegisterOutputConnection(&Name);
 }
 
 void FChaosClothAssetAddWeightMapNode::Evaluate(Dataflow::FContext& Context, const FDataflowOutput* Out) const
@@ -25,22 +25,21 @@ void FChaosClothAssetAddWeightMapNode::Evaluate(Dataflow::FContext& Context, con
 	if (Out->IsA<FManagedArrayCollection>(&Collection))
 	{
 		// Evaluate in collection
-		FString InNameString = GetValue<FString>(Context, &Name);
-
 		FManagedArrayCollection InCollection = GetValue<FManagedArrayCollection>(Context, &Collection);
 		const TSharedRef<FManagedArrayCollection> ClothCollection = MakeShared<FManagedArrayCollection>(MoveTemp(InCollection));
 		FCollectionClothFacade ClothFacade(ClothCollection);
 		if (ClothFacade.IsValid())  // Can only act on the collection if it is a valid cloth collection
 		{
-			const FName InName(InNameString);
+			const FName InName(Name);
 			ClothFacade.AddWeightMap(InName);		// Does nothing if weight map already exists
 
 			TArrayView<float> ClothWeights = ClothFacade.GetWeightMap(InName);
 			const int32 MaxWeightIndex = FMath::Min(VertexWeights.Num(), ClothWeights.Num());
 			if (VertexWeights.Num() > 0 && VertexWeights.Num() != ClothWeights.Num())
 			{
-				DataflowNodes::LogAndToastWarning(
-					FText::Format(LOCTEXT("WeightMapSize", "FChaosClothAssetAddWeightMapNode: Vertex count mismatch: vertex weights in the node: {0}; 3D vertices in cloth: {1}"),
+				FClothDataflowTools::LogAndToastWarning(*this,
+					LOCTEXT("VertexCountMismatchHeadline", "Vertex count mismatch."),
+					FText::Format(LOCTEXT("VertexCountMismatchDetails", "Vertex weights in the node: {0}\n3D vertices in the cloth: {1}"),
 						VertexWeights.Num(),
 						ClothWeights.Num()));
 			}
@@ -51,7 +50,7 @@ void FChaosClothAssetAddWeightMapNode::Evaluate(Dataflow::FContext& Context, con
 			}
 		}
 		SetValue(Context, MoveTemp(*ClothCollection), &Collection);
-		SetValue(Context, MoveTemp(InNameString), &Name);
+		SetValue(Context, Name, &Name);
 	}
 }
 
diff --git a/Engine/Plugins/Experimental/ChaosClothAssetEditor/Source/ChaosClothAssetDataflowNodes/Private/ChaosClothAsset/BindToRootBoneNode.cpp b/Engine/Plugins/Experimental/ChaosClothAssetEditor/Source/ChaosClothAssetDataflowNodes/Private/ChaosClothAsset/BindToRootBoneNode.cpp
--- a/Engine/Plugins/Experimental/ChaosClothAssetEditor/Source/ChaosClothAssetDataflowNodes/Private/ChaosClothAsset/BindToRootBoneNode.cpp
+++ b/Engine/Plugins/Experimental/ChaosClothAssetEditor/Source/ChaosClothAssetDataflowNodes/Private/ChaosClothAsset/BindToRootBoneNode.cpp
@@ -1,8 +1,9 @@
 // Copyright Epic Games, Inc. All Rights Reserved.
 
 #include "ChaosClothAsset/BindToRootBoneNode.h"
-#include "ChaosClothAsset/DataflowNodes.h"
 #include "ChaosClothAsset/ClothGeometryTools.h"
+#include "ChaosClothAsset/CollectionClothFacade.h"
+#include "Dataflow/DataflowInputOutput.h"
 
 #include UE_INLINE_GENERATED_CPP_BY_NAME(BindToRootBoneNode)
 
@@ -25,7 +26,10 @@ void FChaosClothAssetBindToRootBoneNode::Evaluate(Dataflow::FContext& Context, c
 		FManagedArrayCollection InCollection = GetValue<FManagedArrayCollection>(Context, &Collection);
 		const TSharedRef<FManagedArrayCollection> ClothCollection = MakeShared<FManagedArrayCollection>(MoveTemp(InCollection));
 
-		FClothGeometryTools::BindMeshToRootBone(ClothCollection, bBindSimMesh, bBindRenderMesh);
+		if (FCollectionClothFacade(ClothCollection).IsValid())  // Can only act on the collection if it is a valid cloth collection
+		{
+			FClothGeometryTools::BindMeshToRootBone(ClothCollection, bBindSimMesh, bBindRenderMesh);
+		}
 
 		SetValue(Context, MoveTemp(*ClothCollection), &Collection);
 	}
diff --git a/Engine/Plugins/Experimental/ChaosClothAssetEditor/Source/ChaosClothAssetDataflowNodes/Private/ChaosClothAsset/ChaosClothAssetDataflowNodesModule.cpp b/Engine/Plugins/Experimental/ChaosClothAssetEditor/Source/ChaosClothAssetDataflowNodes/Private/ChaosClothAsset/ChaosClothAssetDataflowNodesModule.cpp
--- a/Engine/Plugins/Experimental/ChaosClothAssetEditor/Source/ChaosClothAssetDataflowNodes/Private/ChaosClothAsset/ChaosClothAssetDataflowNodesModule.cpp
+++ b/Engine/Plugins/Experimental/ChaosClothAssetEditor/Source/ChaosClothAssetDataflowNodes/Private/ChaosClothAsset/ChaosClothAssetDataflowNodesModule.cpp
@@ -1,17 +1,98 @@
 // Copyright Epic Games, Inc. All Rights Reserved.
 
 #include "ChaosClothAsset/ChaosClothAssetDataflowNodesModule.h"
-#include "ChaosClothAsset/DataflowNodes.h"
+#include "ChaosClothAsset/AddWeightMapNode.h"
+#include "ChaosClothAsset/BindToRootBoneNode.h"
+#include "ChaosClothAsset/CopySimulationToRenderMeshNode.h"
+#include "ChaosClothAsset/DatasmithImportNode.h"
+#include "ChaosClothAsset/DeleteRenderMeshNode.h"
+#include "ChaosClothAsset/ImportNode.h"
+#include "ChaosClothAsset/MergeClothCollectionsNode.h"
+#include "ChaosClothAsset/ReverseNormalsNode.h"
+#include "ChaosClothAsset/SetPhysicsAssetNode.h"
+#include "ChaosClothAsset/SimulationAerodynamicsConfigNode.h"
+#include "ChaosClothAsset/SimulationAnimDriveConfigNode.h"
+#include "ChaosClothAsset/SimulationBackstopConfigNode.h"
+#include "ChaosClothAsset/SimulationCollisionConfigNode.h"
+#include "ChaosClothAsset/SimulationDampingConfigNode.h"
+#include "ChaosClothAsset/SimulationDefaultConfigNode.h"
+#include "ChaosClothAsset/SimulationGravityConfigNode.h"
+#include "ChaosClothAsset/SimulationLongRangeAttachmentConfigNode.h"
+#include "ChaosClothAsset/SimulationMassConfigNode.h"
+#include "ChaosClothAsset/SimulationMaxDistanceConfigNode.h"
+#include "ChaosClothAsset/SimulationPBDAreaSpringConfigNode.h"
+#include "ChaosClothAsset/SimulationPBDBendingElementConfigNode.h"
+#include "ChaosClothAsset/SimulationPBDBendingSpringConfigNode.h"
+#include "ChaosClothAsset/SimulationPBDEdgeSpringConfigNode.h"
+#include "ChaosClothAsset/SimulationPressureConfigNode.h"
+#include "ChaosClothAsset/SimulationSelfCollisionConfigNode.h"
+#include "ChaosClothAsset/SimulationSolverConfigNode.h"
+#include "ChaosClothAsset/SimulationVelocityScaleConfigNode.h"
+#include "ChaosClothAsset/SimulationXPBDAreaSpringConfigNode.h"
+#include "ChaosClothAsset/SimulationXPBDAnisoBendingConfigNode.h"
+#include "ChaosClothAsset/SimulationXPBDAnisoStretchConfigNode.h"
+#include "ChaosClothAsset/SimulationXPBDBendingElementConfigNode.h"
+#include "ChaosClothAsset/SimulationXPBDBendingSpringConfigNode.h"
+#include "ChaosClothAsset/SimulationXPBDEdgeSpringConfigNode.h"
+#include "ChaosClothAsset/SkeletalMeshImportNode.h"
+#include "ChaosClothAsset/StaticMeshImportNode.h"
+#include "ChaosClothAsset/TerminalNode.h"
+#include "ChaosClothAsset/TransferSkinWeightsNode.h"
 #include "ChaosClothAsset/WeightedValueCustomization.h"
+#include "Dataflow/DataflowNodeFactory.h"
 #include "Modules/ModuleManager.h"
 
 #define LOCTEXT_NAMESPACE "FChaosClothAssetDataflowNodesModule"
 
+namespace UE::Chaos::ClothAsset::Private
+{
+	static void RegisterDataflowNodes()
+	{
+		DATAFLOW_NODE_REGISTER_CREATION_FACTORY(FChaosClothAssetAddWeightMapNode);
+		DATAFLOW_NODE_REGISTER_CREATION_FACTORY(FChaosClothAssetBindToRootBoneNode);
+		DATAFLOW_NODE_REGISTER_CREATION_FACTORY(FChaosClothAssetCopySimulationToRenderMeshNode);
+		DATAFLOW_NODE_REGISTER_CREATION_FACTORY(FChaosClothAssetDatasmithImportNode);
+		DATAFLOW_NODE_REGISTER_CREATION_FACTORY(FChaosClothAssetDeleteRenderMeshNode);
+		DATAFLOW_NODE_REGISTER_CREATION_FACTORY(FChaosClothAssetImportNode);
+		DATAFLOW_NODE_REGISTER_CREATION_FACTORY(FChaosClothAssetMergeClothCollectionsNode);
+		DATAFLOW_NODE_REGISTER_CREATION_FACTORY(FChaosClothAssetReverseNormalsNode);
+		DATAFLOW_NODE_REGISTER_CREATION_FACTORY(FChaosClothAssetSetPhysicsAssetNode);
+		DATAFLOW_NODE_REGISTER_CREATION_FACTORY(FChaosClothAssetSimulationAerodynamicsConfigNode);
+		DATAFLOW_NODE_REGISTER_CREATION_FACTORY(FChaosClothAssetSimulationAnimDriveConfigNode);
+		DATAFLOW_NODE_REGISTER_CREATION_FACTORY(FChaosClothAssetSimulationBackstopConfigNode);
+		DATAFLOW_NODE_REGISTER_CREATION_FACTORY(FChaosClothAssetSimulationCollisionConfigNode);
+		DATAFLOW_NODE_REGISTER_CREATION_FACTORY(FChaosClothAssetSimulationDampingConfigNode);
+		DATAFLOW_NODE_REGISTER_CREATION_FACTORY(FChaosClothAssetSimulationDefaultConfigNode);
+		DATAFLOW_NODE_REGISTER_CREATION_FACTORY(FChaosClothAssetSimulationGravityConfigNode);
+		DATAFLOW_NODE_REGISTER_CREATION_FACTORY(FChaosClothAssetSimulationLongRangeAttachmentConfigNode);
+		DATAFLOW_NODE_REGISTER_CREATION_FACTORY(FChaosClothAssetSimulationMassConfigNode);
+		DATAFLOW_NODE_REGISTER_CREATION_FACTORY(FChaosClothAssetSimulationMaxDistanceConfigNode);
+		DATAFLOW_NODE_REGISTER_CREATION_FACTORY(FChaosClothAssetSimulationPBDAreaSpringConfigNode);
+		DATAFLOW_NODE_REGISTER_CREATION_FACTORY(FChaosClothAssetSimulationPBDBendingElementConfigNode);
+		DATAFLOW_NODE_REGISTER_CREATION_FACTORY(FChaosClothAssetSimulationPBDBendingSpringConfigNode);
+		DATAFLOW_NODE_REGISTER_CREATION_FACTORY(FChaosClothAssetSimulationPBDEdgeSpringConfigNode);
+		DATAFLOW_NODE_REGISTER_CREATION_FACTORY(FChaosClothAssetSimulationPressureConfigNode);
+		DATAFLOW_NODE_REGISTER_CREATION_FACTORY(FChaosClothAssetSimulationSelfCollisionConfigNode);
+		DATAFLOW_NODE_REGISTER_CREATION_FACTORY(FChaosClothAssetSimulationSolverConfigNode);
+		DATAFLOW_NODE_REGISTER_CREATION_FACTORY(FChaosClothAssetSimulationVelocityScaleConfigNode);
+		DATAFLOW_NODE_REGISTER_CREATION_FACTORY(FChaosClothAssetSimulationXPBDAnisoBendingConfigNode);
+		DATAFLOW_NODE_REGISTER_CREATION_FACTORY(FChaosClothAssetSimulationXPBDAnisoStretchConfigNode);
+		DATAFLOW_NODE_REGISTER_CREATION_FACTORY(FChaosClothAssetSimulationXPBDAreaSpringConfigNode);
+		DATAFLOW_NODE_REGISTER_CREATION_FACTORY(FChaosClothAssetSimulationXPBDBendingElementConfigNode);
+		DATAFLOW_NODE_REGISTER_CREATION_FACTORY(FChaosClothAssetSimulationXPBDBendingSpringConfigNode);
+		DATAFLOW_NODE_REGISTER_CREATION_FACTORY(FChaosClothAssetSimulationXPBDEdgeSpringConfigNode);
+		DATAFLOW_NODE_REGISTER_CREATION_FACTORY(FChaosClothAssetSkeletalMeshImportNode);
+		DATAFLOW_NODE_REGISTER_CREATION_FACTORY(FChaosClothAssetStaticMeshImportNode);
+		DATAFLOW_NODE_REGISTER_CREATION_FACTORY(FChaosClothAssetTerminalNode);
+		DATAFLOW_NODE_REGISTER_CREATION_FACTORY(FChaosClothAssetTransferSkinWeightsNode);
+	}
+}
+
 void FChaosClothAssetDataflowNodesModule::StartupModule()
 {
 	using namespace UE::Chaos::ClothAsset;
 
-	DataflowNodes::Register();
+	Private::RegisterDataflowNodes();
 
 	// Register type customizations
 	if (FPropertyEditorModule* const PropertyModule = FModuleManager::GetModulePtr<FPropertyEditorModule>("PropertyEditor"))
diff --git a/Engine/Plugins/Experimental/ChaosClothAssetEditor/Source/ChaosClothAssetDataflowNodes/Private/ChaosClothAsset/ClothDataflowTools.cpp b/Engine/Plugins/Experimental/ChaosClothAssetEditor/Source/ChaosClothAssetDataflowNodes/Private/ChaosClothAsset/ClothDataflowTools.cpp
--- a/Engine/Plugins/Experimental/ChaosClothAssetEditor/Source/ChaosClothAssetDataflowNodes/Private/ChaosClothAsset/ClothDataflowTools.cpp
+++ b/Engine/Plugins/Experimental/ChaosClothAssetEditor/Source/ChaosClothAssetDataflowNodes/Private/ChaosClothAsset/ClothDataflowTools.cpp
@@ -4,13 +4,17 @@
 #include "ChaosClothAsset/ClothGeometryTools.h"
 #include "ChaosClothAsset/CollectionClothFacade.h"
 #include "Animation/Skeleton.h"
-#include "BoneWeights.h"
+#include "Dataflow/DataflowNode.h"
 #include "DynamicMesh/NonManifoldMappingSupport.h"
+#include "Framework/Notifications/NotificationManager.h"
 #include "Math/UnrealMathUtility.h"
 #include "Rendering/SkeletalMeshLODModel.h"
+#include "Widgets/Notifications/SNotificationList.h"
 #include "SkeletalMeshAttributes.h"
 #include "ToDynamicMesh.h"
 
+DEFINE_LOG_CATEGORY(LogChaosClothAssetDataflowNodes);
+
 namespace UE::Chaos::ClothAsset
 {
 	namespace Private
@@ -433,4 +437,23 @@ namespace UE::Chaos::ClothAsset
 		constexpr bool bAppend = true;
 		FClothGeometryTools::BuildSimMeshFromDynamicMesh(ClothCollection, DynamicMesh, UVChannelIndex, FVector2f(1.f), bAppend);
 	}
+
+	void FClothDataflowTools::LogAndToastWarning(const FDataflowNode& DataflowNode, const FText& Headline, const FText& Details)
+	{
+		static const FTextFormat TextFormat = FTextFormat::FromString(TEXT("{0}: {1}\n{2}"));
+		const FText NodeName = FText::FromName(DataflowNode.GetName());
+		const FText Text = FText::Format(TextFormat, NodeName, Headline, Details);
+
+		FNotificationInfo NotificationInfo(Text);
+		NotificationInfo.ExpireDuration = 5.0f;
+		FSlateNotificationManager::Get().AddNotification(NotificationInfo);
+
+		UE_LOG(LogChaosClothAssetDataflowNodes, Display, TEXT("%s"), *Text.ToString());
+	}
+
+	FString FClothDataflowTools::MakeAttributeName(const FString& String)
+	{
+		return SlugStringForValidName(String, TEXT("_")).Replace(TEXT("\\"), TEXT("_"));
+	}
+
 }  // End namespace UE::Chaos::ClothAsset
diff --git a/Engine/Plugins/Experimental/ChaosClothAssetEditor/Source/ChaosClothAssetDataflowNodes/Private/ChaosClothAsset/ClothDataflowTools.h b/Engine/Plugins/Experimental/ChaosClothAssetEditor/Source/ChaosClothAssetDataflowNodes/Private/ChaosClothAsset/ClothDataflowTools.h
--- a/Engine/Plugins/Experimental/ChaosClothAssetEditor/Source/ChaosClothAssetDataflowNodes/Private/ChaosClothAsset/ClothDataflowTools.h
+++ b/Engine/Plugins/Experimental/ChaosClothAssetEditor/Source/ChaosClothAssetDataflowNodes/Private/ChaosClothAsset/ClothDataflowTools.h
@@ -3,20 +3,28 @@
 #pragma once
 
 #include "Templates/SharedPointer.h"
+#include "Logging/LogMacros.h"
 
 struct FManagedArrayCollection;
+struct FDataflowNode;
 class FSkeletalMeshLODModel;
 class FString;
 
 namespace UE::Chaos::ClothAsset
 {
 	/**
-	* Tools shared by cloth data flow nodes
+	* Tools shared by cloth dataflow nodes
 	*/
 	struct FClothDataflowTools
 	{
 		static void AddRenderPatternFromSkeletalMeshSection(const TSharedRef<FManagedArrayCollection>& ClothCollection, const FSkeletalMeshLODModel& SkeletalMeshModel, const int32 SectionIndex, const FString& RenderMaterialPathName);
 
 		static void AddSimPatternsFromSkeletalMeshSection(const TSharedRef<FManagedArrayCollection>& ClothCollection, const FSkeletalMeshLODModel& SkeletalMeshModel, const int32 SectionIndex, const int32 UVChannelIndex);
+
+		static void LogAndToastWarning(const FDataflowNode& DataflowNode, const FText& Headline, const FText& Details);
+
+		static FString MakeAttributeName(const FString& String);
 	};
 }  // End namespace UE::Chaos::ClothAsset
+
+DECLARE_LOG_CATEGORY_EXTERN(LogChaosClothAssetDataflowNodes, Log, All);
diff --git a/Engine/Plugins/Experimental/ChaosClothAssetEditor/Source/ChaosClothAssetDataflowNodes/Private/ChaosClothAsset/CopySimulationToRenderMeshNode.cpp b/Engine/Plugins/Experimental/ChaosClothAssetEditor/Source/ChaosClothAssetDataflowNodes/Private/ChaosClothAsset/CopySimulationToRenderMeshNode.cpp
--- a/Engine/Plugins/Experimental/ChaosClothAssetEditor/Source/ChaosClothAssetDataflowNodes/Private/ChaosClothAsset/CopySimulationToRenderMeshNode.cpp
+++ b/Engine/Plugins/Experimental/ChaosClothAssetEditor/Source/ChaosClothAssetDataflowNodes/Private/ChaosClothAsset/CopySimulationToRenderMeshNode.cpp
@@ -1,8 +1,9 @@
 // Copyright Epic Games, Inc. All Rights Reserved.
 
 #include "ChaosClothAsset/CopySimulationToRenderMeshNode.h"
-#include "ChaosClothAsset/DataflowNodes.h"
 #include "ChaosClothAsset/ClothGeometryTools.h"
+#include "ChaosClothAsset/CollectionClothFacade.h"
+#include "Dataflow/DataflowInputOutput.h"
 #include "Materials/Material.h"
 
 #include UE_INLINE_GENERATED_CPP_BY_NAME(CopySimulationToRenderMeshNode)
@@ -13,7 +14,6 @@ FChaosClothAssetCopySimulationToRenderMeshNode::FChaosClothAssetCopySimulationTo
 	: FDataflowNode(InParam, InGuid)
 {
 	RegisterInputConnection(&Collection);
-	RegisterInputConnection(&Patterns);
 	RegisterOutputConnection(&Collection, &Collection);
 }
 
@@ -27,16 +27,19 @@ void FChaosClothAssetCopySimulationToRenderMeshNode::Evaluate(Dataflow::FContext
 		FManagedArrayCollection InCollection = GetValue<FManagedArrayCollection>(Context, &Collection);
 		const TSharedRef<FManagedArrayCollection> ClothCollection = MakeShared<FManagedArrayCollection>(MoveTemp(InCollection));
 
-		// Empty mesh and materials
-		FClothGeometryTools::DeleteRenderMesh(ClothCollection);
+		if (FCollectionClothFacade(ClothCollection).IsValid())  // Can only act on the collection if it is a valid cloth collection
+		{
+			// Empty mesh and materials
+			FClothGeometryTools::DeleteRenderMesh(ClothCollection);
 
-		// Find the material asset path
-		const FString MaterialPathName = Material ? 
-			Material->GetPathName() :
-			FString(TEXT("/Engine/EditorMaterials/Cloth/CameraLitDoubleSided.CameraLitDoubleSided"));
+			// Find the material asset path
+			const FString MaterialPathName = Material ? 
+				Material->GetPathName() :
+				FString(TEXT("/Engine/EditorMaterials/Cloth/CameraLitDoubleSided.CameraLitDoubleSided"));
 
-		// Copy the mesh data
-		FClothGeometryTools::CopySimMeshToRenderMesh(ClothCollection, MaterialPathName, bGenerateSingleRenderPattern);
+			// Copy the mesh data
+			FClothGeometryTools::CopySimMeshToRenderMesh(ClothCollection, MaterialPathName, bGenerateSingleRenderPattern);
+		}
 
 		SetValue(Context, MoveTemp(*ClothCollection), &Collection);
 	}
diff --git a/Engine/Plugins/Experimental/ChaosClothAssetEditor/Source/ChaosClothAssetDataflowNodes/Private/ChaosClothAsset/DataflowNodes.cpp b/Engine/Plugins/Experimental/ChaosClothAssetEditor/Source/ChaosClothAssetDataflowNodes/Private/ChaosClothAsset/DataflowNodes.cpp
deleted file mode 100644
diff --git a/Engine/Plugins/Experimental/ChaosClothAssetEditor/Source/ChaosClothAssetDataflowNodes/Private/ChaosClothAsset/DatasmithImportNode.cpp b/Engine/Plugins/Experimental/ChaosClothAssetEditor/Source/ChaosClothAssetDataflowNodes/Private/ChaosClothAsset/DatasmithImportNode.cpp
--- a/Engine/Plugins/Experimental/ChaosClothAssetEditor/Source/ChaosClothAssetDataflowNodes/Private/ChaosClothAsset/DatasmithImportNode.cpp
+++ b/Engine/Plugins/Experimental/ChaosClothAssetEditor/Source/ChaosClothAssetDataflowNodes/Private/ChaosClothAsset/DatasmithImportNode.cpp
@@ -1,9 +1,9 @@
 // Copyright Epic Games, Inc. All Rights Reserved.
 
 #include "ChaosClothAsset/DatasmithImportNode.h"
-#include "ChaosClothAsset/DataflowNodes.h"
 #include "ChaosClothAsset/ClothAsset.h"
 #include "ChaosClothAsset/CollectionClothFacade.h"
+#include "Dataflow/DataflowInputOutput.h"
 #include "DatasmithImportContext.h"
 #include "DatasmithImportFactory.h"
 #include "ExternalSource.h"
diff --git a/Engine/Plugins/Experimental/ChaosClothAssetEditor/Source/ChaosClothAssetDataflowNodes/Private/ChaosClothAsset/DeleteRenderMeshNode.cpp b/Engine/Plugins/Experimental/ChaosClothAssetEditor/Source/ChaosClothAssetDataflowNodes/Private/ChaosClothAsset/DeleteRenderMeshNode.cpp
--- a/Engine/Plugins/Experimental/ChaosClothAssetEditor/Source/ChaosClothAssetDataflowNodes/Private/ChaosClothAsset/DeleteRenderMeshNode.cpp
+++ b/Engine/Plugins/Experimental/ChaosClothAssetEditor/Source/ChaosClothAssetDataflowNodes/Private/ChaosClothAsset/DeleteRenderMeshNode.cpp
@@ -1,8 +1,9 @@
 // Copyright Epic Games, Inc. All Rights Reserved.
 
 #include "ChaosClothAsset/DeleteRenderMeshNode.h"
-#include "ChaosClothAsset/DataflowNodes.h"
 #include "ChaosClothAsset/ClothGeometryTools.h"
+#include "ChaosClothAsset/CollectionClothFacade.h"
+#include "Dataflow/DataflowInputOutput.h"
 
 #include UE_INLINE_GENERATED_CPP_BY_NAME(DeleteRenderMeshNode)
 
@@ -12,7 +13,6 @@ FChaosClothAssetDeleteRenderMeshNode::FChaosClothAssetDeleteRenderMeshNode(const
 	: FDataflowNode(InParam, InGuid)
 {
 	RegisterInputConnection(&Collection);
-	RegisterInputConnection(&Patterns);
 	RegisterOutputConnection(&Collection, &Collection);
 }
 
@@ -25,16 +25,10 @@ void FChaosClothAssetDeleteRenderMeshNode::Evaluate(Dataflow::FContext& Context,
 		// Evaluate in collection
 		FManagedArrayCollection InCollection = GetValue<FManagedArrayCollection>(Context, &Collection);
 		const TSharedRef<FManagedArrayCollection> ClothCollection = MakeShared<FManagedArrayCollection>(MoveTemp(InCollection));
-
-		if (Patterns.IsEmpty())
+		if (FCollectionClothFacade(ClothCollection).IsValid())  // Can only act on the collection if it is a valid cloth collection
 		{
 			FClothGeometryTools::DeleteRenderMesh(ClothCollection);
 		}
-		else
-		{
-			DataflowNodes::LogAndToastWarning(LOCTEXT("DeletePerPatternNotImplemented", "FClothAssetDeleteRenderMeshNode: Delete per patterns not implemented."));
-		}
-
 		SetValue(Context, MoveTemp(*ClothCollection), &Collection);
 	}
 }
diff --git a/Engine/Plugins/Experimental/ChaosClothAssetEditor/Source/ChaosClothAssetDataflowNodes/Private/ChaosClothAsset/ImportNode.cpp b/Engine/Plugins/Experimental/ChaosClothAssetEditor/Source/ChaosClothAssetDataflowNodes/Private/ChaosClothAsset/ImportNode.cpp
--- a/Engine/Plugins/Experimental/ChaosClothAssetEditor/Source/ChaosClothAssetDataflowNodes/Private/ChaosClothAsset/ImportNode.cpp
+++ b/Engine/Plugins/Experimental/ChaosClothAssetEditor/Source/ChaosClothAssetDataflowNodes/Private/ChaosClothAsset/ImportNode.cpp
@@ -1,9 +1,11 @@
 // Copyright Epic Games, Inc. All Rights Reserved.
 
 #include "ChaosClothAsset/ImportNode.h"
-#include "ChaosClothAsset/DataflowNodes.h"
 #include "ChaosClothAsset/ClothAsset.h"
+#include "ChaosClothAsset/ClothDataflowTools.h"
 #include "ChaosClothAsset/CollectionClothFacade.h"
+#include "Dataflow/DataflowInputOutput.h"
+#include "Dataflow/DataflowObjectInterface.h"
 #include "Engine/SkinnedAssetCommon.h"
 #include "Materials/Material.h"
 #include "PhysicsEngine/PhysicsAsset.h"
@@ -15,15 +17,12 @@
 FChaosClothAssetImportNode::FChaosClothAssetImportNode(const Dataflow::FNodeParameters& InParam, FGuid InGuid)
 	: FDataflowNode(InParam, InGuid)
 {
-	RegisterInputConnection(&ClothAsset);
-	RegisterInputConnection(&ImportLod);
 	RegisterOutputConnection(&Collection);
 }
 
 void FChaosClothAssetImportNode::Evaluate(Dataflow::FContext& Context, const FDataflowOutput* Out) const
 {
 	using namespace UE::Chaos::ClothAsset;
-	using namespace UE::Chaos::ClothAsset::DataflowNodes;
 
 	if (Out->IsA<FManagedArrayCollection>(&Collection))
 	{
@@ -39,7 +38,9 @@ void FChaosClothAssetImportNode::Evaluate(Dataflow::FContext& Context, const FDa
 					{
 						// Can't create a loop
 						bSetValue = false;
-						LogAndToastWarning(LOCTEXT("RecursiveAssetLoop", "FClothAssetNode: The source asset cannot be the terminal asset."));
+						FClothDataflowTools::LogAndToastWarning(*this,
+							LOCTEXT("RecursiveAssetLoopHeadline", "Recursive asset loop."),
+							LOCTEXT("RecursiveAssetLoopDetails", "The source asset cannot be the same as the terminal asset."));
 					}
 					else
 					{
@@ -62,12 +63,10 @@ void FChaosClothAssetImportNode::Evaluate(Dataflow::FContext& Context, const FDa
 		// Copy the cloth asset to this node's output collection
 		if (bSetValue)
 		{
-			const TObjectPtr<const UChaosClothAsset>& InClothAsset = GetValue<TObjectPtr<const UChaosClothAsset>>(Context, &ClothAsset);
-			const int32& InImportLod = GetValue<int32>(Context, &ImportLod);
-			const TArray<TSharedRef<const FManagedArrayCollection>>& InClothCollections = InClothAsset->GetClothCollections();
-			if (InImportLod >= 0 && InClothCollections.Num() > InImportLod)
+			const TArray<TSharedRef<const FManagedArrayCollection>>& InClothCollections = ClothAsset->GetClothCollections();
+			if (ImportLod >= 0 && InClothCollections.Num() > ImportLod)
 			{
-				const FCollectionClothConstFacade InClothFacade(InClothCollections[InImportLod]);
+				const FCollectionClothConstFacade InClothFacade(InClothCollections[ImportLod]);
 				ClothFacade.Initialize(InClothFacade);
 			}
 		}
diff --git a/Engine/Plugins/Experimental/ChaosClothAssetEditor/Source/ChaosClothAssetDataflowNodes/Private/ChaosClothAsset/MergeClothCollectionsNode.cpp b/Engine/Plugins/Experimental/ChaosClothAssetEditor/Source/ChaosClothAssetDataflowNodes/Private/ChaosClothAsset/MergeClothCollectionsNode.cpp
--- a/Engine/Plugins/Experimental/ChaosClothAssetEditor/Source/ChaosClothAssetDataflowNodes/Private/ChaosClothAsset/MergeClothCollectionsNode.cpp
+++ b/Engine/Plugins/Experimental/ChaosClothAssetEditor/Source/ChaosClothAssetDataflowNodes/Private/ChaosClothAsset/MergeClothCollectionsNode.cpp
@@ -1,8 +1,8 @@
 // Copyright Epic Games, Inc. All Rights Reserved.
 
 #include "ChaosClothAsset/MergeClothCollectionsNode.h"
-#include "ChaosClothAsset/DataflowNodes.h"
 #include "ChaosClothAsset/CollectionClothFacade.h"
+#include "Dataflow/DataflowInputOutput.h"
 
 #include UE_INLINE_GENERATED_CPP_BY_NAME(MergeClothCollectionsNode)
 
@@ -24,21 +24,43 @@ void FChaosClothAssetMergeClothCollectionsNode::Evaluate(Dataflow::FContext& Con
 		// Evaluate in collection
 		FManagedArrayCollection InCollection = GetValue<FManagedArrayCollection>(Context, &Collection);
 		const TSharedRef<FManagedArrayCollection> ClothCollection = MakeShared<FManagedArrayCollection>(MoveTemp(InCollection));
-		
-		// Make it a valid cloth collection if needed
+
+		// Keep track of whether any of these collections are valid cloth collections
 		FCollectionClothFacade ClothFacade(ClothCollection);
-		ClothFacade.DefineSchema();
+		bool bAreAnyValid = ClothFacade.IsValid();
+
+		// Make it a valid cloth collection if needed
+		if (!bAreAnyValid)
+		{
+			ClothFacade.DefineSchema();
+		}
+
 		// Iterate through the inputs and append them to LOD 0
 		const TArray<const FManagedArrayCollection*> Collections = GetCollections();
 		for (int32 InputIndex = 1; InputIndex < Collections.Num(); ++InputIndex)
 		{
-			const FManagedArrayCollection& OtherCollection = GetValue<FManagedArrayCollection>(Context, Collections[InputIndex]);
-			const TSharedRef<const FManagedArrayCollection> OtherClothCollection = MakeShared<const FManagedArrayCollection>(OtherCollection);
-			const FCollectionClothConstFacade OtherClothFacade(OtherClothCollection);  // TODO: Remove Shared pointer from facade and skip duplication
-			ClothFacade.Append(OtherClothFacade);
+			FManagedArrayCollection OtherCollection = GetValue<FManagedArrayCollection>(Context, Collections[InputIndex]);  // Can't use a const reference here sadly since the facade needs a SharedRef to be created
+			const TSharedRef<const FManagedArrayCollection> OtherClothCollection = MakeShared<const FManagedArrayCollection>(MoveTemp(OtherCollection));
+			const FCollectionClothConstFacade OtherClothFacade(OtherClothCollection);
+			if (OtherClothFacade.IsValid())
+			{
+				ClothFacade.Append(OtherClothFacade);
+				bAreAnyValid = true;
+			}
 		}
 
-		SetValue(Context, MoveTemp(*ClothCollection), &Collection);
+		// Set the output
+		if (bAreAnyValid)
+		{
+			// Use the merged cloth collection, but only if there were at least one valid input cloth collections
+			SetValue(Context, MoveTemp(*ClothCollection), &Collection);
+		}
+		else
+		{
+			// Otherwise pass through the first input unchanged
+			const FManagedArrayCollection& Passthrough = GetValue<FManagedArrayCollection>(Context, &Collection);
+			SetValue(Context, Passthrough, &Collection);
+		}
 	}
 }
 
diff --git a/Engine/Plugins/Experimental/ChaosClothAssetEditor/Source/ChaosClothAssetDataflowNodes/Private/ChaosClothAsset/ReverseNormalsNode.cpp b/Engine/Plugins/Experimental/ChaosClothAssetEditor/Source/ChaosClothAssetDataflowNodes/Private/ChaosClothAsset/ReverseNormalsNode.cpp
--- a/Engine/Plugins/Experimental/ChaosClothAssetEditor/Source/ChaosClothAssetDataflowNodes/Private/ChaosClothAsset/ReverseNormalsNode.cpp
+++ b/Engine/Plugins/Experimental/ChaosClothAssetEditor/Source/ChaosClothAssetDataflowNodes/Private/ChaosClothAsset/ReverseNormalsNode.cpp
@@ -1,8 +1,9 @@
 // Copyright Epic Games, Inc. All Rights Reserved.
 
 #include "ChaosClothAsset/ReverseNormalsNode.h"
-#include "ChaosClothAsset/DataflowNodes.h"
 #include "ChaosClothAsset/ClothGeometryTools.h"
+#include "ChaosClothAsset/CollectionClothFacade.h"
+#include "Dataflow/DataflowInputOutput.h"
 
 #include UE_INLINE_GENERATED_CPP_BY_NAME(ReverseNormalsNode)
 
@@ -12,8 +13,6 @@ FChaosClothAssetReverseNormalsNode::FChaosClothAssetReverseNormalsNode(const Dat
 	: FDataflowNode(InParam, InGuid)
 {
 	RegisterInputConnection(&Collection);
-	RegisterInputConnection(&SimPatterns);
-	RegisterInputConnection(&RenderPatterns);
 	RegisterOutputConnection(&Collection, &Collection);
 }
 
@@ -27,14 +26,17 @@ void FChaosClothAssetReverseNormalsNode::Evaluate(Dataflow::FContext& Context, c
 		FManagedArrayCollection InCollection = GetValue<FManagedArrayCollection>(Context, &Collection);
 		const TSharedRef<FManagedArrayCollection> ClothCollection = MakeShared<FManagedArrayCollection>(MoveTemp(InCollection));
 
-		FClothGeometryTools::ReverseMesh(
-			ClothCollection,
-			bReverseSimMeshNormals,
-			bReverseSimMeshWindingOrder,
-			bReverseRenderMeshNormals,
-			bReverseRenderMeshWindingOrder,
-			GetValue<TArray<int32>>(Context, &SimPatterns),
-			GetValue<TArray<int32>>(Context, &RenderPatterns));
+		if (FCollectionClothFacade(ClothCollection).IsValid())  // Can only act on the collection if it is a valid cloth collection
+		{
+			FClothGeometryTools::ReverseMesh(
+				ClothCollection,
+				bReverseSimMeshNormals,
+				bReverseSimMeshWindingOrder,
+				bReverseRenderMeshNormals,
+				bReverseRenderMeshWindingOrder,
+				SimPatterns,
+				RenderPatterns);
+		}
 
 		SetValue(Context, MoveTemp(*ClothCollection), &Collection);
 	}
diff --git a/Engine/Plugins/Experimental/ChaosClothAssetEditor/Source/ChaosClothAssetDataflowNodes/Private/ChaosClothAsset/SetPhysicsAssetNode.cpp b/Engine/Plugins/Experimental/ChaosClothAssetEditor/Source/ChaosClothAssetDataflowNodes/Private/ChaosClothAsset/SetPhysicsAssetNode.cpp
--- a/Engine/Plugins/Experimental/ChaosClothAssetEditor/Source/ChaosClothAssetDataflowNodes/Private/ChaosClothAsset/SetPhysicsAssetNode.cpp
+++ b/Engine/Plugins/Experimental/ChaosClothAssetEditor/Source/ChaosClothAssetDataflowNodes/Private/ChaosClothAsset/SetPhysicsAssetNode.cpp
@@ -1,9 +1,9 @@
 // Copyright Epic Games, Inc. All Rights Reserved.
 
 #include "ChaosClothAsset/SetPhysicsAssetNode.h"
-#include "ChaosClothAsset/DataflowNodes.h"
 #include "ChaosClothAsset/ClothGeometryTools.h"
 #include "ChaosClothAsset/CollectionClothFacade.h"
+#include "Dataflow/DataflowInputOutput.h"
 
 #include UE_INLINE_GENERATED_CPP_BY_NAME(SetPhysicsAssetNode)
 
@@ -13,7 +13,6 @@ FChaosClothAssetSetPhysicsAssetNode::FChaosClothAssetSetPhysicsAssetNode(const D
 	: FDataflowNode(InParam, InGuid)
 {
 	RegisterInputConnection(&Collection);
-	RegisterInputConnection(&PhysicsAsset);
 	RegisterOutputConnection(&Collection, &Collection);
 }
 
@@ -27,8 +26,11 @@ void FChaosClothAssetSetPhysicsAssetNode::Evaluate(Dataflow::FContext& Context,
 		FManagedArrayCollection InCollection = GetValue<FManagedArrayCollection>(Context, &Collection);
 		const TSharedRef<FManagedArrayCollection> ClothCollection = MakeShared<FManagedArrayCollection>(MoveTemp(InCollection));
 
-		FCollectionClothFacade CollectionClothFacade(ClothCollection);
-		CollectionClothFacade.SetPhysicsAssetPathName(PhysicsAsset.GetPathName());
+		if (FCollectionClothFacade(ClothCollection).IsValid())  // Can only act on the collection if it is a valid cloth collection
+		{
+			FCollectionClothFacade CollectionClothFacade(ClothCollection);
+			CollectionClothFacade.SetPhysicsAssetPathName(PhysicsAsset.GetPathName());
+		}
 
 		SetValue(Context, MoveTemp(*ClothCollection), &Collection);
 	}
diff --git a/Engine/Plugins/Experimental/ChaosClothAssetEditor/Source/ChaosClothAssetDataflowNodes/Private/ChaosClothAsset/SimulationAerodynamicsConfigNode.cpp b/Engine/Plugins/Experimental/ChaosClothAssetEditor/Source/ChaosClothAssetDataflowNodes/Private/ChaosClothAsset/SimulationAerodynamicsConfigNode.cpp
--- a/Engine/Plugins/Experimental/ChaosClothAssetEditor/Source/ChaosClothAssetDataflowNodes/Private/ChaosClothAsset/SimulationAerodynamicsConfigNode.cpp
+++ b/Engine/Plugins/Experimental/ChaosClothAssetEditor/Source/ChaosClothAssetDataflowNodes/Private/ChaosClothAsset/SimulationAerodynamicsConfigNode.cpp
@@ -1,7 +1,6 @@
 // Copyright Epic Games, Inc. All Rights Reserved.
 
 #include "ChaosClothAsset/SimulationAerodynamicsConfigNode.h"
-#include "ChaosClothAsset/DataflowNodes.h"
 #include "ChaosClothAsset/SimulationBaseConfigNodePrivate.h"
 #include "Chaos/CollectionPropertyFacade.h"
 
diff --git a/Engine/Plugins/Experimental/ChaosClothAssetEditor/Source/ChaosClothAssetDataflowNodes/Private/ChaosClothAsset/SimulationAnimDriveConfigNode.cpp b/Engine/Plugins/Experimental/ChaosClothAssetEditor/Source/ChaosClothAssetDataflowNodes/Private/ChaosClothAsset/SimulationAnimDriveConfigNode.cpp
--- a/Engine/Plugins/Experimental/ChaosClothAssetEditor/Source/ChaosClothAssetDataflowNodes/Private/ChaosClothAsset/SimulationAnimDriveConfigNode.cpp
+++ b/Engine/Plugins/Experimental/ChaosClothAssetEditor/Source/ChaosClothAssetDataflowNodes/Private/ChaosClothAsset/SimulationAnimDriveConfigNode.cpp
@@ -1,7 +1,6 @@
 // Copyright Epic Games, Inc. All Rights Reserved.
 
 #include "ChaosClothAsset/SimulationAnimDriveConfigNode.h"
-#include "ChaosClothAsset/DataflowNodes.h"
 #include "ChaosClothAsset/SimulationBaseConfigNodePrivate.h"
 #include "Chaos/CollectionPropertyFacade.h"
 
diff --git a/Engine/Plugins/Experimental/ChaosClothAssetEditor/Source/ChaosClothAssetDataflowNodes/Private/ChaosClothAsset/SimulationBackstopConfigNode.cpp b/Engine/Plugins/Experimental/ChaosClothAssetEditor/Source/ChaosClothAssetDataflowNodes/Private/ChaosClothAsset/SimulationBackstopConfigNode.cpp
--- a/Engine/Plugins/Experimental/ChaosClothAssetEditor/Source/ChaosClothAssetDataflowNodes/Private/ChaosClothAsset/SimulationBackstopConfigNode.cpp
+++ b/Engine/Plugins/Experimental/ChaosClothAssetEditor/Source/ChaosClothAssetDataflowNodes/Private/ChaosClothAsset/SimulationBackstopConfigNode.cpp
@@ -1,7 +1,6 @@
 // Copyright Epic Games, Inc. All Rights Reserved.
 
 #include "ChaosClothAsset/SimulationBackstopConfigNode.h"
-#include "ChaosClothAsset/DataflowNodes.h"
 #include "ChaosClothAsset/SimulationBaseConfigNodePrivate.h"
 #include "Chaos/CollectionPropertyFacade.h"
 
diff --git a/Engine/Plugins/Experimental/ChaosClothAssetEditor/Source/ChaosClothAssetDataflowNodes/Private/ChaosClothAsset/SimulationBaseConfigNode.cpp b/Engine/Plugins/Experimental/ChaosClothAssetEditor/Source/ChaosClothAssetDataflowNodes/Private/ChaosClothAsset/SimulationBaseConfigNode.cpp
--- a/Engine/Plugins/Experimental/ChaosClothAssetEditor/Source/ChaosClothAssetDataflowNodes/Private/ChaosClothAsset/SimulationBaseConfigNode.cpp
+++ b/Engine/Plugins/Experimental/ChaosClothAssetEditor/Source/ChaosClothAssetDataflowNodes/Private/ChaosClothAsset/SimulationBaseConfigNode.cpp
@@ -1,41 +1,46 @@
 // Copyright Epic Games, Inc. All Rights Reserved.
 
 #include "ChaosClothAsset/SimulationBaseConfigNode.h"
-#include "ChaosClothAsset/DataflowNodes.h"
+#include "ChaosClothAsset/CollectionClothFacade.h"
+#include "ChaosClothAsset/ClothDataflowTools.h"
 #include "Chaos/CollectionPropertyFacade.h"
+#include "Dataflow/DataflowInputOutput.h"
 
 #define LOCTEXT_NAMESPACE "ChaosClothAssetSimulationBaseConfigNode"
 
 namespace UE::Chaos::ClothAsset::Private
 {
-	static void LogAndToastDuplicateProperty(const FName& NodeName, const FName& PropertyName)
+	static void LogAndToastDuplicateProperty(const FDataflowNode& DataflowNode, const FName& PropertyName)
 	{
-		using namespace UE::Chaos::ClothAsset::DataflowNodes;
+		using namespace UE::Chaos::ClothAsset;
 
-		const FText Message = FText::Format(
+		static const FText Headline = LOCTEXT("DuplicatePropertyHeadline", "Duplicate property.");
+
+		const FText Details = FText::Format(
 			LOCTEXT(
-				"DuplicateProperty",
-				"Cloth collection property '{1}' was already set in an upstream node, and its values are now overriden by node '{0}'."),
-			FText::FromName(NodeName),
+				"DuplicatePropertyDetails",
+				"Cloth collection property '{0}' was already set in an upstream node.\n"
+				"Its values have now been overridden."),
 			FText::FromName(PropertyName));
 
-		LogAndToastWarning(Message);
+		FClothDataflowTools::LogAndToastWarning(DataflowNode, Headline, Details);
 	}
 
-	static void LogAndToastSimilarProperty(const FName& NodeName, const FName& PropertyName, const FName& SimilarPropertyName)
+	static void LogAndToastSimilarProperty(const FDataflowNode& DataflowNode, const FName& PropertyName, const FName& SimilarPropertyName)
 	{
-		using namespace UE::Chaos::ClothAsset::DataflowNodes;
+		using namespace UE::Chaos::ClothAsset;
+
+		static const FText Headline = LOCTEXT("SimilarPropertyHeadline", "Similar property.");
 
-		const FText Message = FText::Format(
+		const FText Details = FText::Format(
 			LOCTEXT(
-				"SimilarProperty",
-				"Cloth collection property '{1}' set in node and '{0}' is similar to the property '{2}' already set in an upstream node, "
-				"which might result in an undefined simulation behavior."),
-			FText::FromName(NodeName),
+				"SimilarPropertyDetails",
+				"Cloth collection property '{0}' is similar to the property '{1}' already set in an upstream node.\n"
+				"This might result in an undefined simulation behavior."),
 			FText::FromName(PropertyName),
 			FText::FromName(SimilarPropertyName));
 
-		LogAndToastWarning(Message);
+		FClothDataflowTools::LogAndToastWarning(DataflowNode, Headline, Details);
 	}
 }
 
@@ -52,19 +57,22 @@ void FChaosClothAssetSimulationBaseConfigNode::RegisterCollectionConnections()
 void FChaosClothAssetSimulationBaseConfigNode::Evaluate(Dataflow::FContext& Context, const FDataflowOutput* Out) const
 {
 	using namespace Chaos::Softs;
-	using namespace UE::Chaos::ClothAsset::DataflowNodes;
+	using namespace UE::Chaos::ClothAsset;
 
 	if (Out->IsA<FManagedArrayCollection>(&Collection))
 	{
 		FManagedArrayCollection InCollection = GetValue<FManagedArrayCollection>(Context, &Collection);
 		const TSharedRef<FManagedArrayCollection> ClothCollection = MakeShared<FManagedArrayCollection>(MoveTemp(InCollection));
 
-		FCollectionPropertyMutableFacade Properties(ClothCollection);
-		Properties.DefineSchema();
+		if (FCollectionClothFacade(ClothCollection).IsValid())  // Can only act on the collection if it is a valid cloth collection
+		{
+			FCollectionPropertyMutableFacade Properties(ClothCollection);
+			Properties.DefineSchema();
 
-		AddProperties(Properties);
+			AddProperties(Properties);
 
-		EvaluateClothCollection(Context, ClothCollection);
+			EvaluateClothCollection(Context, ClothCollection);
+		}
 
 		SetValue(Context, MoveTemp(*ClothCollection), &Collection);
 	}
@@ -76,7 +84,7 @@ int32 FChaosClothAssetSimulationBaseConfigNode::AddPropertyHelper(
 	bool bIsAnimatable,
 	const TArray<FName>& SimilarPropertyNames) const
 {
-	using namespace UE::Chaos::ClothAsset::Private;
+	using namespace UE::Chaos::ClothAsset;
 
 	constexpr bool bIsEnabled = true;
 
@@ -88,14 +96,14 @@ int32 FChaosClothAssetSimulationBaseConfigNode::AddPropertyHelper(
 	else
 	{
 		Properties.SetAnimatable(KeyIndex, bIsAnimatable);
-		LogAndToastDuplicateProperty(FDataflowNode::Name, PropertyName);
+		Private::LogAndToastDuplicateProperty(*this, PropertyName);
 	}
 
 	for (const FName& SimilarPropertyName : SimilarPropertyNames)
 	{
 		if (Properties.GetKeyIndex(SimilarPropertyName.ToString()) != INDEX_NONE)
 		{
-			LogAndToastSimilarProperty(FDataflowNode::Name, PropertyName, SimilarPropertyName);
+			Private::LogAndToastSimilarProperty(*this, PropertyName, SimilarPropertyName);
 		}
 	}
 
diff --git a/Engine/Plugins/Experimental/ChaosClothAssetEditor/Source/ChaosClothAssetDataflowNodes/Private/ChaosClothAsset/SimulationCollisionConfigNode.cpp b/Engine/Plugins/Experimental/ChaosClothAssetEditor/Source/ChaosClothAssetDataflowNodes/Private/ChaosClothAsset/SimulationCollisionConfigNode.cpp
--- a/Engine/Plugins/Experimental/ChaosClothAssetEditor/Source/ChaosClothAssetDataflowNodes/Private/ChaosClothAsset/SimulationCollisionConfigNode.cpp
+++ b/Engine/Plugins/Experimental/ChaosClothAssetEditor/Source/ChaosClothAssetDataflowNodes/Private/ChaosClothAsset/SimulationCollisionConfigNode.cpp
@@ -1,7 +1,6 @@
 // Copyright Epic Games, Inc. All Rights Reserved.
 
 #include "ChaosClothAsset/SimulationCollisionConfigNode.h"
-#include "ChaosClothAsset/DataflowNodes.h"
 #include "ChaosClothAsset/SimulationBaseConfigNodePrivate.h"
 #include "Chaos/CollectionPropertyFacade.h"
 
diff --git a/Engine/Plugins/Experimental/ChaosClothAssetEditor/Source/ChaosClothAssetDataflowNodes/Private/ChaosClothAsset/SimulationDampingConfigNode.cpp b/Engine/Plugins/Experimental/ChaosClothAssetEditor/Source/ChaosClothAssetDataflowNodes/Private/ChaosClothAsset/SimulationDampingConfigNode.cpp
--- a/Engine/Plugins/Experimental/ChaosClothAssetEditor/Source/ChaosClothAssetDataflowNodes/Private/ChaosClothAsset/SimulationDampingConfigNode.cpp
+++ b/Engine/Plugins/Experimental/ChaosClothAssetEditor/Source/ChaosClothAssetDataflowNodes/Private/ChaosClothAsset/SimulationDampingConfigNode.cpp
@@ -1,7 +1,6 @@
 // Copyright Epic Games, Inc. All Rights Reserved.
 
 #include "ChaosClothAsset/SimulationDampingConfigNode.h"
-#include "ChaosClothAsset/DataflowNodes.h"
 #include "ChaosClothAsset/SimulationBaseConfigNodePrivate.h"
 #include "Chaos/CollectionPropertyFacade.h"
 
diff --git a/Engine/Plugins/Experimental/ChaosClothAssetEditor/Source/ChaosClothAssetDataflowNodes/Private/ChaosClothAsset/SimulationDefaultConfigNode.cpp b/Engine/Plugins/Experimental/ChaosClothAssetEditor/Source/ChaosClothAssetDataflowNodes/Private/ChaosClothAsset/SimulationDefaultConfigNode.cpp
--- a/Engine/Plugins/Experimental/ChaosClothAssetEditor/Source/ChaosClothAssetDataflowNodes/Private/ChaosClothAsset/SimulationDefaultConfigNode.cpp
+++ b/Engine/Plugins/Experimental/ChaosClothAssetEditor/Source/ChaosClothAssetDataflowNodes/Private/ChaosClothAsset/SimulationDefaultConfigNode.cpp
@@ -1,12 +1,13 @@
 // Copyright Epic Games, Inc. All Rights Reserved.
 
 #include "ChaosClothAsset/SimulationDefaultConfigNode.h"
-#include "ChaosClothAsset/DataflowNodes.h"
 #include "ChaosClothAsset/ClothEngineTools.h"
+#include "ChaosClothAsset/CollectionClothFacade.h"
 #include "ChaosCloth/ChaosClothConfig.h"
 #include "ChaosCloth/ChaosClothingSimulationConfig.h"
-#include "Chaos/PBDLongRangeConstraints.h"  // For Tether modes
 #include "Chaos/CollectionPropertyFacade.h"
+#include "Chaos/PBDLongRangeConstraints.h"  // For Tether modes
+#include "Dataflow/DataflowInputOutput.h"
 
 #include UE_INLINE_GENERATED_CPP_BY_NAME(SimulationDefaultConfigNode)
 
@@ -38,27 +39,31 @@ void FChaosClothAssetSimulationDefaultConfigNode::Serialize(FArchive& Ar)
 
 void FChaosClothAssetSimulationDefaultConfigNode::Evaluate(Dataflow::FContext& Context, const FDataflowOutput* Out) const
 {
-	using namespace Chaos;
-	using namespace Chaos::Softs;
+	using namespace ::Chaos;
+	using namespace ::Chaos::Softs;
+	using namespace UE::Chaos::ClothAsset;
 
 	if (Out->IsA<FManagedArrayCollection>(&Collection))
 	{
 		FManagedArrayCollection InCollection = GetValue<FManagedArrayCollection>(Context, &Collection);
 		const TSharedRef<FManagedArrayCollection> ClothCollection = MakeShared<FManagedArrayCollection>(MoveTemp(InCollection));
 
-		FClothingSimulationConfig ClothingSimulationConfig;
-		ClothingSimulationConfig.Initialize(SimulationConfig.Get(), SharedSimulationConfig.Get());
+		if (FCollectionClothFacade(ClothCollection).IsValid())  // Can only act on the collection if it is a valid cloth collection
+		{
+			FClothingSimulationConfig ClothingSimulationConfig;
+			ClothingSimulationConfig.Initialize(SimulationConfig.Get(), SharedSimulationConfig.Get());
 
-		ClothingSimulationConfig.GetPropertyCollection()->CopyTo(&ClothCollection.Get());
+			ClothingSimulationConfig.GetPropertyCollection()->CopyTo(&ClothCollection.Get());
 
-		// Generate tethers
-		const FCollectionPropertyConstFacade& Properties = ClothingSimulationConfig.GetProperties();		
-		constexpr bool bUseGeodesicTethersDefault = true;
-		const bool bUseGeodesicTethers = Properties.GetValue<bool>(TEXT("UseGeodesicTethers"), bUseGeodesicTethersDefault);
-		// Use the "MaxDistance" weight map to generate tethers. This follows legacy behavior.
-		static const FName MaxDistanceName(TEXT("MaxDistance"));
+			// Generate tethers
+			const FCollectionPropertyConstFacade& Properties = ClothingSimulationConfig.GetProperties();
+			constexpr bool bUseGeodesicTethersDefault = true;
+			const bool bUseGeodesicTethers = Properties.GetValue<bool>(TEXT("UseGeodesicTethers"), bUseGeodesicTethersDefault);
+			// Use the "MaxDistance" weight map to generate tethers. This follows legacy behavior.
+			static const FName MaxDistanceName(TEXT("MaxDistance"));
 
-		UE::Chaos::ClothAsset::FClothEngineTools::GenerateTethers(ClothCollection, MaxDistanceName, bUseGeodesicTethers);
+			UE::Chaos::ClothAsset::FClothEngineTools::GenerateTethers(ClothCollection, MaxDistanceName, bUseGeodesicTethers);
+		}
 
 		SetValue(Context, MoveTemp(*ClothCollection), &Collection);
 	}
diff --git a/Engine/Plugins/Experimental/ChaosClothAssetEditor/Source/ChaosClothAssetDataflowNodes/Private/ChaosClothAsset/SimulationGravityConfigNode.cpp b/Engine/Plugins/Experimental/ChaosClothAssetEditor/Source/ChaosClothAssetDataflowNodes/Private/ChaosClothAsset/SimulationGravityConfigNode.cpp
--- a/Engine/Plugins/Experimental/ChaosClothAssetEditor/Source/ChaosClothAssetDataflowNodes/Private/ChaosClothAsset/SimulationGravityConfigNode.cpp
+++ b/Engine/Plugins/Experimental/ChaosClothAssetEditor/Source/ChaosClothAssetDataflowNodes/Private/ChaosClothAsset/SimulationGravityConfigNode.cpp
@@ -1,7 +1,6 @@
 // Copyright Epic Games, Inc. All Rights Reserved.
 
 #include "ChaosClothAsset/SimulationGravityConfigNode.h"
-#include "ChaosClothAsset/DataflowNodes.h"
 #include "ChaosClothAsset/SimulationBaseConfigNodePrivate.h"
 #include "Chaos/CollectionPropertyFacade.h"
 
diff --git a/Engine/Plugins/Experimental/ChaosClothAssetEditor/Source/ChaosClothAssetDataflowNodes/Private/ChaosClothAsset/SimulationLongRangeAttachmentConfigNode.cpp b/Engine/Plugins/Experimental/ChaosClothAssetEditor/Source/ChaosClothAssetDataflowNodes/Private/ChaosClothAsset/SimulationLongRangeAttachmentConfigNode.cpp
--- a/Engine/Plugins/Experimental/ChaosClothAssetEditor/Source/ChaosClothAssetDataflowNodes/Private/ChaosClothAsset/SimulationLongRangeAttachmentConfigNode.cpp
+++ b/Engine/Plugins/Experimental/ChaosClothAssetEditor/Source/ChaosClothAssetDataflowNodes/Private/ChaosClothAsset/SimulationLongRangeAttachmentConfigNode.cpp
@@ -1,10 +1,10 @@
 // Copyright Epic Games, Inc. All Rights Reserved.
 
 #include "ChaosClothAsset/SimulationLongRangeAttachmentConfigNode.h"
-#include "ChaosClothAsset/DataflowNodes.h"
 #include "ChaosClothAsset/SimulationBaseConfigNodePrivate.h"
 #include "ChaosClothAsset/ClothEngineTools.h"
 #include "Chaos/CollectionPropertyFacade.h"
+#include "Dataflow/DataflowInputOutput.h"
 
 #include UE_INLINE_GENERATED_CPP_BY_NAME(SimulationLongRangeAttachmentConfigNode)
 
diff --git a/Engine/Plugins/Experimental/ChaosClothAssetEditor/Source/ChaosClothAssetDataflowNodes/Private/ChaosClothAsset/SimulationMassConfigNode.cpp b/Engine/Plugins/Experimental/ChaosClothAssetEditor/Source/ChaosClothAssetDataflowNodes/Private/ChaosClothAsset/SimulationMassConfigNode.cpp
--- a/Engine/Plugins/Experimental/ChaosClothAssetEditor/Source/ChaosClothAssetDataflowNodes/Private/ChaosClothAsset/SimulationMassConfigNode.cpp
+++ b/Engine/Plugins/Experimental/ChaosClothAssetEditor/Source/ChaosClothAssetDataflowNodes/Private/ChaosClothAsset/SimulationMassConfigNode.cpp
@@ -1,7 +1,6 @@
 // Copyright Epic Games, Inc. All Rights Reserved.
 
 #include "ChaosClothAsset/SimulationMassConfigNode.h"
-#include "ChaosClothAsset/DataflowNodes.h"
 #include "ChaosClothAsset/SimulationBaseConfigNodePrivate.h"
 #include "Chaos/CollectionPropertyFacade.h"
 
diff --git a/Engine/Plugins/Experimental/ChaosClothAssetEditor/Source/ChaosClothAssetDataflowNodes/Private/ChaosClothAsset/SimulationMaxDistanceConfigNode.cpp b/Engine/Plugins/Experimental/ChaosClothAssetEditor/Source/ChaosClothAssetDataflowNodes/Private/ChaosClothAsset/SimulationMaxDistanceConfigNode.cpp
--- a/Engine/Plugins/Experimental/ChaosClothAssetEditor/Source/ChaosClothAssetDataflowNodes/Private/ChaosClothAsset/SimulationMaxDistanceConfigNode.cpp
+++ b/Engine/Plugins/Experimental/ChaosClothAssetEditor/Source/ChaosClothAssetDataflowNodes/Private/ChaosClothAsset/SimulationMaxDistanceConfigNode.cpp
@@ -1,7 +1,6 @@
 // Copyright Epic Games, Inc. All Rights Reserved.
 
 #include "ChaosClothAsset/SimulationMaxDistanceConfigNode.h"
-#include "ChaosClothAsset/DataflowNodes.h"
 #include "ChaosClothAsset/SimulationBaseConfigNodePrivate.h"
 #include "Chaos/CollectionPropertyFacade.h"
 
diff --git a/Engine/Plugins/Experimental/ChaosClothAssetEditor/Source/ChaosClothAssetDataflowNodes/Private/ChaosClothAsset/SimulationPBDAreaSpringConfigNode.cpp b/Engine/Plugins/Experimental/ChaosClothAssetEditor/Source/ChaosClothAssetDataflowNodes/Private/ChaosClothAsset/SimulationPBDAreaSpringConfigNode.cpp
--- a/Engine/Plugins/Experimental/ChaosClothAssetEditor/Source/ChaosClothAssetDataflowNodes/Private/ChaosClothAsset/SimulationPBDAreaSpringConfigNode.cpp
+++ b/Engine/Plugins/Experimental/ChaosClothAssetEditor/Source/ChaosClothAssetDataflowNodes/Private/ChaosClothAsset/SimulationPBDAreaSpringConfigNode.cpp
@@ -1,7 +1,6 @@
 // Copyright Epic Games, Inc. All Rights Reserved.
 
 #include "ChaosClothAsset/SimulationPBDAreaSpringConfigNode.h"
-#include "ChaosClothAsset/DataflowNodes.h"
 #include "ChaosClothAsset/SimulationBaseConfigNodePrivate.h"
 #include "Chaos/CollectionPropertyFacade.h"
 
diff --git a/Engine/Plugins/Experimental/ChaosClothAssetEditor/Source/ChaosClothAssetDataflowNodes/Private/ChaosClothAsset/SimulationPBDBendingElementConfigNode.cpp b/Engine/Plugins/Experimental/ChaosClothAssetEditor/Source/ChaosClothAssetDataflowNodes/Private/ChaosClothAsset/SimulationPBDBendingElementConfigNode.cpp
--- a/Engine/Plugins/Experimental/ChaosClothAssetEditor/Source/ChaosClothAssetDataflowNodes/Private/ChaosClothAsset/SimulationPBDBendingElementConfigNode.cpp
+++ b/Engine/Plugins/Experimental/ChaosClothAssetEditor/Source/ChaosClothAssetDataflowNodes/Private/ChaosClothAsset/SimulationPBDBendingElementConfigNode.cpp
@@ -1,7 +1,6 @@
 // Copyright Epic Games, Inc. All Rights Reserved.
 
 #include "ChaosClothAsset/SimulationPBDBendingElementConfigNode.h"
-#include "ChaosClothAsset/DataflowNodes.h"
 #include "ChaosClothAsset/SimulationBaseConfigNodePrivate.h"
 #include "Chaos/CollectionPropertyFacade.h"
 
diff --git a/Engine/Plugins/Experimental/ChaosClothAssetEditor/Source/ChaosClothAssetDataflowNodes/Private/ChaosClothAsset/SimulationPBDBendingSpringConfigNode.cpp b/Engine/Plugins/Experimental/ChaosClothAssetEditor/Source/ChaosClothAssetDataflowNodes/Private/ChaosClothAsset/SimulationPBDBendingSpringConfigNode.cpp
--- a/Engine/Plugins/Experimental/ChaosClothAssetEditor/Source/ChaosClothAssetDataflowNodes/Private/ChaosClothAsset/SimulationPBDBendingSpringConfigNode.cpp
+++ b/Engine/Plugins/Experimental/ChaosClothAssetEditor/Source/ChaosClothAssetDataflowNodes/Private/ChaosClothAsset/SimulationPBDBendingSpringConfigNode.cpp
@@ -1,7 +1,6 @@
 // Copyright Epic Games, Inc. All Rights Reserved.
 
 #include "ChaosClothAsset/SimulationPBDBendingSpringConfigNode.h"
-#include "ChaosClothAsset/DataflowNodes.h"
 #include "ChaosClothAsset/SimulationBaseConfigNodePrivate.h"
 #include "Chaos/CollectionPropertyFacade.h"
 
diff --git a/Engine/Plugins/Experimental/ChaosClothAssetEditor/Source/ChaosClothAssetDataflowNodes/Private/ChaosClothAsset/SimulationPBDEdgeSpringConfigNode.cpp b/Engine/Plugins/Experimental/ChaosClothAssetEditor/Source/ChaosClothAssetDataflowNodes/Private/ChaosClothAsset/SimulationPBDEdgeSpringConfigNode.cpp
--- a/Engine/Plugins/Experimental/ChaosClothAssetEditor/Source/ChaosClothAssetDataflowNodes/Private/ChaosClothAsset/SimulationPBDEdgeSpringConfigNode.cpp
+++ b/Engine/Plugins/Experimental/ChaosClothAssetEditor/Source/ChaosClothAssetDataflowNodes/Private/ChaosClothAsset/SimulationPBDEdgeSpringConfigNode.cpp
@@ -1,7 +1,6 @@
 // Copyright Epic Games, Inc. All Rights Reserved.
 
 #include "ChaosClothAsset/SimulationPBDEdgeSpringConfigNode.h"
-#include "ChaosClothAsset/DataflowNodes.h"
 #include "ChaosClothAsset/SimulationBaseConfigNodePrivate.h"
 #include "Chaos/CollectionPropertyFacade.h"
 
diff --git a/Engine/Plugins/Experimental/ChaosClothAssetEditor/Source/ChaosClothAssetDataflowNodes/Private/ChaosClothAsset/SimulationPressureConfigNode.cpp b/Engine/Plugins/Experimental/ChaosClothAssetEditor/Source/ChaosClothAssetDataflowNodes/Private/ChaosClothAsset/SimulationPressureConfigNode.cpp
--- a/Engine/Plugins/Experimental/ChaosClothAssetEditor/Source/ChaosClothAssetDataflowNodes/Private/ChaosClothAsset/SimulationPressureConfigNode.cpp
+++ b/Engine/Plugins/Experimental/ChaosClothAssetEditor/Source/ChaosClothAssetDataflowNodes/Private/ChaosClothAsset/SimulationPressureConfigNode.cpp
@@ -1,7 +1,6 @@
 // Copyright Epic Games, Inc. All Rights Reserved.
 
 #include "ChaosClothAsset/SimulationPressureConfigNode.h"
-#include "ChaosClothAsset/DataflowNodes.h"
 #include "ChaosClothAsset/SimulationBaseConfigNodePrivate.h"
 #include "Chaos/CollectionPropertyFacade.h"
 
diff --git a/Engine/Plugins/Experimental/ChaosClothAssetEditor/Source/ChaosClothAssetDataflowNodes/Private/ChaosClothAsset/SimulationSelfCollisionConfigNode.cpp b/Engine/Plugins/Experimental/ChaosClothAssetEditor/Source/ChaosClothAssetDataflowNodes/Private/ChaosClothAsset/SimulationSelfCollisionConfigNode.cpp
--- a/Engine/Plugins/Experimental/ChaosClothAssetEditor/Source/ChaosClothAssetDataflowNodes/Private/ChaosClothAsset/SimulationSelfCollisionConfigNode.cpp
+++ b/Engine/Plugins/Experimental/ChaosClothAssetEditor/Source/ChaosClothAssetDataflowNodes/Private/ChaosClothAsset/SimulationSelfCollisionConfigNode.cpp
@@ -1,7 +1,6 @@
 // Copyright Epic Games, Inc. All Rights Reserved.
 
 #include "ChaosClothAsset/SimulationSelfCollisionConfigNode.h"
-#include "ChaosClothAsset/DataflowNodes.h"
 #include "ChaosClothAsset/SimulationBaseConfigNodePrivate.h"
 #include "Chaos/CollectionPropertyFacade.h"
 
diff --git a/Engine/Plugins/Experimental/ChaosClothAssetEditor/Source/ChaosClothAssetDataflowNodes/Private/ChaosClothAsset/SimulationSolverConfigNode.cpp b/Engine/Plugins/Experimental/ChaosClothAssetEditor/Source/ChaosClothAssetDataflowNodes/Private/ChaosClothAsset/SimulationSolverConfigNode.cpp
--- a/Engine/Plugins/Experimental/ChaosClothAssetEditor/Source/ChaosClothAssetDataflowNodes/Private/ChaosClothAsset/SimulationSolverConfigNode.cpp
+++ b/Engine/Plugins/Experimental/ChaosClothAssetEditor/Source/ChaosClothAssetDataflowNodes/Private/ChaosClothAsset/SimulationSolverConfigNode.cpp
@@ -1,7 +1,6 @@
 // Copyright Epic Games, Inc. All Rights Reserved.
 
 #include "ChaosClothAsset/SimulationSolverConfigNode.h"
-#include "ChaosClothAsset/DataflowNodes.h"
 #include "ChaosClothAsset/SimulationBaseConfigNodePrivate.h"
 #include "Chaos/CollectionPropertyFacade.h"
 
diff --git a/Engine/Plugins/Experimental/ChaosClothAssetEditor/Source/ChaosClothAssetDataflowNodes/Private/ChaosClothAsset/SimulationVelocityScaleConfigNode.cpp b/Engine/Plugins/Experimental/ChaosClothAssetEditor/Source/ChaosClothAssetDataflowNodes/Private/ChaosClothAsset/SimulationVelocityScaleConfigNode.cpp
--- a/Engine/Plugins/Experimental/ChaosClothAssetEditor/Source/ChaosClothAssetDataflowNodes/Private/ChaosClothAsset/SimulationVelocityScaleConfigNode.cpp
+++ b/Engine/Plugins/Experimental/ChaosClothAssetEditor/Source/ChaosClothAssetDataflowNodes/Private/ChaosClothAsset/SimulationVelocityScaleConfigNode.cpp
@@ -1,7 +1,6 @@
 // Copyright Epic Games, Inc. All Rights Reserved.
 
 #include "ChaosClothAsset/SimulationVelocityScaleConfigNode.h"
-#include "ChaosClothAsset/DataflowNodes.h"
 #include "ChaosClothAsset/SimulationBaseConfigNodePrivate.h"
 #include "Chaos/CollectionPropertyFacade.h"
 
diff --git a/Engine/Plugins/Experimental/ChaosClothAssetEditor/Source/ChaosClothAssetDataflowNodes/Private/ChaosClothAsset/SimulationXPBDAnisoBendingConfigNode.cpp b/Engine/Plugins/Experimental/ChaosClothAssetEditor/Source/ChaosClothAssetDataflowNodes/Private/ChaosClothAsset/SimulationXPBDAnisoBendingConfigNode.cpp
--- a/Engine/Plugins/Experimental/ChaosClothAssetEditor/Source/ChaosClothAssetDataflowNodes/Private/ChaosClothAsset/SimulationXPBDAnisoBendingConfigNode.cpp
+++ b/Engine/Plugins/Experimental/ChaosClothAssetEditor/Source/ChaosClothAssetDataflowNodes/Private/ChaosClothAsset/SimulationXPBDAnisoBendingConfigNode.cpp
@@ -1,7 +1,6 @@
 // Copyright Epic Games, Inc. All Rights Reserved.
 
 #include "ChaosClothAsset/SimulationXPBDAnisoBendingConfigNode.h"
-#include "ChaosClothAsset/DataflowNodes.h"
 #include "ChaosClothAsset/SimulationBaseConfigNodePrivate.h"
 #include "Chaos/CollectionPropertyFacade.h"
 
diff --git a/Engine/Plugins/Experimental/ChaosClothAssetEditor/Source/ChaosClothAssetDataflowNodes/Private/ChaosClothAsset/SimulationXPBDAnisoStretchConfigNode.cpp b/Engine/Plugins/Experimental/ChaosClothAssetEditor/Source/ChaosClothAssetDataflowNodes/Private/ChaosClothAsset/SimulationXPBDAnisoStretchConfigNode.cpp
--- a/Engine/Plugins/Experimental/ChaosClothAssetEditor/Source/ChaosClothAssetDataflowNodes/Private/ChaosClothAsset/SimulationXPBDAnisoStretchConfigNode.cpp
+++ b/Engine/Plugins/Experimental/ChaosClothAssetEditor/Source/ChaosClothAssetDataflowNodes/Private/ChaosClothAsset/SimulationXPBDAnisoStretchConfigNode.cpp
@@ -1,7 +1,6 @@
 // Copyright Epic Games, Inc. All Rights Reserved.
 
 #include "ChaosClothAsset/SimulationXPBDAnisoStretchConfigNode.h"
-#include "ChaosClothAsset/DataflowNodes.h"
 #include "ChaosClothAsset/SimulationBaseConfigNodePrivate.h"
 #include "Chaos/CollectionPropertyFacade.h"
 
diff --git a/Engine/Plugins/Experimental/ChaosClothAssetEditor/Source/ChaosClothAssetDataflowNodes/Private/ChaosClothAsset/SimulationXPBDAreaSpringConfigNode.cpp b/Engine/Plugins/Experimental/ChaosClothAssetEditor/Source/ChaosClothAssetDataflowNodes/Private/ChaosClothAsset/SimulationXPBDAreaSpringConfigNode.cpp
--- a/Engine/Plugins/Experimental/ChaosClothAssetEditor/Source/ChaosClothAssetDataflowNodes/Private/ChaosClothAsset/SimulationXPBDAreaSpringConfigNode.cpp
+++ b/Engine/Plugins/Experimental/ChaosClothAssetEditor/Source/ChaosClothAssetDataflowNodes/Private/ChaosClothAsset/SimulationXPBDAreaSpringConfigNode.cpp
@@ -1,7 +1,6 @@
 // Copyright Epic Games, Inc. All Rights Reserved.
 
 #include "ChaosClothAsset/SimulationXPBDAreaSpringConfigNode.h"
-#include "ChaosClothAsset/DataflowNodes.h"
 #include "ChaosClothAsset/SimulationBaseConfigNodePrivate.h"
 #include "Chaos/CollectionPropertyFacade.h"
 
diff --git a/Engine/Plugins/Experimental/ChaosClothAssetEditor/Source/ChaosClothAssetDataflowNodes/Private/ChaosClothAsset/SimulationXPBDBendingElementConfigNode.cpp b/Engine/Plugins/Experimental/ChaosClothAssetEditor/Source/ChaosClothAssetDataflowNodes/Private/ChaosClothAsset/SimulationXPBDBendingElementConfigNode.cpp
--- a/Engine/Plugins/Experimental/ChaosClothAssetEditor/Source/ChaosClothAssetDataflowNodes/Private/ChaosClothAsset/SimulationXPBDBendingElementConfigNode.cpp
+++ b/Engine/Plugins/Experimental/ChaosClothAssetEditor/Source/ChaosClothAssetDataflowNodes/Private/ChaosClothAsset/SimulationXPBDBendingElementConfigNode.cpp
@@ -1,7 +1,6 @@
 // Copyright Epic Games, Inc. All Rights Reserved.
 
 #include "ChaosClothAsset/SimulationXPBDBendingElementConfigNode.h"
-#include "ChaosClothAsset/DataflowNodes.h"
 #include "ChaosClothAsset/SimulationBaseConfigNodePrivate.h"
 #include "Chaos/CollectionPropertyFacade.h"
 
diff --git a/Engine/Plugins/Experimental/ChaosClothAssetEditor/Source/ChaosClothAssetDataflowNodes/Private/ChaosClothAsset/SimulationXPBDBendingSpringConfigNode.cpp b/Engine/Plugins/Experimental/ChaosClothAssetEditor/Source/ChaosClothAssetDataflowNodes/Private/ChaosClothAsset/SimulationXPBDBendingSpringConfigNode.cpp
--- a/Engine/Plugins/Experimental/ChaosClothAssetEditor/Source/ChaosClothAssetDataflowNodes/Private/ChaosClothAsset/SimulationXPBDBendingSpringConfigNode.cpp
+++ b/Engine/Plugins/Experimental/ChaosClothAssetEditor/Source/ChaosClothAssetDataflowNodes/Private/ChaosClothAsset/SimulationXPBDBendingSpringConfigNode.cpp
@@ -1,7 +1,6 @@
 // Copyright Epic Games, Inc. All Rights Reserved.
 
 #include "ChaosClothAsset/SimulationXPBDBendingSpringConfigNode.h"
-#include "ChaosClothAsset/DataflowNodes.h"
 #include "ChaosClothAsset/SimulationBaseConfigNodePrivate.h"
 #include "Chaos/CollectionPropertyFacade.h"
 
diff --git a/Engine/Plugins/Experimental/ChaosClothAssetEditor/Source/ChaosClothAssetDataflowNodes/Private/ChaosClothAsset/SimulationXPBDEdgeSpringConfigNode.cpp b/Engine/Plugins/Experimental/ChaosClothAssetEditor/Source/ChaosClothAssetDataflowNodes/Private/ChaosClothAsset/SimulationXPBDEdgeSpringConfigNode.cpp
--- a/Engine/Plugins/Experimental/ChaosClothAssetEditor/Source/ChaosClothAssetDataflowNodes/Private/ChaosClothAsset/SimulationXPBDEdgeSpringConfigNode.cpp
+++ b/Engine/Plugins/Experimental/ChaosClothAssetEditor/Source/ChaosClothAssetDataflowNodes/Private/ChaosClothAsset/SimulationXPBDEdgeSpringConfigNode.cpp
@@ -1,7 +1,6 @@
 // Copyright Epic Games, Inc. All Rights Reserved.
 
 #include "ChaosClothAsset/SimulationXPBDEdgeSpringConfigNode.h"
-#include "ChaosClothAsset/DataflowNodes.h"
 #include "ChaosClothAsset/SimulationBaseConfigNodePrivate.h"
 #include "Chaos/CollectionPropertyFacade.h"
 
diff --git a/Engine/Plugins/Experimental/ChaosClothAssetEditor/Source/ChaosClothAssetDataflowNodes/Private/ChaosClothAsset/SkeletalMeshImportNode.cpp b/Engine/Plugins/Experimental/ChaosClothAssetEditor/Source/ChaosClothAssetDataflowNodes/Private/ChaosClothAsset/SkeletalMeshImportNode.cpp
--- a/Engine/Plugins/Experimental/ChaosClothAssetEditor/Source/ChaosClothAssetDataflowNodes/Private/ChaosClothAsset/SkeletalMeshImportNode.cpp
+++ b/Engine/Plugins/Experimental/ChaosClothAssetEditor/Source/ChaosClothAssetDataflowNodes/Private/ChaosClothAsset/SkeletalMeshImportNode.cpp
@@ -1,11 +1,10 @@
 // Copyright Epic Games, Inc. All Rights Reserved.
 
 #include "ChaosClothAsset/SkeletalMeshImportNode.h"
-#include "ChaosClothAsset/DataflowNodes.h"
-#include "ChaosClothAsset/CollectionClothFacade.h"
 #include "ChaosClothAsset/ClothDataflowTools.h"
+#include "ChaosClothAsset/CollectionClothFacade.h"
 #include "Animation/Skeleton.h"
-#include "BoneWeights.h"
+#include "Dataflow/DataflowInputOutput.h"
 #include "Engine/SkeletalMesh.h"
 #include "Engine/SkinnedAssetCommon.h"
 #include "GeometryCollection/ManagedArrayCollection.h"
@@ -26,32 +25,44 @@ FChaosClothAssetSkeletalMeshImportNode::FChaosClothAssetSkeletalMeshImportNode(c
 void FChaosClothAssetSkeletalMeshImportNode::Evaluate(Dataflow::FContext& Context, const FDataflowOutput* Out) const
 {
 	using namespace UE::Chaos::ClothAsset;
-	using namespace UE::Chaos::ClothAsset::DataflowNodes;
 
 	if (Out->IsA<FManagedArrayCollection>(&Collection))
 	{
 		const TSharedRef<FManagedArrayCollection> ClothCollection = MakeShared<FManagedArrayCollection>();
+		FCollectionClothFacade ClothFacade(ClothCollection);
+		ClothFacade.DefineSchema();
+
 		if (SkeletalMesh)
 		{
 			const FSkeletalMeshModel* ImportedModel = SkeletalMesh->GetImportedModel();
 			const bool bIsValidLOD = ImportedModel && ImportedModel->LODModels.IsValidIndex(LODIndex);
-			if (!ensureAlways(bIsValidLOD))
+			if (!bIsValidLOD)
 			{
+				FClothDataflowTools::LogAndToastWarning(*this, LOCTEXT("InvalidLODHeadline", "Invalid LOD."),
+					FText::Format(
+						LOCTEXT("InvalidLODDetails", "No valid LOD {0} found for skeletal mesh {1}."),
+						LODIndex,
+						FText::FromString(SkeletalMesh->GetName())));
+
 				SetValue(Context, MoveTemp(*ClothCollection), &Collection);
 				return;
 			}
-			const FSkeletalMeshLODModel &LODModel = ImportedModel->LODModels[LODIndex];
 
+			const FSkeletalMeshLODModel &LODModel = ImportedModel->LODModels[LODIndex];
 			const bool bIsValidSection = LODModel.Sections.IsValidIndex(SectionIndex);;
-			if (!ensureAlways(bIsValidSection))
+			if (!bIsValidSection)
 			{
+				FClothDataflowTools::LogAndToastWarning(*this, LOCTEXT("InvalidSectionHeadline", "Invalid section."),
+					FText::Format(
+						LOCTEXT("InvalidSectionDetails", "No valid section {0} found for skeletal mesh {1}."),
+						SectionIndex,
+						FText::FromString(SkeletalMesh->GetName())));
+
 				SetValue(Context, MoveTemp(*ClothCollection), &Collection);
 				return;
 			}
 
 			const FSkelMeshSection &Section = LODModel.Sections[SectionIndex];
-			FCollectionClothFacade ClothFacade(ClothCollection);
-			ClothFacade.DefineSchema();
 
 			if (bImportSimMesh)
 			{
diff --git a/Engine/Plugins/Experimental/ChaosClothAssetEditor/Source/ChaosClothAssetDataflowNodes/Private/ChaosClothAsset/StaticMeshImportNode.cpp b/Engine/Plugins/Experimental/ChaosClothAssetEditor/Source/ChaosClothAssetDataflowNodes/Private/ChaosClothAsset/StaticMeshImportNode.cpp
--- a/Engine/Plugins/Experimental/ChaosClothAssetEditor/Source/ChaosClothAssetDataflowNodes/Private/ChaosClothAsset/StaticMeshImportNode.cpp
+++ b/Engine/Plugins/Experimental/ChaosClothAssetEditor/Source/ChaosClothAssetDataflowNodes/Private/ChaosClothAsset/StaticMeshImportNode.cpp
@@ -1,11 +1,11 @@
 // Copyright Epic Games, Inc. All Rights Reserved.
 
 #include "ChaosClothAsset/StaticMeshImportNode.h"
-#include "ChaosClothAsset/DataflowNodes.h"
+#include "ChaosClothAsset/ClothDataflowTools.h"
 #include "ChaosClothAsset/CollectionClothFacade.h"
 #include "ChaosClothAsset/ClothAsset.h"
 #include "ChaosClothAsset/ClothGeometryTools.h"
-#include "ChaosClothAsset/ClothDataflowTools.h"
+#include "Dataflow/DataflowInputOutput.h"
 #include "Engine/StaticMesh.h"
 #include "Interfaces/ITargetPlatformManagerModule.h"
 #include "MeshDescriptionToDynamicMesh.h"
@@ -86,14 +86,14 @@ bool BuildSkeletalMeshModelFromMeshDescription(const FMeshDescription* const InM
 	{
 		for (const FText& Message : WarningMessages)
 		{
-			DataflowNodes::LogAndToastWarning(FText::Format(LOCTEXT("SkelMeshConvertWarningFmt", "{0}"), Message));
+			UE_LOG(LogChaosClothAssetDataflowNodes, Warning, TEXT("%s"), *Message.ToString());
 		}
 		return false;
 	}
 	return true;
 }
 
-void InitializeDataFromMeshDescription(
+bool InitializeDataFromMeshDescription(
 	const FMeshDescription* const InMeshDescription,
 	const FMeshBuildSettings& InBuildSettings,
 	const TArray<FStaticMaterial>& StaticMaterials,
@@ -108,7 +108,9 @@ void InitializeDataFromMeshDescription(
 			const FString RenderMaterialPathName = StaticMaterials[SectionIndex].MaterialInterface ? StaticMaterials[SectionIndex].MaterialInterface->GetPathName() : "";
 			FClothDataflowTools::AddRenderPatternFromSkeletalMeshSection(ClothCollection, SkeletalMeshModel, SectionIndex, RenderMaterialPathName);
 		}
+		return true;
 	}
+	return false;
 }
 
 } // namespace UE::Chaos::ClothAsset::Private
@@ -122,7 +124,7 @@ FChaosClothAssetStaticMeshImportNode::FChaosClothAssetStaticMeshImportNode(const
 void FChaosClothAssetStaticMeshImportNode::Evaluate(Dataflow::FContext& Context, const FDataflowOutput* Out) const
 {
 	using namespace UE::Chaos::ClothAsset;
-	using namespace UE::Chaos::ClothAsset::DataflowNodes;
+	using namespace UE::Chaos::ClothAsset::Private;
 
 	if (Out->IsA<FManagedArrayCollection>(&Collection))
 	{
@@ -136,7 +138,7 @@ void FChaosClothAssetStaticMeshImportNode::Evaluate(Dataflow::FContext& Context,
 		if (StaticMesh && (bImportAsSimMesh || bImportAsRenderMesh))
 		{
 			const int32 NumLods = StaticMesh->GetNumSourceModels();
-			if(LodIndex < NumLods)
+			if (LodIndex < NumLods)
 			{
 				const FMeshDescription* const MeshDescription = StaticMesh->GetMeshDescription(LodIndex);
 				check(MeshDescription);
@@ -151,12 +153,19 @@ void FChaosClothAssetStaticMeshImportNode::Evaluate(Dataflow::FContext& Context,
 
 					Converter.Convert(MeshDescription, DynamicMesh);
 					constexpr bool bAppend = false;
-					UE::Chaos::ClothAsset::FClothGeometryTools::BuildSimMeshFromDynamicMesh(ClothCollection, DynamicMesh, UVChannel, UVScale, bAppend);
+					FClothGeometryTools::BuildSimMeshFromDynamicMesh(ClothCollection, DynamicMesh, UVChannel, UVScale, bAppend);
 				}
 				if (bImportAsRenderMesh)
 				{
 					// Add render data into a single pattern for now
-					UE::Chaos::ClothAsset::Private::InitializeDataFromMeshDescription(MeshDescription, StaticMesh->GetSourceModel(LodIndex).BuildSettings, StaticMesh->GetStaticMaterials(), ClothCollection);
+					if (!InitializeDataFromMeshDescription(MeshDescription, StaticMesh->GetSourceModel(LodIndex).BuildSettings, StaticMesh->GetStaticMaterials(), ClothCollection))
+					{
+						FClothDataflowTools::LogAndToastWarning(*this,
+							LOCTEXT("InvalidRenderMeshHeadline", "Invalid render mesh."),
+							FText::Format(
+								LOCTEXT("InvalidRenderMeshDetails", "The input static mesh {0} failed to convert to a valid render mesh."),
+								FText::FromString(StaticMesh->GetName())));
+					}
 				}
 
 				// Set a default skeleton
@@ -164,8 +173,12 @@ void FChaosClothAssetStaticMeshImportNode::Evaluate(Dataflow::FContext& Context,
 			}
 			else
 			{
-				DataflowNodes::LogAndToastWarning(
-					FText::Format(LOCTEXT("FChaosClothAssetStaticMeshImportNode::InvalidLod", "FChaosClothAssetStaticMeshImportNode: Invalid LodIndex {0} >= Num Lods in static mesh {1}"),
+				FClothDataflowTools::LogAndToastWarning(*this,
+					LOCTEXT("InvalidLodIndexHeadline", "Invalid LOD index."),
+					FText::Format(LOCTEXT("InvalidLodIndexDetails",
+						"{0} is not a valid LOD index for static mesh {1}.\n"
+						"This static mesh has {2} LOD(s)."),
+						FText::FromString(StaticMesh->GetName()),
 						LodIndex,
 						NumLods));
 			}
diff --git a/Engine/Plugins/Experimental/ChaosClothAssetEditor/Source/ChaosClothAssetDataflowNodes/Private/ChaosClothAsset/TerminalNode.cpp b/Engine/Plugins/Experimental/ChaosClothAssetEditor/Source/ChaosClothAssetDataflowNodes/Private/ChaosClothAsset/TerminalNode.cpp
--- a/Engine/Plugins/Experimental/ChaosClothAssetEditor/Source/ChaosClothAssetDataflowNodes/Private/ChaosClothAsset/TerminalNode.cpp
+++ b/Engine/Plugins/Experimental/ChaosClothAssetEditor/Source/ChaosClothAssetDataflowNodes/Private/ChaosClothAsset/TerminalNode.cpp
@@ -1,12 +1,13 @@
 // Copyright Epic Games, Inc. All Rights Reserved.
 
 #include "ChaosClothAsset/TerminalNode.h"
-#include "ChaosClothAsset/DataflowNodes.h"
 #include "ChaosClothAsset/ClothAsset.h"
-#include "ChaosClothAsset/CollectionClothFacade.h"
+#include "ChaosClothAsset/ClothDataflowTools.h"
 #include "ChaosClothAsset/ClothGeometryTools.h"
-#include "Chaos/CollectionPropertyFacade.h"
+#include "ChaosClothAsset/CollectionClothFacade.h"
 #include "Animation/Skeleton.h"
+#include "Chaos/CollectionPropertyFacade.h"
+#include "Dataflow/DataflowInputOutput.h"
 #include "Engine/SkinnedAssetCommon.h"
 #include "Materials/Material.h"
 #include "PhysicsEngine/PhysicsAsset.h"
@@ -107,21 +108,21 @@ void FChaosClothAssetTerminalNode::SetAssetValue(TObjectPtr<UObject> Asset, Data
 				{
 					ClothFacade.Initialize(FCollectionClothConstFacade(ClothCollections[LastValidLodIndex]));
 
-					DataflowNodes::LogAndToastWarning(
+					FClothDataflowTools::LogAndToastWarning(*this,
+						LOCTEXT("InvalidInputLodNHeadline", "Invalid input LOD."),
 						FText::Format(
-							LOCTEXT(
-								"InvalidLodN",
-								"ClothAssetTerminal: Invalid or empty input LOD for LOD {0}.\n"
-								"Using previous valid LOD {1} instead."),
-						LodIndex,
-						LastValidLodIndex));
+							LOCTEXT("InvalidInputLodNDetails",
+								"Invalid or empty input LOD for LOD {0}.\n"
+								"Using the previous valid LOD {1} instead."),
+							LodIndex,
+							LastValidLodIndex));
 				}
 				else
 				{
-					DataflowNodes::LogAndToastWarning(
-						LOCTEXT(
-							"InvalidLod0",
-							"ClothAssetTerminal: Invalid or empty input LOD for LOD 0.\n"
+					FClothDataflowTools::LogAndToastWarning(*this,
+						LOCTEXT("InvalidInputLod0Headline", "Invalid input LOD 0."),
+						LOCTEXT("InvalidInputLod0Details",
+							"Invalid or empty input LOD for LOD 0.\n"
 							"LOD 0 cannot be empty in order to construct a valid Cloth Asset."));
 					break; // Empty cloth asset
 				}
diff --git a/Engine/Plugins/Experimental/ChaosClothAssetEditor/Source/ChaosClothAssetDataflowNodes/Private/ChaosClothAsset/TransferSkinWeightsNode.cpp b/Engine/Plugins/Experimental/ChaosClothAssetEditor/Source/ChaosClothAssetDataflowNodes/Private/ChaosClothAsset/TransferSkinWeightsNode.cpp
--- a/Engine/Plugins/Experimental/ChaosClothAssetEditor/Source/ChaosClothAssetDataflowNodes/Private/ChaosClothAsset/TransferSkinWeightsNode.cpp
+++ b/Engine/Plugins/Experimental/ChaosClothAssetEditor/Source/ChaosClothAssetDataflowNodes/Private/ChaosClothAsset/TransferSkinWeightsNode.cpp
@@ -1,45 +1,45 @@
 // Copyright Epic Games, Inc. All Rights Reserved.
 
 #include "ChaosClothAsset/TransferSkinWeightsNode.h"
-#include "ChaosClothAsset/DataflowNodes.h"
+#include "ChaosClothAsset/ClothDataflowTools.h"
 #include "ChaosClothAsset/CollectionClothFacade.h"
+#include "Animation/Skeleton.h"
+#include "Dataflow/DataflowInputOutput.h"
+#include "DynamicMesh/DynamicBoneAttribute.h"
+#include "DynamicMesh/DynamicMeshAttributeSet.h"
+#include "DynamicMesh/DynamicVertexSkinWeightsAttribute.h"
+#include "DynamicMesh/MeshNormals.h"
+#include "DynamicMesh/MeshTransforms.h"
+#include "Engine/SkeletalMesh.h"
 #include "MeshDescriptionToDynamicMesh.h"
+#include "Operations/TransferBoneWeights.h"
 #include "Rendering/SkeletalMeshLODImporterData.h"
 #include "Rendering/SkeletalMeshModel.h"
-#include "DynamicMesh/MeshTransforms.h"
-#include "Operations/TransferBoneWeights.h"
 #include "SkeletalMeshAttributes.h"
-#include "Engine/SkeletalMesh.h"
-#include "Animation/Skeleton.h"
-#include "DynamicMesh/DynamicVertexSkinWeightsAttribute.h"
-#include "DynamicMesh/DynamicMeshAttributeSet.h"
-#include "DynamicMesh/DynamicBoneAttribute.h"
-#include "DynamicMesh/MeshNormals.h"
 
 #include UE_INLINE_GENERATED_CPP_BY_NAME(TransferSkinWeightsNode)
 
 #define LOCTEXT_NAMESPACE "ChaosClothAssetTransferSkinWeightsNode"
 
-using namespace UE::Chaos::ClothAsset;
-using namespace UE::AnimationCore;
-using namespace UE::Geometry;
-
-namespace UE::ChaosClothAsset::Private
+namespace UE::Chaos::ClothAsset::Private
 {
-	void SkeletalMeshToDynamicMesh(USkeletalMesh* FromSkeletalMeshAsset, int32 SkeletalMeshLOD, FDynamicMesh3& ToDynamicMesh)
+	static void SkeletalMeshToDynamicMesh(USkeletalMesh* FromSkeletalMeshAsset, int32 LodIndex, FDynamicMesh3& ToDynamicMesh)
 	{	
 		FMeshDescription SourceMesh;
-		FromSkeletalMeshAsset->GetMeshDescription(SkeletalMeshLOD, SourceMesh);
+		FromSkeletalMeshAsset->GetMeshDescription(LodIndex, SourceMesh);
 		FMeshDescriptionToDynamicMesh Converter;
 		Converter.Convert(&SourceMesh, ToDynamicMesh);
 	}
 
 	// Convert the cloth into DynamicMesh
 	// @todo: This should instead be handled by a cloth lod to dynamic mesh converter similar to functions in the ClothPatternToDynamicMesh.h
-	bool SimClothToDynamicMesh(const FCollectionClothConstFacade& ClothFacade,
-								const FReferenceSkeleton& TargetRefSkeleton, // the reference skeleton to add to the dynamic mesh
-								FDynamicMesh3& WeldedSimMesh)
+	static bool SimClothToDynamicMesh(
+		const UE::Chaos::ClothAsset::FCollectionClothConstFacade& ClothFacade,
+		const FReferenceSkeleton& TargetRefSkeleton, // the reference skeleton to add to the dynamic mesh
+		UE::Geometry::FDynamicMesh3& WeldedSimMesh)
 	{
+		using namespace UE::Geometry;
+
 		// Convert the sim mesh to DynamicMesh. 
 		// @todo: FTransferBoneWeights should accept raw data arrays (vertices/triangles) to avoid this conversion
 		WeldedSimMesh.Clear();
@@ -73,8 +73,15 @@ namespace UE::ChaosClothAsset::Private
 		return true;
 	}
 
-	void TransferInpaintWeights(const FReferenceSkeleton& TargetRefSkeleton, const double NormalThreshold, const double RadiusPercentage, bool bUseParallel, FCollectionClothFacade& ClothFacade, FTransferBoneWeights& TransferBoneWeights)
+	static void TransferInpaintWeights(
+		const FReferenceSkeleton& TargetRefSkeleton,
+		const double NormalThreshold,
+		const double RadiusPercentage,
+		bool bUseParallel,
+		UE::Chaos::ClothAsset::FCollectionClothFacade& ClothFacade,
+		UE::Geometry::FTransferBoneWeights& TransferBoneWeights)
 	{
+		using namespace UE::Geometry;
 
 		//
 		// Convert cloth sim mesh LOD to the welded dynamic sim mesh.
@@ -82,28 +89,26 @@ namespace UE::ChaosClothAsset::Private
 		FDynamicMesh3 WeldedSimMesh;
 		if (!ensure(SimClothToDynamicMesh(ClothFacade, TargetRefSkeleton, WeldedSimMesh)))
 		{
-			UE::Chaos::ClothAsset::DataflowNodes::LogAndToastWarning(LOCTEXT("Warning_TransferWeightsFailedLodToDynamicMesh", "TransferSkinWeightsNode: Failed to weld the simulation mesh for LOD."));
+			UE_LOG(LogChaosClothAssetDataflowNodes, Warning, TEXT("TransferSkinWeightsNode: Failed to weld the simulation mesh for LOD."));
 			return;
 		}
 
-
 		//
 		// Transfer the weights from the body to the welded sim mesh.
 		//
 		TransferBoneWeights.NormalThreshold = FMathd::DegToRad * NormalThreshold;
 		TransferBoneWeights.SearchRadius = RadiusPercentage * WeldedSimMesh.GetBounds().DiagonalLength();
-		if (TransferBoneWeights.Validate() != EOperationValidationResult::Ok)
+		if (!ensure(TransferBoneWeights.Validate() == EOperationValidationResult::Ok))
 		{
-			UE::Chaos::ClothAsset::DataflowNodes::LogAndToastWarning(LOCTEXT("Warning_TransferWeightsInpaintWeightsInvalidParameters", "TransferSkinWeightsNode: Transfer method parameters are invalid."));
+			UE_LOG(LogChaosClothAssetDataflowNodes, Warning, TEXT("TransferSkinWeightsNode: Transfer method parameters are invalid."));
 			return;
 		}
-		if (!TransferBoneWeights.TransferWeightsToMesh(WeldedSimMesh, FSkeletalMeshAttributes::DefaultSkinWeightProfileName))
+		if (!ensure(TransferBoneWeights.TransferWeightsToMesh(WeldedSimMesh, FSkeletalMeshAttributes::DefaultSkinWeightProfileName)))
 		{
-			UE::Chaos::ClothAsset::DataflowNodes::LogAndToastWarning(LOCTEXT("Warning_TransferWeightsFailed", "TransferSkinWeightsNode: Transferring skin weights failed"));
+			UE_LOG(LogChaosClothAssetDataflowNodes, Warning, TEXT("TransferSkinWeightsNode: Transferring skin weights failed."));
 			return;
 		}
 
-
 		//
 		// Copy the new bone weight data from the welded sim mesh back to the cloth patterns.
 		//
@@ -118,7 +123,6 @@ namespace UE::ChaosClothAsset::Private
 				ClothFacade.GetSimBoneWeights()[WeldedID]);
 		}, bUseParallel ? EParallelForFlags::None : EParallelForFlags::ForceSingleThread);
 
-
 		//
 		// Compute the bone weights for the render mesh by transferring weights from the sim mesh
 		// @todo If render mesh eventually supports welding, we should be transferring weights from the 
@@ -136,8 +140,14 @@ namespace UE::ChaosClothAsset::Private
 		}, bUseParallel ? EParallelForFlags::None : EParallelForFlags::ForceSingleThread);
 	}
 
-	void TransferClosestPointOnSurface(const FReferenceSkeleton& TargetRefSkeleton, const bool bUseParallel, FCollectionClothFacade& ClothFacade, FTransferBoneWeights& TransferBoneWeights)
+	static void TransferClosestPointOnSurface(
+		const FReferenceSkeleton& TargetRefSkeleton,
+		const bool bUseParallel,
+		UE::Chaos::ClothAsset::FCollectionClothFacade& ClothFacade,
+		UE::Geometry::FTransferBoneWeights& TransferBoneWeights)
 	{
+		using namespace UE::Geometry;
+
 		//
 		// Compute the bone index mappings. This allows the transfer operator to retarget weights to the correct skeleton.
 		//
@@ -148,9 +158,9 @@ namespace UE::ChaosClothAsset::Private
 			TargetBoneToIndex.Add(TargetRefSkeleton.GetRawRefBoneInfo()[BoneIdx].Name, BoneIdx);
 		}
 
-		if (TransferBoneWeights.Validate() != EOperationValidationResult::Ok)
+		if (!ensure(TransferBoneWeights.Validate() == EOperationValidationResult::Ok))
 		{
-			UE::Chaos::ClothAsset::DataflowNodes::LogAndToastWarning(LOCTEXT("Warning_TransferWeightsClosestPointOnSurfaceInvalidParameters", "TransferSkinWeightsNode: Transfer method parameters are invalid."));
+			UE_LOG(LogChaosClothAssetDataflowNodes, Warning, TEXT("TransferSkinWeightsNode: Transfer method parameters are invalid."));
 			return;
 		}
 
@@ -187,25 +197,30 @@ FChaosClothAssetTransferSkinWeightsNode::FChaosClothAssetTransferSkinWeightsNode
 {
 	RegisterInputConnection(&Collection);
 	RegisterOutputConnection(&Collection, &Collection);
-	RegisterInputConnection(&SkeletalMesh);
 }
 
 void FChaosClothAssetTransferSkinWeightsNode::Evaluate(Dataflow::FContext& Context, const FDataflowOutput* Out) const
 {
-	using namespace UE::ChaosClothAsset::Private;
+	using namespace UE::Chaos::ClothAsset;
+	using namespace UE::Chaos::ClothAsset::Private;
+	using namespace UE::Geometry;
 
 	if (Out->IsA<FManagedArrayCollection>(&Collection))
 	{
 		// Evaluate inputs
 		FManagedArrayCollection InputCollection = GetValue<FManagedArrayCollection>(Context, &Collection);
 		const TSharedRef<FManagedArrayCollection> ClothCollection = MakeShared<FManagedArrayCollection>(MoveTemp(InputCollection));
-		const TObjectPtr<USkeletalMesh> InputSkeletalMesh = GetValue<TObjectPtr<USkeletalMesh>>(Context, &SkeletalMesh);
 
-		if (InputSkeletalMesh)
+		if (SkeletalMesh && FCollectionClothFacade(ClothCollection).IsValid())  // Can only act on the collection if it is a valid cloth collection
 		{
-			if (!InputSkeletalMesh->IsValidLODIndex(SkeletalMeshLOD))
+			if (!SkeletalMesh->IsValidLODIndex(LodIndex))
 			{
-				UE::Chaos::ClothAsset::DataflowNodes::LogAndToastWarning(LOCTEXT("Warning_TransferWeightsInvalidSkeletalMeshLOD", "TransferSkinWeightsNode: The specified input SkeletalMesh LOD is not valid."));
+				FClothDataflowTools::LogAndToastWarning(*this,
+					LOCTEXT("InvalidLodIndexHeadline", "Invalid LOD Index."),
+					FText::Format(
+						LOCTEXT("InvalidLodIndexDetails", "LOD index {0} is not a valid LOD for skeletal mesh {1}."),
+						LodIndex,
+						FText::FromString(SkeletalMesh.GetName())));
 				SetValue(Context, MoveTemp(*ClothCollection), &Collection);
 				return;
 			}
@@ -214,10 +229,9 @@ void FChaosClothAssetTransferSkinWeightsNode::Evaluate(Dataflow::FContext& Conte
 			// Convert source Skeletal Mesh to Dynamic Mesh.
 			//
 			FDynamicMesh3 SourceDynamicMesh;
-			SkeletalMeshToDynamicMesh(InputSkeletalMesh, SkeletalMeshLOD, SourceDynamicMesh);
+			SkeletalMeshToDynamicMesh(SkeletalMesh, LodIndex, SourceDynamicMesh);
 			MeshTransforms::ApplyTransform(SourceDynamicMesh, Transform, true);
-			const FReferenceSkeleton& TargetRefSkeleton = InputSkeletalMesh->GetSkeleton()->GetReferenceSkeleton();
-
+			const FReferenceSkeleton& TargetRefSkeleton = SkeletalMesh->GetSkeleton()->GetReferenceSkeleton();
 
 			//
 			// Setup the bone weight transfer operator for the source mesh.
@@ -227,12 +241,11 @@ void FChaosClothAssetTransferSkinWeightsNode::Evaluate(Dataflow::FContext& Conte
 			TransferBoneWeights.bUseParallel = bUseParallel;
 			TransferBoneWeights.TransferMethod = static_cast<FTransferBoneWeights::ETransferBoneWeightsMethod>(TransferMethod);
 
-
 			//
 			// Transfer the bone weights from the source Skeletal mesh to the Cloth asset.
 			//
 			FCollectionClothFacade ClothFacade(ClothCollection);
-			ClothFacade.SetSkeletonAssetPathName(InputSkeletalMesh->GetSkeleton()->GetPathName());
+			ClothFacade.SetSkeletonAssetPathName(SkeletalMesh->GetSkeleton()->GetPathName());
 
 			if (TransferMethod == EChaosClothAssetTransferSkinWeightsMethod::InpaintWeights)
 			{
diff --git a/Engine/Plugins/Experimental/ChaosClothAssetEditor/Source/ChaosClothAssetDataflowNodes/Public/ChaosClothAsset/AddWeightMapNode.h b/Engine/Plugins/Experimental/ChaosClothAssetEditor/Source/ChaosClothAssetDataflowNodes/Public/ChaosClothAsset/AddWeightMapNode.h
--- a/Engine/Plugins/Experimental/ChaosClothAssetEditor/Source/ChaosClothAssetDataflowNodes/Public/ChaosClothAsset/AddWeightMapNode.h
+++ b/Engine/Plugins/Experimental/ChaosClothAssetEditor/Source/ChaosClothAssetDataflowNodes/Public/ChaosClothAsset/AddWeightMapNode.h
@@ -6,26 +6,27 @@
 #include "GeometryCollection/ManagedArrayCollection.h"
 #include "AddWeightMapNode.generated.h"
 
+/** Painted weight map attributes node. */
 USTRUCT(Meta = (DataflowCloth))
 struct FChaosClothAssetAddWeightMapNode : public FDataflowNode
 {
 	GENERATED_USTRUCT_BODY()
 	DATAFLOW_NODE_DEFINE_INTERNAL(FChaosClothAssetAddWeightMapNode, "AddWeightMap", "Cloth", "Cloth Add Weight Map")
-	//DATAFLOW_NODE_RENDER_TYPE(FGeometryCollection::StaticType(), "Collection")  // TODO: Leave out the render type until there is something to render
 
 public:
 
-	UPROPERTY(Meta = (Dataflowinput, DataflowOutput, DisplayName = "Collection", DataflowPassthrough = "Collection"))
+	UPROPERTY(Meta = (Dataflowinput, DataflowOutput, DataflowPassthrough = "Collection"))
 	FManagedArrayCollection Collection;
 
-	UPROPERTY(EditAnywhere, Category = "Weight Map", Meta = (Dataflowinput, DataflowOutput, DisplayName = "Name", DataflowPassthrough = "Name"))
+	/** The name to be set as a weight map attribute. */
+	UPROPERTY(EditAnywhere, Category = "Add Weight Map", Meta = (DataflowOutput))
 	FString Name;
 
-	UPROPERTY(EditAnywhere, Category = "Weight Map")
+	UPROPERTY()
 	TArray<float> VertexWeights;
 
-
 	FChaosClothAssetAddWeightMapNode(const Dataflow::FNodeParameters& InParam, FGuid InGuid = FGuid::NewGuid());
 
+private:
 	virtual void Evaluate(Dataflow::FContext& Context, const FDataflowOutput* Out) const override;
 };
diff --git a/Engine/Plugins/Experimental/ChaosClothAssetEditor/Source/ChaosClothAssetDataflowNodes/Public/ChaosClothAsset/BindToRootBoneNode.h b/Engine/Plugins/Experimental/ChaosClothAssetEditor/Source/ChaosClothAssetDataflowNodes/Public/ChaosClothAsset/BindToRootBoneNode.h
--- a/Engine/Plugins/Experimental/ChaosClothAssetEditor/Source/ChaosClothAssetDataflowNodes/Public/ChaosClothAsset/BindToRootBoneNode.h
+++ b/Engine/Plugins/Experimental/ChaosClothAssetEditor/Source/ChaosClothAssetDataflowNodes/Public/ChaosClothAsset/BindToRootBoneNode.h
@@ -6,16 +6,15 @@
 #include "GeometryCollection/ManagedArrayCollection.h"
 #include "BindToRootBoneNode.generated.h"
 
+/** Bind an entire mesh to the single root bone of the current skeleton set on the cloth collection. */
 USTRUCT(Meta = (DataflowCloth))
 struct FChaosClothAssetBindToRootBoneNode : public FDataflowNode
 {
 	GENERATED_USTRUCT_BODY()
 	DATAFLOW_NODE_DEFINE_INTERNAL(FChaosClothAssetBindToRootBoneNode, "BindToRootBone", "Cloth", "Cloth Bind Skinning Weights To Root Bone")
-	//DATAFLOW_NODE_RENDER_TYPE(FGeometryCollection::StaticType(), "Collection")  // TODO: Leave out the render type until there is something to render
 
 public:
-
-	UPROPERTY(Meta = (Dataflowinput, DataflowOutput, DisplayName = "Collection", DataflowPassthrough = "Collection"))
+	UPROPERTY(Meta = (Dataflowinput, DataflowOutput, DataflowPassthrough = "Collection"))
 	FManagedArrayCollection Collection;
 
 	/** Whether to bind the simulation mesh. */
@@ -28,5 +27,6 @@ public:
 
 	FChaosClothAssetBindToRootBoneNode(const Dataflow::FNodeParameters& InParam, FGuid InGuid = FGuid::NewGuid());
 
+private:
 	virtual void Evaluate(Dataflow::FContext& Context, const FDataflowOutput* Out) const override;
 };
diff --git a/Engine/Plugins/Experimental/ChaosClothAssetEditor/Source/ChaosClothAssetDataflowNodes/Public/ChaosClothAsset/CopySimulationToRenderMeshNode.h b/Engine/Plugins/Experimental/ChaosClothAssetEditor/Source/ChaosClothAssetDataflowNodes/Public/ChaosClothAsset/CopySimulationToRenderMeshNode.h
--- a/Engine/Plugins/Experimental/ChaosClothAssetEditor/Source/ChaosClothAssetDataflowNodes/Public/ChaosClothAsset/CopySimulationToRenderMeshNode.h
+++ b/Engine/Plugins/Experimental/ChaosClothAssetEditor/Source/ChaosClothAssetDataflowNodes/Public/ChaosClothAsset/CopySimulationToRenderMeshNode.h
@@ -8,31 +8,28 @@
 
 class UMaterialInterface;
 
+/** Copy the simulation mesh to the render mesh to be able to render the simulation mesh, or when not using a different mesh for rendering. */
 USTRUCT(Meta = (DataflowCloth))
 struct FChaosClothAssetCopySimulationToRenderMeshNode : public FDataflowNode
 {
 	GENERATED_USTRUCT_BODY()
 	DATAFLOW_NODE_DEFINE_INTERNAL(FChaosClothAssetCopySimulationToRenderMeshNode, "CopySimulationToRenderMesh", "Cloth", "Cloth Simulation Render Mesh")
-	//DATAFLOW_NODE_RENDER_TYPE(FGeometryCollection::StaticType(), "Collection")  // TODO: Leave out the render type until there is something to render
 
 public:
 
-	UPROPERTY(Meta = (Dataflowinput, DataflowOutput, DisplayName = "Collection", DataflowPassthrough = "Collection"))
+	UPROPERTY(Meta = (Dataflowinput, DataflowOutput, DataflowPassthrough = "Collection"))
 	FManagedArrayCollection Collection;
 
-	/** List of patterns to apply the operation on. All patterns will be used if left empty. */
-	UPROPERTY(EditAnywhere, Category = "Pattern Selection")
-	TArray<int32> Patterns;
-
 	/** New material for the render mesh. */
-	UPROPERTY(EditAnywhere, Category = "Render Mesh")
+	UPROPERTY(EditAnywhere, Category = "Copy Simulation To Render Mesh")
 	TObjectPtr<const UMaterialInterface> Material;
 
-	/** Generate a single render pattern rather than a render pattern per sim pattern */
-	UPROPERTY(EditAnywhere, Category = "Render Mesh")
+	/** Generate a single render pattern rather than a render pattern per sim pattern. */
+	UPROPERTY(EditAnywhere, Category = "Copy Simulation To Render Mesh")
 	bool bGenerateSingleRenderPattern = false;
 
 	FChaosClothAssetCopySimulationToRenderMeshNode(const Dataflow::FNodeParameters& InParam, FGuid InGuid = FGuid::NewGuid());
 
+private:
 	virtual void Evaluate(Dataflow::FContext& Context, const FDataflowOutput* Out) const override;
 };
diff --git a/Engine/Plugins/Experimental/ChaosClothAssetEditor/Source/ChaosClothAssetDataflowNodes/Public/ChaosClothAsset/DataflowNodes.h b/Engine/Plugins/Experimental/ChaosClothAssetEditor/Source/ChaosClothAssetDataflowNodes/Public/ChaosClothAsset/DataflowNodes.h
deleted file mode 100644
diff --git a/Engine/Plugins/Experimental/ChaosClothAssetEditor/Source/ChaosClothAssetDataflowNodes/Public/ChaosClothAsset/DatasmithImportNode.h b/Engine/Plugins/Experimental/ChaosClothAssetEditor/Source/ChaosClothAssetDataflowNodes/Public/ChaosClothAsset/DatasmithImportNode.h
--- a/Engine/Plugins/Experimental/ChaosClothAssetEditor/Source/ChaosClothAssetDataflowNodes/Public/ChaosClothAsset/DatasmithImportNode.h
+++ b/Engine/Plugins/Experimental/ChaosClothAssetEditor/Source/ChaosClothAssetDataflowNodes/Public/ChaosClothAsset/DatasmithImportNode.h
@@ -7,6 +7,7 @@
 #include "Misc/SecureHash.h"
 #include "DatasmithImportNode.generated.h"
 
+/** Import a file from a third party garment construction package compatible with the Datasmith scene format. */
 USTRUCT(meta = (DataflowCloth))
 struct FChaosClothAssetDatasmithImportNode : public FDataflowNode
 {
@@ -14,7 +15,7 @@ struct FChaosClothAssetDatasmithImportNode : public FDataflowNode
 	DATAFLOW_NODE_DEFINE_INTERNAL(FChaosClothAssetDatasmithImportNode, "DatasmithImport", "Cloth", "Cloth Datasmith Import")
 
 public:
-	UPROPERTY(meta = (DataflowOutput, DisplayName = "Collection"))
+	UPROPERTY(meta = (DataflowOutput))
 	FManagedArrayCollection Collection;
 
 	/** Path of the file to import using any available Datasmith cloth translator. */
diff --git a/Engine/Plugins/Experimental/ChaosClothAssetEditor/Source/ChaosClothAssetDataflowNodes/Public/ChaosClothAsset/DeleteRenderMeshNode.h b/Engine/Plugins/Experimental/ChaosClothAssetEditor/Source/ChaosClothAssetDataflowNodes/Public/ChaosClothAsset/DeleteRenderMeshNode.h
--- a/Engine/Plugins/Experimental/ChaosClothAssetEditor/Source/ChaosClothAssetDataflowNodes/Public/ChaosClothAsset/DeleteRenderMeshNode.h
+++ b/Engine/Plugins/Experimental/ChaosClothAssetEditor/Source/ChaosClothAssetDataflowNodes/Public/ChaosClothAsset/DeleteRenderMeshNode.h
@@ -6,23 +6,19 @@
 #include "GeometryCollection/ManagedArrayCollection.h"
 #include "DeleteRenderMeshNode.generated.h"
 
+/** Delete the current render mesh stored on the cloth collection. */
 USTRUCT(Meta = (DataflowCloth))
 struct FChaosClothAssetDeleteRenderMeshNode : public FDataflowNode
 {
 	GENERATED_USTRUCT_BODY()
 	DATAFLOW_NODE_DEFINE_INTERNAL(FChaosClothAssetDeleteRenderMeshNode, "DeleteRenderMesh", "Cloth", "Cloth Delete Render Mesh")
-	//DATAFLOW_NODE_RENDER_TYPE(FGeometryCollection::StaticType(), "Collection")  // TODO: Leave out the render type until there is something to render
 
 public:
-
-	UPROPERTY(Meta = (Dataflowinput, DataflowOutput, DisplayName = "Collection", DataflowPassthrough = "Collection"))
+	UPROPERTY(Meta = (Dataflowinput, DataflowOutput, DataflowPassthrough = "Collection"))
 	FManagedArrayCollection Collection;
 
-	/** List of patterns to apply the operation on. All patterns will be used if left empty. */
-	UPROPERTY(EditAnywhere, Category = "Pattern Selection")
-	TArray<int32> Patterns;
-
 	FChaosClothAssetDeleteRenderMeshNode(const Dataflow::FNodeParameters& InParam, FGuid InGuid = FGuid::NewGuid());
 
+private:
 	virtual void Evaluate(Dataflow::FContext& Context, const FDataflowOutput* Out) const override;
 };
diff --git a/Engine/Plugins/Experimental/ChaosClothAssetEditor/Source/ChaosClothAssetDataflowNodes/Public/ChaosClothAsset/ImportNode.h b/Engine/Plugins/Experimental/ChaosClothAssetEditor/Source/ChaosClothAssetDataflowNodes/Public/ChaosClothAsset/ImportNode.h
--- a/Engine/Plugins/Experimental/ChaosClothAssetEditor/Source/ChaosClothAssetDataflowNodes/Public/ChaosClothAsset/ImportNode.h
+++ b/Engine/Plugins/Experimental/ChaosClothAssetEditor/Source/ChaosClothAssetDataflowNodes/Public/ChaosClothAsset/ImportNode.h
@@ -8,27 +8,27 @@
 
 class UChaosClothAsset;
 
+/** Import an existing Cloth Asset into the graph. */
 USTRUCT(Meta = (DataflowCloth))
 struct FChaosClothAssetImportNode : public FDataflowNode
 {
 	GENERATED_USTRUCT_BODY()
 	DATAFLOW_NODE_DEFINE_INTERNAL(FChaosClothAssetImportNode, "ClothAssetImport", "Cloth", "Cloth Asset Import")
-	//DATAFLOW_NODE_RENDER_TYPE(FGeometryCollection::StaticType(), "Collection")  // TODO: Leave out the render type until there is something to render
 
 public:
-
-	UPROPERTY(Meta = (DataflowOutput, DisplayName = "Collection"))
+	UPROPERTY(Meta = (DataflowOutput))
 	FManagedArrayCollection Collection;
 
 	/** The Cloth Asset to import into a collection. */
-	UPROPERTY(EditAnywhere, Category = "ClothAsset")
+	UPROPERTY(EditAnywhere, Category = "Cloth Asset Import")
 	TObjectPtr<const UChaosClothAsset> ClothAsset;
 
 	/** The LOD to import into the collection. Only one LOD can be imported at a time. */
-	UPROPERTY(EditAnywhere, Category = "ClothAsset", Meta = (DisplayName = "Import LOD"))
+	UPROPERTY(EditAnywhere, Category = "Cloth Asset Import", Meta = (DisplayName = "Import LOD", ClampMin = "0"))
 	int32 ImportLod = 0;
 
 	FChaosClothAssetImportNode(const Dataflow::FNodeParameters& InParam, FGuid InGuid = FGuid::NewGuid());
 
+private:
 	virtual void Evaluate(Dataflow::FContext& Context, const FDataflowOutput* Out) const override;
 };
diff --git a/Engine/Plugins/Experimental/ChaosClothAssetEditor/Source/ChaosClothAssetDataflowNodes/Public/ChaosClothAsset/MergeClothCollectionsNode.h b/Engine/Plugins/Experimental/ChaosClothAssetEditor/Source/ChaosClothAssetDataflowNodes/Public/ChaosClothAsset/MergeClothCollectionsNode.h
--- a/Engine/Plugins/Experimental/ChaosClothAssetEditor/Source/ChaosClothAssetDataflowNodes/Public/ChaosClothAsset/MergeClothCollectionsNode.h
+++ b/Engine/Plugins/Experimental/ChaosClothAssetEditor/Source/ChaosClothAssetDataflowNodes/Public/ChaosClothAsset/MergeClothCollectionsNode.h
@@ -6,6 +6,7 @@
 #include "GeometryCollection/ManagedArrayCollection.h"
 #include "MergeClothCollectionsNode.generated.h"
 
+/** Merge multiple cloth collections into a single cloth collection of multiple patterns. */
 USTRUCT(Meta = (DataflowCloth))
 struct FChaosClothAssetMergeClothCollectionsNode : public FDataflowNode
 {
@@ -13,7 +14,7 @@ struct FChaosClothAssetMergeClothCollectionsNode : public FDataflowNode
 	DATAFLOW_NODE_DEFINE_INTERNAL(FChaosClothAssetMergeClothCollectionsNode, "MergeClothCollections", "Cloth", "Cloth Merge Collection")
 
 public:
-	static constexpr int32 MaxInputs = 6;  // Hardcoded number of nputs since it is currently not possible to use arrays for optional inputs
+	static constexpr int32 MaxInputs = 6;  // Hardcoded number of inputs since it is currently not possible to use arrays for optional inputs
 
 	/** Input 0, right click on the node and add pins to add more merge inputs. */
 	UPROPERTY(Meta = (DataflowInput, DataflowOutput, DataflowPassthrough = "Collection"))
@@ -39,14 +40,15 @@ public:
 
 	FChaosClothAssetMergeClothCollectionsNode(const Dataflow::FNodeParameters& InParam, FGuid InGuid = FGuid::NewGuid());
 
+private:
+	//~ Begin FDataflowNode interface
 	virtual void Evaluate(Dataflow::FContext& Context, const FDataflowOutput* Out) const override;
-
 	virtual Dataflow::FPin AddPin() override;
 	virtual bool CanAddPin() const override { return NumInputs < MaxInputs; }
 	virtual bool CanRemovePin() const override { return NumInputs > 1; }
 	virtual Dataflow::FPin RemovePin() override;
-	virtual void Serialize(FArchive& Ar);
+	virtual void Serialize(FArchive& Ar) override;
+	//~ End FDataflowNode interface
 
-private:
 	TArray<const FManagedArrayCollection*> GetCollections() const;
 };
diff --git a/Engine/Plugins/Experimental/ChaosClothAssetEditor/Source/ChaosClothAssetDataflowNodes/Public/ChaosClothAsset/ReverseNormalsNode.h b/Engine/Plugins/Experimental/ChaosClothAssetEditor/Source/ChaosClothAssetDataflowNodes/Public/ChaosClothAsset/ReverseNormalsNode.h
--- a/Engine/Plugins/Experimental/ChaosClothAssetEditor/Source/ChaosClothAssetDataflowNodes/Public/ChaosClothAsset/ReverseNormalsNode.h
+++ b/Engine/Plugins/Experimental/ChaosClothAssetEditor/Source/ChaosClothAssetDataflowNodes/Public/ChaosClothAsset/ReverseNormalsNode.h
@@ -6,43 +6,43 @@
 #include "GeometryCollection/ManagedArrayCollection.h"
 #include "ReverseNormalsNode.generated.h"
 
+/** Reverse the geometry's normals or/and winding order of the simulation or/and render meshes stored in the cloth collection. */
 USTRUCT(Meta = (DataflowCloth))
 struct FChaosClothAssetReverseNormalsNode : public FDataflowNode
 {
 	GENERATED_USTRUCT_BODY()
 	DATAFLOW_NODE_DEFINE_INTERNAL(FChaosClothAssetReverseNormalsNode, "ReverseNormals", "Cloth", "Cloth Reverse Simulation Render Mesh Normals")
-	//DATAFLOW_NODE_RENDER_TYPE(FGeometryCollection::StaticType(), "Collection")  // TODO: Leave out the render type until there is something to render
 
 public:
-
-	UPROPERTY(Meta = (Dataflowinput, DataflowOutput, DisplayName = "Collection", DataflowPassthrough = "Collection"))
+	UPROPERTY(Meta = (Dataflowinput, DataflowOutput, DataflowPassthrough = "Collection"))
 	FManagedArrayCollection Collection;
 
 	/** List of sim patterns to apply the operation on. All patterns will be used if left empty. */
-	UPROPERTY(EditAnywhere, Category = "Reverse Normals", Meta = (EditCondition = "false"))
+	UPROPERTY(EditAnywhere, Category = "Reverse Normals|Simulation Mesh")
 	TArray<int32> SimPatterns;
 
-	/** List of render patterns to apply the operation on. All patterns will be used if left empty. */
-	UPROPERTY(EditAnywhere, Category = "Reverse Normals", Meta = (EditCondition = "false"))
-	TArray<int32> RenderPatterns;
-
 	/** Whether to reverse the simulation mesh normals. */
-	UPROPERTY(EditAnywhere, Category = "Reverse Normals")
+	UPROPERTY(EditAnywhere, Category = "Reverse Normals|Simulation Mesh")
 	bool bReverseSimMeshNormals = true;
 
 	/** Whether to reverse the simulation mesh triangles' winding order. */
-	UPROPERTY(EditAnywhere, Category = "Reverse Normals")
+	UPROPERTY(EditAnywhere, Category = "Reverse Normals|Simulation Mesh")
 	bool bReverseSimMeshWindingOrder = false;
 
+	/** List of render patterns to apply the operation on. All patterns will be used if left empty. */
+	UPROPERTY(EditAnywhere, Category = "Reverse Normals|Render Mesh")
+	TArray<int32> RenderPatterns;
+
 	/** Whether to reverse the render mesh normals. */
-	UPROPERTY(EditAnywhere, Category = "Reverse Normals")
+	UPROPERTY(EditAnywhere, Category = "Reverse Normals|Render Mesh")
 	bool bReverseRenderMeshNormals = true;
 
 	/** Whether to reverse the render mesh triangles' winding order. */
-	UPROPERTY(EditAnywhere, Category = "Reverse Normals")
+	UPROPERTY(EditAnywhere, Category = "Reverse Normals|Render Mesh")
 	bool bReverseRenderMeshWindingOrder = false;
 
 	FChaosClothAssetReverseNormalsNode(const Dataflow::FNodeParameters& InParam, FGuid InGuid = FGuid::NewGuid());
 
+private:
 	virtual void Evaluate(Dataflow::FContext& Context, const FDataflowOutput* Out) const override;
 };
diff --git a/Engine/Plugins/Experimental/ChaosClothAssetEditor/Source/ChaosClothAssetDataflowNodes/Public/ChaosClothAsset/SetPhysicsAssetNode.h b/Engine/Plugins/Experimental/ChaosClothAssetEditor/Source/ChaosClothAssetDataflowNodes/Public/ChaosClothAsset/SetPhysicsAssetNode.h
--- a/Engine/Plugins/Experimental/ChaosClothAssetEditor/Source/ChaosClothAssetDataflowNodes/Public/ChaosClothAsset/SetPhysicsAssetNode.h
+++ b/Engine/Plugins/Experimental/ChaosClothAssetEditor/Source/ChaosClothAssetDataflowNodes/Public/ChaosClothAsset/SetPhysicsAssetNode.h
@@ -8,6 +8,7 @@
 
 class UPhysicsAsset;
 
+/** Replace the current physics assets to collide the simulation mesh against. */
 USTRUCT(Meta = (DataflowCloth))
 struct FChaosClothAssetSetPhysicsAssetNode : public FDataflowNode
 {
@@ -15,15 +16,15 @@ struct FChaosClothAssetSetPhysicsAssetNode : public FDataflowNode
 	DATAFLOW_NODE_DEFINE_INTERNAL(FChaosClothAssetSetPhysicsAssetNode, "SetPhysicsAsset", "Cloth", "Cloth Set Physics Asset")
 
 public:
-
-	UPROPERTY(Meta = (Dataflowinput, DataflowOutput, DisplayName = "Collection", DataflowPassthrough = "Collection"))
+	UPROPERTY(Meta = (Dataflowinput, DataflowOutput, DataflowPassthrough = "Collection"))
 	FManagedArrayCollection Collection;
 
 	/** The physics asset to assign to the Cloth Collection. */
-	UPROPERTY(EditAnywhere, Category = "Physics Asset")
+	UPROPERTY(EditAnywhere, Category = "Set Physics Asset")
 	TObjectPtr<UPhysicsAsset> PhysicsAsset;
 
 	FChaosClothAssetSetPhysicsAssetNode(const Dataflow::FNodeParameters& InParam, FGuid InGuid = FGuid::NewGuid());
 
+private:
 	virtual void Evaluate(Dataflow::FContext& Context, const FDataflowOutput* Out) const override;
 };
diff --git a/Engine/Plugins/Experimental/ChaosClothAssetEditor/Source/ChaosClothAssetDataflowNodes/Public/ChaosClothAsset/SimulationBaseConfigNode.h b/Engine/Plugins/Experimental/ChaosClothAssetEditor/Source/ChaosClothAssetDataflowNodes/Public/ChaosClothAsset/SimulationBaseConfigNode.h
--- a/Engine/Plugins/Experimental/ChaosClothAssetEditor/Source/ChaosClothAssetDataflowNodes/Public/ChaosClothAsset/SimulationBaseConfigNode.h
+++ b/Engine/Plugins/Experimental/ChaosClothAssetEditor/Source/ChaosClothAssetDataflowNodes/Public/ChaosClothAsset/SimulationBaseConfigNode.h
@@ -21,16 +21,16 @@ struct FChaosClothAssetSimulationBaseConfigNode : public FDataflowNode
 	GENERATED_USTRUCT_BODY()
 
 public:
-	UPROPERTY(Meta = (Dataflowinput, DataflowOutput, DisplayName = "Collection", DataflowPassthrough = "Collection"))
+	UPROPERTY(Meta = (Dataflowinput, DataflowOutput, DataflowPassthrough = "Collection"))
 	FManagedArrayCollection Collection;
 
 	FChaosClothAssetSimulationBaseConfigNode() = default;
 
 	FChaosClothAssetSimulationBaseConfigNode(const Dataflow::FNodeParameters& InParam, FGuid InGuid = FGuid::NewGuid());
 
+protected:
 	virtual void Evaluate(Dataflow::FContext& Context, const FDataflowOutput* Out) const override;
 
-protected:
 	virtual void AddProperties(::Chaos::Softs::FCollectionPropertyMutableFacade& Properties) const
 	PURE_VIRTUAL(FChaosClothAssetSimulationBaseConfigNode::AddProperties, );
 
diff --git a/Engine/Plugins/Experimental/ChaosClothAssetEditor/Source/ChaosClothAssetDataflowNodes/Public/ChaosClothAsset/SimulationDefaultConfigNode.h b/Engine/Plugins/Experimental/ChaosClothAssetEditor/Source/ChaosClothAssetDataflowNodes/Public/ChaosClothAsset/SimulationDefaultConfigNode.h
--- a/Engine/Plugins/Experimental/ChaosClothAssetEditor/Source/ChaosClothAssetDataflowNodes/Public/ChaosClothAsset/SimulationDefaultConfigNode.h
+++ b/Engine/Plugins/Experimental/ChaosClothAssetEditor/Source/ChaosClothAssetDataflowNodes/Public/ChaosClothAsset/SimulationDefaultConfigNode.h
@@ -8,33 +8,33 @@
 
 class UChaosClothConfig;
 class UChaosClothSharedSimConfig;
-class UMaterialInterface;
 
+/** Add default simulation properties to the cloth collection in the format of the skeletal mesh cloth editor. */
 USTRUCT(Meta = (DataflowCloth))
 struct FChaosClothAssetSimulationDefaultConfigNode : public FDataflowNode, public FGCObject
 {
 	GENERATED_USTRUCT_BODY()
 	DATAFLOW_NODE_DEFINE_INTERNAL(FChaosClothAssetSimulationDefaultConfigNode, "SimulationDefaultConfig", "Cloth", "Cloth Simulation Default Config")
-	//DATAFLOW_NODE_RENDER_TYPE(FManagedArrayCollection::StaticType(), "Collection")  // TODO: Leave out the render type until there is something to render
 
 public:
-
-	UPROPERTY(Meta = (Dataflowinput, DataflowOutput, DisplayName = "Collection", DataflowPassthrough = "Collection"))
+	UPROPERTY(Meta = (Dataflowinput, DataflowOutput, DataflowPassthrough = "Collection"))
 	FManagedArrayCollection Collection;
 
 	/** Cloth Simulation Properties. */
-	UPROPERTY(EditAnywhere, Instanced, NoClear, Category = "Dataflow")
+	UPROPERTY(EditAnywhere, Instanced, NoClear, Category = "Simulation Default Config")
 	TObjectPtr<UChaosClothConfig> SimulationConfig;
 
 	/** Cloth Shared Simulation Properties. */
-	UPROPERTY(EditAnywhere, Instanced, NoClear, Category = "Dataflow")
+	UPROPERTY(EditAnywhere, Instanced, NoClear, Category = "Simulation Default Config")
 	TObjectPtr<UChaosClothSharedSimConfig> SharedSimulationConfig;
 
 	FChaosClothAssetSimulationDefaultConfigNode(const Dataflow::FNodeParameters& InParam, FGuid InGuid = FGuid::NewGuid());
 
-	virtual void Serialize(FArchive& Ar) override;
-
+private:
+	//~ Begin FDataflowNode Interface
 	virtual void Evaluate(Dataflow::FContext& Context, const FDataflowOutput* Out) const override;
+	virtual void Serialize(FArchive& Ar) override;
+	//~ End FDataflowNode Interface
 
 	//~ Begin FGCObject Interface
 	virtual void AddReferencedObjects(FReferenceCollector& Collector) override;
diff --git a/Engine/Plugins/Experimental/ChaosClothAssetEditor/Source/ChaosClothAssetDataflowNodes/Public/ChaosClothAsset/SimulationGravityConfigNode.h b/Engine/Plugins/Experimental/ChaosClothAssetEditor/Source/ChaosClothAssetDataflowNodes/Public/ChaosClothAsset/SimulationGravityConfigNode.h
--- a/Engine/Plugins/Experimental/ChaosClothAssetEditor/Source/ChaosClothAssetDataflowNodes/Public/ChaosClothAsset/SimulationGravityConfigNode.h
+++ b/Engine/Plugins/Experimental/ChaosClothAssetEditor/Source/ChaosClothAssetDataflowNodes/Public/ChaosClothAsset/SimulationGravityConfigNode.h
@@ -13,15 +13,15 @@ struct FChaosClothAssetSimulationGravityConfigNode : public FChaosClothAssetSimu
 	DATAFLOW_NODE_DEFINE_INTERNAL(FChaosClothAssetSimulationGravityConfigNode, "SimulationGravityConfig", "Cloth", "Cloth Simulation Gravity Config")
 
 public:
-	// Use the config gravity value instead of world gravity.
+	/** Use the config gravity value instead of world gravity. */
 	UPROPERTY(EditAnywhere, Category = "Gravity Properties", Meta = (InlineEditConditionToggle))
 	bool bUseGravityOverride = false;
 
-	// Scale factor applied to the world gravity and also to the clothing simulation interactor gravity. Does not affect the gravity if set using the override below.
+	/** Scale factor applied to the world gravity and also to the clothing simulation interactor gravity. Does not affect the gravity if set using the override below. */
 	UPROPERTY(EditAnywhere, Category = "Gravity Properties", Meta = (UIMin = "0", UIMax = "10", EditCondition = "!bUseGravityOverride"))
 	float GravityScale = 1.f;
 
-	// The gravitational acceleration vector [cm/s^2]
+	/** The gravitational acceleration vector [cm/s^2]. */
 	UPROPERTY(EditAnywhere, Category = "Gravity Properties", Meta = (UIMin = "0", UIMax = "10", EditCondition = "bUseGravityOverride"))
 	FVector3f GravityOverride = { 0.f, 0.f, -980.665f };  // TODO: Should we make this a S.I. unit?
 
diff --git a/Engine/Plugins/Experimental/ChaosClothAssetEditor/Source/ChaosClothAssetDataflowNodes/Public/ChaosClothAsset/SkeletalMeshImportNode.h b/Engine/Plugins/Experimental/ChaosClothAssetEditor/Source/ChaosClothAssetDataflowNodes/Public/ChaosClothAsset/SkeletalMeshImportNode.h
--- a/Engine/Plugins/Experimental/ChaosClothAssetEditor/Source/ChaosClothAssetDataflowNodes/Public/ChaosClothAsset/SkeletalMeshImportNode.h
+++ b/Engine/Plugins/Experimental/ChaosClothAssetEditor/Source/ChaosClothAssetDataflowNodes/Public/ChaosClothAsset/SkeletalMeshImportNode.h
@@ -8,35 +8,46 @@
 
 class USkeletalMesh;
 
+/** Import a skeletal mesh asset into the cloth collection simulation and/or render mesh containers. */
 USTRUCT(Meta = (DataflowCloth))
 struct FChaosClothAssetSkeletalMeshImportNode : public FDataflowNode
 {
 	GENERATED_USTRUCT_BODY()
 	DATAFLOW_NODE_DEFINE_INTERNAL(FChaosClothAssetSkeletalMeshImportNode, "SkeletalMeshImport", "Cloth", "Cloth Skeletal Mesh Import")
-public:
 
-	UPROPERTY(Meta = (DataflowOutput, DisplayName = "Collection"))
+public:
+	UPROPERTY(Meta = (DataflowOutput))
 	FManagedArrayCollection Collection;
 
-	UPROPERTY(EditAnywhere, Category = "Skeletal Mesh")
+	/** The skeletal mesh to import. */
+	UPROPERTY(EditAnywhere, Category = "Skeletal Mesh Import")
 	TObjectPtr<const USkeletalMesh> SkeletalMesh;
 
-	UPROPERTY(VisibleAnywhere, Category = "Skeletal Mesh")
+	/** The skeletal mesh LOD to import. */
+	UPROPERTY(VisibleAnywhere, Category = "Skeletal Mesh Import", Meta = (ClampMin = "0"))
 	int32 LODIndex = 0;
 
-	UPROPERTY(EditAnywhere, Category = "Skeletal Mesh", Meta = (ClampMin = "0"))
+	/** The skeletal mesh LOD section to import. */
+	UPROPERTY(EditAnywhere, Category = "Skeletal Mesh Import", Meta = (ClampMin = "0"))
 	int32 SectionIndex = 0;
 
-	UPROPERTY(EditAnywhere, Category = "Skeletal Mesh")
+	/** Whether to import the simulation mesh from the specified skeletal mesh. */
+	UPROPERTY(EditAnywhere, Category = "Skeletal Mesh Import")
 	bool bImportSimMesh = true;
 
-	UPROPERTY(EditAnywhere, Category = "Skeletal Mesh")
+	/** Whether to import the render mesh from the specified skeletal mesh. */
+	UPROPERTY(EditAnywhere, Category = "Skeletal Mesh Import")
 	bool bImportRenderMesh = true;
 
-	UPROPERTY(EditAnywhere, Category = "Skeletal Mesh")
+	/**
+	 * UV channel of the skeletal meesh to import the 2D simulation mesh patterns from.
+	 * If set to -1, then the import will unwrap the 3D simulation mesh into 2D simulation mesh patterns.
+	 */
+	UPROPERTY(EditAnywhere, Category = "Skeletal Mesh Import", Meta = (ClampMin = "-1"))
 	int32 UVChannel = INDEX_NONE;
 
 	FChaosClothAssetSkeletalMeshImportNode(const Dataflow::FNodeParameters& InParam, FGuid InGuid = FGuid::NewGuid());
 
+private:
 	virtual void Evaluate(Dataflow::FContext& Context, const FDataflowOutput* Out) const override;
 };
diff --git a/Engine/Plugins/Experimental/ChaosClothAssetEditor/Source/ChaosClothAssetDataflowNodes/Public/ChaosClothAsset/StaticMeshImportNode.h b/Engine/Plugins/Experimental/ChaosClothAssetEditor/Source/ChaosClothAssetDataflowNodes/Public/ChaosClothAsset/StaticMeshImportNode.h
--- a/Engine/Plugins/Experimental/ChaosClothAssetEditor/Source/ChaosClothAssetDataflowNodes/Public/ChaosClothAsset/StaticMeshImportNode.h
+++ b/Engine/Plugins/Experimental/ChaosClothAssetEditor/Source/ChaosClothAssetDataflowNodes/Public/ChaosClothAsset/StaticMeshImportNode.h
@@ -8,14 +8,15 @@
 
 class UStaticMesh;
 
+/** Import a static mesh asset into the cloth collection simulation and/or render mesh containers. */
 USTRUCT(Meta = (DataflowCloth))
 struct FChaosClothAssetStaticMeshImportNode : public FDataflowNode
 {
 	GENERATED_USTRUCT_BODY()
 	DATAFLOW_NODE_DEFINE_INTERNAL(FChaosClothAssetStaticMeshImportNode, "StaticMeshImport", "Cloth", "Cloth Static Mesh Import")
-public:
 
-	UPROPERTY(Meta = (DataflowOutput, DisplayName = "Collection"))
+public:
+	UPROPERTY(Meta = (DataflowOutput))
 	FManagedArrayCollection Collection;
 
 	/* The Static Mesh to import from */
@@ -23,7 +24,7 @@ public:
 	TObjectPtr<const UStaticMesh> StaticMesh;
 
 	/* Which static mesh Lod to import.*/
-	UPROPERTY(EditAnywhere, Category = "Static Mesh Import")
+	UPROPERTY(EditAnywhere, Category = "Static Mesh Import", Meta = (ClampMin = "0"))
 	int32 LodIndex = 0;
 
 	/* Import static mesh data as a simulation mesh data.*/
@@ -44,5 +45,6 @@ public:
 
 	FChaosClothAssetStaticMeshImportNode(const Dataflow::FNodeParameters& InParam, FGuid InGuid = FGuid::NewGuid());
 
+private:
 	virtual void Evaluate(Dataflow::FContext& Context, const FDataflowOutput* Out) const override;
 };
diff --git a/Engine/Plugins/Experimental/ChaosClothAssetEditor/Source/ChaosClothAssetDataflowNodes/Public/ChaosClothAsset/TerminalNode.h b/Engine/Plugins/Experimental/ChaosClothAssetEditor/Source/ChaosClothAssetDataflowNodes/Public/ChaosClothAsset/TerminalNode.h
--- a/Engine/Plugins/Experimental/ChaosClothAssetEditor/Source/ChaosClothAssetDataflowNodes/Public/ChaosClothAsset/TerminalNode.h
+++ b/Engine/Plugins/Experimental/ChaosClothAssetEditor/Source/ChaosClothAssetDataflowNodes/Public/ChaosClothAsset/TerminalNode.h
@@ -6,12 +6,12 @@
 #include "GeometryCollection/ManagedArrayCollection.h"
 #include "TerminalNode.generated.h"
 
+/** Cloth terminal node to generate a cloth asset from a cloth collection. */
 USTRUCT(Meta = (DataflowCloth, DataflowTerminal))
 struct FChaosClothAssetTerminalNode : public FDataflowTerminalNode
 {
 	GENERATED_USTRUCT_BODY()
 	DATAFLOW_NODE_DEFINE_INTERNAL(FChaosClothAssetTerminalNode, "ClothAssetTerminal", "Cloth", "Cloth Terminal")  // TODO: Should the category be Terminal instead like all other terminal nodes
-	//DATAFLOW_NODE_RENDER_TYPE(FGeometryCollection::StaticType(), "Collection")  // TODO: Leave out the render type until there is something to render
 
 public:
 	static constexpr int32 MaxLods = 6;  // Hardcoded number of LODs since it is currently not possible to use arrays for optional inputs
@@ -40,16 +40,16 @@ public:
 
 	FChaosClothAssetTerminalNode(const Dataflow::FNodeParameters& InParam, FGuid InGuid = FGuid::NewGuid());
 
+private:
+	//~ Begin FDataflowNode interface
 	virtual void SetAssetValue(TObjectPtr<UObject> Asset, Dataflow::FContext& Context) const override;
-
 	virtual void Evaluate(Dataflow::FContext& Context, const FDataflowOutput* Out) const override {}
-
 	virtual Dataflow::FPin AddPin() override;
 	virtual bool CanAddPin() const override { return NumLods < MaxLods; }
 	virtual bool CanRemovePin() const override { return NumLods > 1; }
 	virtual Dataflow::FPin RemovePin() override;
-	virtual void Serialize(FArchive& Ar);
+	virtual void Serialize(FArchive& Ar) override;
+	//~ End FDataflowNode interface
 
-private:
 	TArray<const FManagedArrayCollection*> GetCollectionLods() const;
 };
diff --git a/Engine/Plugins/Experimental/ChaosClothAssetEditor/Source/ChaosClothAssetDataflowNodes/Public/ChaosClothAsset/TransferSkinWeightsNode.h b/Engine/Plugins/Experimental/ChaosClothAssetEditor/Source/ChaosClothAssetDataflowNodes/Public/ChaosClothAsset/TransferSkinWeightsNode.h
--- a/Engine/Plugins/Experimental/ChaosClothAssetEditor/Source/ChaosClothAssetDataflowNodes/Public/ChaosClothAsset/TransferSkinWeightsNode.h
+++ b/Engine/Plugins/Experimental/ChaosClothAssetEditor/Source/ChaosClothAssetDataflowNodes/Public/ChaosClothAsset/TransferSkinWeightsNode.h
@@ -4,7 +4,6 @@
 
 #include "Dataflow/DataflowNode.h"
 #include "GeometryCollection/ManagedArrayCollection.h"
-#include "BoneWeights.h"
 #include "TransferSkinWeightsNode.generated.h"
 
 class USkeletalMesh;
@@ -12,53 +11,62 @@ class USkeletalMesh;
 UENUM(BlueprintType)
 enum class EChaosClothAssetTransferSkinWeightsMethod : uint8
 {
-	// For every vertex on the target mesh, find the closest point on the surface of the source mesh and copy its weights.
+	/** For every vertex on the target mesh, find the closest point on the surface of the source mesh and copy its weights. */
 	ClosestPointOnSurface,
 	
-	// For every vertex on the target mesh, find the closest point on the surface of the source mesh. If that point 
-	// is within the SearchRadius, and their normals differ by less than the NormalThreshold, then we directly copy  
-	// the weights from the source point to the target mesh vertex. For all the vertices we didn't copy the weights 
-	// directly, automatically compute the smooth weights.
+	/**
+	 * For every vertex on the target mesh, find the closest vertex on the surface of the source mesh.
+	 * If that vertex position is within the search radius, and their normals differ by less than the specified normal threshold,
+	 * then the vertex weights are directly copied from the source vertex to the target mesh vertex.
+	 * For all other vertices whose weights didn't get transferred, smoothed weight values are automatically computed.
+	 */
 	InpaintWeights
 };
 
-
+/** Transfer the skinning weights set on a skeletal mesh to the simulation and/or render mesh stored in the cloth collection. */
 USTRUCT(Meta = (DataflowCloth))
 struct FChaosClothAssetTransferSkinWeightsNode : public FDataflowNode
 {
 	GENERATED_USTRUCT_BODY()
 	DATAFLOW_NODE_DEFINE_INTERNAL(FChaosClothAssetTransferSkinWeightsNode, "TransferSkinWeights", "Cloth", "Cloth Transfer Skin Weights")
-	//DATAFLOW_NODE_RENDER_TYPE(FGeometryCollection::StaticType(), "Collection")  // TODO: Leave out the render type until there is something to render
 
 public:
-
-	UPROPERTY(Meta = (Dataflowinput, DataflowOutput, DisplayName = "Collection", DataflowPassthrough = "Collection"))
+	UPROPERTY(Meta = (Dataflowinput, DataflowOutput, DataflowPassthrough = "Collection"))
 	FManagedArrayCollection Collection;
 
-	UPROPERTY(EditAnywhere, Category = "Source Mesh", Meta = (Dataflowinput, DisplayName = "SkeletalMesh"))
+	/** The skeletal mesh to transfer the skin weights from. */
+	UPROPERTY(EditAnywhere, Category = "Transfer Skin Weights|Source Mesh")
 	TObjectPtr<USkeletalMesh> SkeletalMesh;
 
-	UPROPERTY(EditAnywhere, Category = "Source Mesh", Meta = (DisplayName = "SkeletalMeshLOD"))
-	int32 SkeletalMeshLOD = 0;
+	/** The skeletal mesh LOD to transfer the skin weights from. */
+	UPROPERTY(EditAnywhere, Category = "Transfer Skin Weights|Source Mesh", Meta = (DisplayName = "LOD Index"))
+	int32 LodIndex = 0;
 
-	// The relative transform between the SkeletalMesh and ClothCollection
-	UPROPERTY(EditAnywhere, Category = "Source Mesh", Meta = (DisplayName = "Transform"))
+	/** The relative transform between the skeletal mesh and the cloth asset. */
+	UPROPERTY(EditAnywhere, Category = "Transfer Skin Weights|Source Mesh")
 	FTransform Transform;
 
-	UPROPERTY(EditAnywhere, Category = "Transfer Method", Meta = (DisplayName = "Algorithm"))
+	/** Algorithm used for the transfer method. Use the simple ClosestPointOnSurface method or the more complex InpaintWeights method for better results. */
+	UPROPERTY(EditAnywhere, Category = "Transfer Skin Weights|Transfer Method", Meta = (DisplayName = "Algorithm"))
 	EChaosClothAssetTransferSkinWeightsMethod TransferMethod = EChaosClothAssetTransferSkinWeightsMethod::ClosestPointOnSurface;
 
-	// Defines the search radius as the RadiusPercentage * BoundingBoxDiagonalLength. All points not within the search
-	// radius will be ignored. If negative, all points are considered. Only used in the InpaintWeights algorithm.
-	UPROPERTY(EditAnywhere, Category = "Transfer Method", Meta = (UIMin = -1, UIMax = 2, ClampMin = -1, ClampMax = 2, DisplayName = "Radius Percentage", EditCondition="TransferMethod==EChaosClothAssetTransferSkinWeightsMethod::InpaintWeights"))
+	/**
+	 * Percentage of the bounding box diagonal of the simulation mesh to use as search radius for the InpaintWeights method.
+	 * All points outside of the search radius will be ignored. 
+	 * When set to a negative value (e.g. -1), all points will be considered.
+	 */
+	UPROPERTY(EditAnywhere, Category = "Transfer Skin Weights|Transfer Method", Meta = (UIMin = -1, UIMax = 2, ClampMin = -1, ClampMax = 2, EditCondition="TransferMethod==EChaosClothAssetTransferSkinWeightsMethod::InpaintWeights"))
 	double RadiusPercentage = 0.05;
 
-	// Maximum angle (in degress) difference between target and source point normals to be considred a match. 
-	// If negative, normals are ignored. Only used in the InpaintWeights algorithm.
-	UPROPERTY(EditAnywhere, Category = "Transfer Method", Meta = (UIMin = -1, UIMax = 180, ClampMin = -1, ClampMax = 180, DisplayName = "Normal Threshold", EditCondition="TransferMethod==EChaosClothAssetTransferSkinWeightsMethod::InpaintWeights"))
+	/**
+	 * Maximum angle difference (in degrees) between the target and source point normals to be considered a match for the InpaintWeights method.
+	 * If set to a negative value (e.g. -1), normals will be ignored.
+	 */
+	UPROPERTY(EditAnywhere, Category = "Transfer Skin Weights|Transfer Method", Meta = (UIMin = -1, UIMax = 180, ClampMin = -1, ClampMax = 180, EditCondition="TransferMethod==EChaosClothAssetTransferSkinWeightsMethod::InpaintWeights"))
 	double NormalThreshold = 30;
 
 	FChaosClothAssetTransferSkinWeightsNode(const Dataflow::FNodeParameters& InParam, FGuid InGuid = FGuid::NewGuid());
 
+private:
 	virtual void Evaluate(Dataflow::FContext& Context, const FDataflowOutput* Out) const override;
 };
diff --git a/Engine/Plugins/Experimental/ChaosClothAssetEditor/Source/ChaosClothAssetEditorTools/Private/ChaosClothAsset/ClothTransferSkinWeightsTool.cpp b/Engine/Plugins/Experimental/ChaosClothAssetEditor/Source/ChaosClothAssetEditorTools/Private/ChaosClothAsset/ClothTransferSkinWeightsTool.cpp
--- a/Engine/Plugins/Experimental/ChaosClothAssetEditor/Source/ChaosClothAssetEditorTools/Private/ChaosClothAsset/ClothTransferSkinWeightsTool.cpp
+++ b/Engine/Plugins/Experimental/ChaosClothAssetEditor/Source/ChaosClothAssetEditorTools/Private/ChaosClothAsset/ClothTransferSkinWeightsTool.cpp
@@ -115,7 +115,7 @@ void UClothTransferSkinWeightsTool::Setup()
 
 	SetSRTPropertiesFromTransform(TransferSkinWeightsNode->Transform);
 	ToolProperties->SourceMesh = TransferSkinWeightsNode->SkeletalMesh;
-	ToolProperties->SourceMeshLOD = TransferSkinWeightsNode->SkeletalMeshLOD;
+	ToolProperties->SourceMeshLOD = TransferSkinWeightsNode->LodIndex;
 	
 
 	AddToolPropertySource(ToolProperties);
@@ -266,7 +266,7 @@ void UClothTransferSkinWeightsTool::Shutdown(EToolShutdownType ShutdownType)
 	if (ShutdownType == EToolShutdownType::Accept)
 	{
 		TransferSkinWeightsNode->SkeletalMesh = ToolProperties->SourceMesh;
-		TransferSkinWeightsNode->SkeletalMeshLOD = ToolProperties->SourceMeshLOD;
+		TransferSkinWeightsNode->LodIndex = ToolProperties->SourceMeshLOD;
 		TransferSkinWeightsNode->Transform = TransformFromProperties();
 	}
 
@@ -307,7 +307,7 @@ bool UClothTransferSkinWeightsTool::CanAccept() const
 	const FTransform& TransformOnNode = TransferSkinWeightsNode->Transform;;
 
 	return (ToolProperties->SourceMesh != TransferSkinWeightsNode->SkeletalMesh) ||
-		(ToolProperties->SourceMeshLOD != TransferSkinWeightsNode->SkeletalMeshLOD) ||
+		(ToolProperties->SourceMeshLOD != TransferSkinWeightsNode->LodIndex) ||
 		(ToolProperties->SourceMeshRotation != TransformOnNode.Rotator().Euler()) ||
 		(ToolProperties->SourceMeshTranslation != TransformOnNode.GetTranslation()) ||
 		(ToolProperties->SourceMeshScale != TransformOnNode.GetScale3D());
diff --git a/Engine/Plugins/Experimental/ChaosClothAssetEditor/Source/ChaosClothAssetEditorTools/Private/ChaosClothAsset/ClothWeightMapPaintTool.cpp b/Engine/Plugins/Experimental/ChaosClothAssetEditor/Source/ChaosClothAssetEditorTools/Private/ChaosClothAsset/ClothWeightMapPaintTool.cpp
--- a/Engine/Plugins/Experimental/ChaosClothAssetEditor/Source/ChaosClothAssetEditorTools/Private/ChaosClothAsset/ClothWeightMapPaintTool.cpp
+++ b/Engine/Plugins/Experimental/ChaosClothAssetEditor/Source/ChaosClothAssetEditorTools/Private/ChaosClothAsset/ClothWeightMapPaintTool.cpp
@@ -38,7 +38,6 @@
 #include "Dataflow/DataflowEdNode.h"
 #include "Dataflow/DataflowObject.h"
 #include "ChaosClothAsset/AddWeightMapNode.h"
-#include "ChaosClothAsset/DataflowNodes.h"
 #include "GraphEditor.h"
 #include "Dataflow/DataflowSNode.h"
 #include "Dataflow/DataflowGraphEditor.h"
