Chaos Cloth - Updated CopySimulationToRenderMeshNode to default to generating a single render pattern since there isn't any way to set multiple materials yet.

#rnx

diff --git a/Engine/Plugins/Experimental/ChaosClothAssetEditor/Source/ChaosClothAssetDataflowNodes/Public/ChaosClothAsset/CopySimulationToRenderMeshNode.h b/Engine/Plugins/Experimental/ChaosClothAssetEditor/Source/ChaosClothAssetDataflowNodes/Public/ChaosClothAsset/CopySimulationToRenderMeshNode.h
--- a/Engine/Plugins/Experimental/ChaosClothAssetEditor/Source/ChaosClothAssetDataflowNodes/Public/ChaosClothAsset/CopySimulationToRenderMeshNode.h
+++ b/Engine/Plugins/Experimental/ChaosClothAssetEditor/Source/ChaosClothAssetDataflowNodes/Public/ChaosClothAsset/CopySimulationToRenderMeshNode.h
@@ -26,7 +26,7 @@ public:
 
 	/** Generate a single render pattern rather than a render pattern per sim pattern. */
 	UPROPERTY(EditAnywhere, Category = "Copy Simulation To Render Mesh")
-	bool bGenerateSingleRenderPattern = false;
+	bool bGenerateSingleRenderPattern = true;
 
 	FChaosClothAssetCopySimulationToRenderMeshNode(const Dataflow::FNodeParameters& InParam, FGuid InGuid = FGuid::NewGuid());
 
