Chaos Cloth - Fixed Datasmith import node hash when resetting to an empty file path.

#rnx

diff --git a/Engine/Plugins/Experimental/ChaosClothAssetEditor/Source/ChaosClothAssetDataflowNodes/Private/ChaosClothAsset/DatasmithImportNode.cpp b/Engine/Plugins/Experimental/ChaosClothAssetEditor/Source/ChaosClothAssetDataflowNodes/Private/ChaosClothAsset/DatasmithImportNode.cpp
--- a/Engine/Plugins/Experimental/ChaosClothAssetEditor/Source/ChaosClothAssetDataflowNodes/Private/ChaosClothAsset/DatasmithImportNode.cpp
+++ b/Engine/Plugins/Experimental/ChaosClothAssetEditor/Source/ChaosClothAssetDataflowNodes/Private/ChaosClothAsset/DatasmithImportNode.cpp
@@ -98,31 +98,34 @@ void FChaosClothAssetDatasmithImportNode::Evaluate(Dataflow::FContext& Context,
 {
 	if (Out->IsA<FManagedArrayCollection>(&Collection))
 	{
-		if (FPaths::FileExists(ImportFile.FilePath))
+		const FMD5Hash FileHash = ImportFile.FilePath.IsEmpty() ?
+			FMD5Hash() :  // Reset to an empty cloth collection
+			FPaths::FileExists(ImportFile.FilePath) ?
+				FMD5Hash::HashFile(*ImportFile.FilePath) :  // Update cache to the file content if necessary
+				ImportHash;  // Keep the current cached file
+
+		if (FileHash != ImportHash)
 		{
-			const FMD5Hash FileHash = FMD5Hash::HashFile(*ImportFile.FilePath);
-			if (FileHash != ImportHash)
-			{
-				ImportHash = FileHash;
+			ImportHash = FileHash;
 
-				// The file has changed, reimport
-				const bool bSuccess = EvaluateImpl(Context, ImportCache);
+			// The file has changed, reimport
+			const bool bSuccess = EvaluateImpl(Context, ImportCache);
 
-				if (!bSuccess)
-				{
-					using namespace UE::Chaos::ClothAsset;
+			if (!bSuccess)
+			{
+				using namespace UE::Chaos::ClothAsset;
 
-					// The import has failed, reset the cache to an empty cloth collection
-					ImportCache.Reset();
+				// The import has failed, reset the cache to an empty cloth collection
+				ImportCache.Reset();
 
-					const TSharedRef<FManagedArrayCollection> ClothCollection = MakeShared<FManagedArrayCollection>(MoveTemp(ImportCache));
-					FCollectionClothFacade CollectionClothFacade(ClothCollection);
-					CollectionClothFacade.DefineSchema();
+				const TSharedRef<FManagedArrayCollection> ClothCollection = MakeShared<FManagedArrayCollection>(MoveTemp(ImportCache));
+				FCollectionClothFacade CollectionClothFacade(ClothCollection);
+				CollectionClothFacade.DefineSchema();
 
-					ImportCache = MoveTemp(*ClothCollection);
-				}
+				ImportCache = MoveTemp(*ClothCollection);
 			}
 		}
+
 		SetValue(Context, ImportCache, &Collection);
 	}
 }
