Chaos Cloth Editor - Update the Outliner window combo box with the new group names whenever a node selection is changed.

#enx

diff --git a/Engine/Plugins/Experimental/ChaosClothAssetEditor/Source/ChaosClothAssetEditor/Private/ChaosClothAsset/SClothCollectionOutliner.cpp b/Engine/Plugins/Experimental/ChaosClothAssetEditor/Source/ChaosClothAssetEditor/Private/ChaosClothAsset/SClothCollectionOutliner.cpp
--- a/Engine/Plugins/Experimental/ChaosClothAssetEditor/Source/ChaosClothAssetEditor/Private/ChaosClothAsset/SClothCollectionOutliner.cpp
+++ b/Engine/Plugins/Experimental/ChaosClothAssetEditor/Source/ChaosClothAssetEditor/Private/ChaosClothAsset/SClothCollectionOutliner.cpp
@@ -135,6 +135,8 @@ void SClothCollectionOutliner::SetClothCollection(TSharedPtr<FManagedArrayCollec
 		ListItems.Empty();
 		SelectedGroupName = FName();
 	}
+	// Refresh the combo box with the new set of group names after ClothCollectionGroupNames changed
+	SelectedGroupNameComboBox->RefreshOptions();
 }
 
 void SClothCollectionOutliner::SetSelectedGroupName(const FName& InSelectedGroupName)
