Chaos Cloth Preview Scene:
- set up permanent Cloth and SkeletalMesh Components on the scene actor when the PreviewScene is first initialized
- do not delete/create components when the assets in the scene description change, just call the appropirate SetAsset() function on the component. This means the root component of the scene actor doesn't need to be updated when the assets change, which was causing problems when undoing
- when an undo/redo is performed, call SceneDescriptionPropertyChanged to update the scene components

diff --git a/Engine/Plugins/Experimental/ChaosClothAssetEditor/Source/ChaosClothAssetEditor/Private/ChaosClothAsset/ClothEditorMode.cpp b/Engine/Plugins/Experimental/ChaosClothAssetEditor/Source/ChaosClothAssetEditor/Private/ChaosClothAsset/ClothEditorMode.cpp
--- a/Engine/Plugins/Experimental/ChaosClothAssetEditor/Source/ChaosClothAssetEditor/Private/ChaosClothAsset/ClothEditorMode.cpp
+++ b/Engine/Plugins/Experimental/ChaosClothAssetEditor/Source/ChaosClothAssetEditor/Private/ChaosClothAsset/ClothEditorMode.cpp
@@ -648,7 +648,7 @@ void UChaosClothAssetEditorMode::ResumeSimulation()
 
 bool UChaosClothAssetEditorMode::IsSimulationSuspended() const
 {
-	if (PreviewScene && PreviewScene->GetClothComponent())
+	if (PreviewScene && PreviewScene->GetClothComponent() && PreviewScene->GetClothComponent()->GetClothSimulationProxy())
 	{
 		return PreviewScene->GetClothComponent()->IsSimulationSuspended();
 	}
@@ -790,12 +790,18 @@ FBox UChaosClothAssetEditorMode::PreviewBoundingBox() const
 
 	if (const UChaosClothComponent* const Cloth = PreviewScene->GetClothComponent())
 	{
-		Bounds += Cloth->Bounds.GetBox();
+		if (Cloth->GetClothAsset())
+		{
+			Bounds += Cloth->Bounds.GetBox();
+		}
 	}
 
 	if (const USkeletalMeshComponent* const SkeletalMesh = PreviewScene->GetSkeletalMeshComponent())
 	{
-		Bounds += SkeletalMesh->Bounds.GetBox();
+		if (SkeletalMesh->GetSkeletalMeshAsset())
+		{
+			Bounds += SkeletalMesh->Bounds.GetBox();
+		}
 	}
 
 	return Bounds;
diff --git a/Engine/Plugins/Experimental/ChaosClothAssetEditor/Source/ChaosClothAssetEditor/Private/ChaosClothAsset/ClothEditorPreviewScene.cpp b/Engine/Plugins/Experimental/ChaosClothAssetEditor/Source/ChaosClothAssetEditor/Private/ChaosClothAsset/ClothEditorPreviewScene.cpp
--- a/Engine/Plugins/Experimental/ChaosClothAssetEditor/Source/ChaosClothAssetEditor/Private/ChaosClothAsset/ClothEditorPreviewScene.cpp
+++ b/Engine/Plugins/Experimental/ChaosClothAssetEditor/Source/ChaosClothAssetEditor/Private/ChaosClothAsset/ClothEditorPreviewScene.cpp
@@ -10,6 +10,8 @@
 #include "ComponentReregisterContext.h"
 #include "ChaosClothAsset/ClothAsset.h"
 #include "SkinnedAssetCompiler.h"
+#include "Misc/TransactionObjectEvent.h"
+
 
 #define LOCTEXT_NAMESPACE "UChaosClothEditorPreviewScene"
 
@@ -20,9 +22,26 @@ void UChaosClothPreviewSceneDescription::SetPreviewScene(UE::Chaos::ClothAsset::
 
 void UChaosClothPreviewSceneDescription::PostEditChangeProperty(FPropertyChangedEvent& PropertyChangedEvent)
 {
+	Super::PostEditChangeProperty(PropertyChangedEvent);
+
 	if (PreviewScene)
 	{
-		PreviewScene->SceneDescriptionPropertyChanged(PropertyChangedEvent);
+		PreviewScene->SceneDescriptionPropertyChanged(PropertyChangedEvent.GetPropertyName());
+	}
+}
+
+void UChaosClothPreviewSceneDescription::PostTransacted(const FTransactionObjectEvent& TransactionEvent)
+{
+	Super::PostTransacted(TransactionEvent);
+
+	// On Undo/Redo, PostEditChangeProperty just gets an empty FPropertyChangedEvent. However this function gets enough info to figure out which property changed
+	if (TransactionEvent.GetEventType() == ETransactionObjectEventType::UndoRedo && TransactionEvent.HasPropertyChanges())
+	{
+		const TArray<FName>& PropertyNames = TransactionEvent.GetChangedProperties();
+		for (const FName& PropertyName : PropertyNames)
+		{
+			PreviewScene->SceneDescriptionPropertyChanged(PropertyName);
+		}
 	}
 }
 
@@ -37,11 +56,15 @@ FChaosClothPreviewScene::FChaosClothPreviewScene(FPreviewScene::ConstructionValu
 
 	SceneActor = GetWorld()->SpawnActor<AActor>(AActor::StaticClass());
 
-	if (PreviewSceneDescription->SkeletalMeshAsset)
-	{
-		CreateSkeletalMeshComponent();
-		UpdateSkeletalMeshAnimation();
-	}
+	SkeletalMeshComponent = NewObject<USkeletalMeshComponent>(SceneActor);
+	SkeletalMeshComponent->TransformUpdated.AddRaw(this, &FChaosClothPreviewScene::SkeletalMeshTransformChanged);
+	SkeletalMeshComponent->SelectionOverrideDelegate = UPrimitiveComponent::FSelectionOverride::CreateRaw(this, &FChaosClothPreviewScene::IsComponentSelected);
+	SkeletalMeshComponent->SetDisablePostProcessBlueprint(true);
+
+	ClothComponent = NewObject<UChaosClothComponent>(SceneActor);
+	ClothComponent->SelectionOverrideDelegate = UPrimitiveComponent::FSelectionOverride::CreateRaw(this, &FChaosClothPreviewScene::IsComponentSelected);
+
+	SceneActor->RegisterAllComponents();
 }
 
 FChaosClothPreviewScene::~FChaosClothPreviewScene()
@@ -74,65 +97,70 @@ void FChaosClothPreviewScene::AddReferencedObjects(FReferenceCollector& Collecto
 
 void FChaosClothPreviewScene::UpdateSkeletalMeshAnimation()
 {
-	if (SkeletalMeshComponent)
+	check(SkeletalMeshComponent);
+	
+	if (PreviewSceneDescription->AnimationAsset)
 	{
-		if (PreviewSceneDescription->AnimationAsset)
-		{
-			PreviewAnimInstance = NewObject<UAnimSingleNodeInstance>(SkeletalMeshComponent);
-			PreviewAnimInstance->SetAnimationAsset(PreviewSceneDescription->AnimationAsset);
-
-			SkeletalMeshComponent->SetAnimationMode(EAnimationMode::AnimationSingleNode);
-			SkeletalMeshComponent->InitAnim(true);
-			SkeletalMeshComponent->AnimationData.PopulateFrom(PreviewAnimInstance);
-			SkeletalMeshComponent->AnimScriptInstance = PreviewAnimInstance;
-			SkeletalMeshComponent->AnimScriptInstance->InitializeAnimation();
-			SkeletalMeshComponent->ValidateAnimation();
-			SkeletalMeshComponent->Stop();
-		}
-		else
-		{
-			SkeletalMeshComponent->Stop();
-			SkeletalMeshComponent->AnimationData = FSingleAnimationPlayData();
-			SkeletalMeshComponent->AnimScriptInstance = nullptr;
-		}
+		PreviewAnimInstance = NewObject<UAnimSingleNodeInstance>(SkeletalMeshComponent);
+		PreviewAnimInstance->SetAnimationAsset(PreviewSceneDescription->AnimationAsset);
+
+		SkeletalMeshComponent->SetAnimationMode(EAnimationMode::AnimationSingleNode);
+		SkeletalMeshComponent->InitAnim(true);
+		SkeletalMeshComponent->AnimationData.PopulateFrom(PreviewAnimInstance);
+		SkeletalMeshComponent->AnimScriptInstance = PreviewAnimInstance;
+		SkeletalMeshComponent->AnimScriptInstance->InitializeAnimation();
+		SkeletalMeshComponent->ValidateAnimation();
+	}
+	else
+	{
+		SkeletalMeshComponent->Stop();
+		SkeletalMeshComponent->AnimationData = FSingleAnimationPlayData();
+		SkeletalMeshComponent->AnimScriptInstance = nullptr;
 	}
 }
 
 void FChaosClothPreviewScene::UpdateClothComponentAttachment()
 {
-	if (ClothComponent && SkeletalMeshComponent && !ClothComponent->IsAttachedTo(SkeletalMeshComponent))
+	check(SkeletalMeshComponent);
+	check(ClothComponent);
+
+	if (SkeletalMeshComponent->GetSkeletalMeshAsset() && !ClothComponent->IsAttachedTo(SkeletalMeshComponent))
 	{
 		ClothComponent->AttachToComponent(SkeletalMeshComponent, FAttachmentTransformRules::SnapToTargetNotIncludingScale);
 	}
+	else if (!SkeletalMeshComponent->GetSkeletalMeshAsset() && ClothComponent->IsAttachedTo(SkeletalMeshComponent))
+	{
+		ClothComponent->DetachFromComponent(FDetachmentTransformRules::KeepWorldTransform);
+
+		// Hard reset cloth simulation if we are losing the attachment
+		{
+			const FComponentReregisterContext Context(ClothComponent);
+		}
+	}
 }
 
-void FChaosClothPreviewScene::SceneDescriptionPropertyChanged(struct FPropertyChangedEvent& PropertyChangedEvent)
+void FChaosClothPreviewScene::SceneDescriptionPropertyChanged(const FName& PropertyName)
 {
-	if (PropertyChangedEvent.Property->GetFName() == GET_MEMBER_NAME_CHECKED(UChaosClothPreviewSceneDescription, SkeletalMeshAsset))
+	if (PropertyName == GET_MEMBER_NAME_CHECKED(UChaosClothPreviewSceneDescription, SkeletalMeshAsset))
 	{
-		if (PreviewSceneDescription->SkeletalMeshAsset)
-		{
-			CreateSkeletalMeshComponent();
-		}
-		else
-		{
-			DeleteSkeletalMeshComponent();
-		}
+		check(SkeletalMeshComponent);
+
+		SkeletalMeshComponent->SetSkeletalMeshAsset(PreviewSceneDescription->SkeletalMeshAsset);
+
 		UpdateSkeletalMeshAnimation();
 		UpdateClothComponentAttachment();
 	}
 
-	if (PropertyChangedEvent.Property->GetFName() == GET_MEMBER_NAME_CHECKED(UChaosClothPreviewSceneDescription, SkeletalMeshTransform))
+	if (PropertyName == GET_MEMBER_NAME_CHECKED(UChaosClothPreviewSceneDescription, SkeletalMeshTransform))
 	{
-		if (SkeletalMeshComponent)
-		{
-			SkeletalMeshComponent->SetComponentToWorld(PreviewSceneDescription->SkeletalMeshTransform);
-		}
+		check(SkeletalMeshComponent);
+
+		SkeletalMeshComponent->SetComponentToWorld(PreviewSceneDescription->SkeletalMeshTransform);
 	}
 
-	if (PropertyChangedEvent.Property->GetFName() == GET_MEMBER_NAME_CHECKED(UChaosClothPreviewSceneDescription, AnimationAsset))
+	if (PropertyName == GET_MEMBER_NAME_CHECKED(UChaosClothPreviewSceneDescription, AnimationAsset))
 	{
-		if (!SkeletalMeshComponent || !PreviewSceneDescription->AnimationAsset)
+		if (!PreviewSceneDescription->AnimationAsset)
 		{
 			PreviewAnimInstance = nullptr;
 		}
@@ -172,60 +200,6 @@ void FChaosClothPreviewScene::SkeletalMeshTransformChanged(USceneComponent* Upda
 	PreviewSceneDescription->SkeletalMeshTransform = UpdatedComponent->GetComponentToWorld();
 }
 
-
-void FChaosClothPreviewScene::DeleteSkeletalMeshComponent()
-{
-	if (SkeletalMeshComponent)
-	{
-		if (SceneActor)
-		{
-			SceneActor->RemoveOwnedComponent(SkeletalMeshComponent);
-		}
-
-		SkeletalMeshComponent->TransformUpdated.RemoveAll(this);
-		SkeletalMeshComponent->SelectionOverrideDelegate.Unbind();
-		SkeletalMeshComponent->UnregisterComponent();
-	}
-
-	SkeletalMeshComponent = nullptr;
-
-	// Detach the ClothComponent and make it the RootComponent of the SceneActor
-	if (SceneActor && ClothComponent)
-	{
-		const bool bWasSimming = !ClothComponent->IsSimulationSuspended();
-		ClothComponent->SuspendSimulation();
-
-		ClothComponent->DetachFromComponent(FDetachmentTransformRules::KeepRelativeTransform);
-		SceneActor->SetRootComponent(ClothComponent);
-		SceneActor->RegisterAllComponents();
-
-		// hard reset and resume simulation if it was running prior
-		{
-			const FComponentReregisterContext Context(ClothComponent);
-		}
-		if (bWasSimming)
-		{
-			ClothComponent->ResumeSimulation();
-		}
-	}
-}
-
-void FChaosClothPreviewScene::CreateSkeletalMeshComponent()
-{
-	ensure(PreviewSceneDescription->SkeletalMeshAsset);
-
-	// Remove any existing skeletal mesh component
-	DeleteSkeletalMeshComponent();
-
-	SkeletalMeshComponent = NewObject<USkeletalMeshComponent>(SceneActor);
-	SkeletalMeshComponent->SetSkeletalMeshAsset(PreviewSceneDescription->SkeletalMeshAsset);
-	SkeletalMeshComponent->TransformUpdated.AddRaw(this, &FChaosClothPreviewScene::SkeletalMeshTransformChanged);
-	SkeletalMeshComponent->SelectionOverrideDelegate = UPrimitiveComponent::FSelectionOverride::CreateRaw(this, &FChaosClothPreviewScene::IsComponentSelected);
-
-	SceneActor->SetRootComponent(SkeletalMeshComponent);
-	SceneActor->RegisterAllComponents();
-}
-
 bool FChaosClothPreviewScene::IsComponentSelected(const UPrimitiveComponent* InComponent)
 {
 	if (const UTypedElementSelectionSet* const TypedElementSelectionSet = ClothPreviewEditorModeManager->GetEditorSelectionSet())
@@ -242,35 +216,15 @@ bool FChaosClothPreviewScene::IsComponentSelected(const UPrimitiveComponent* InC
 
 void FChaosClothPreviewScene::SetClothAsset(UChaosClothAsset* Asset)
 {
-	// Clean up old cloth component if it exists
-
 	check(Asset);
 	check(SceneActor);
+	check(ClothComponent);
 
-	if (ClothComponent)
-	{
-		ClothComponent->SetClothAsset(Asset);
-	}
-	else
-	{
-		ClothComponent = NewObject<UChaosClothComponent>(SceneActor);
-		ClothComponent->SetClothAsset(Asset);
-		ClothComponent->SelectionOverrideDelegate = UPrimitiveComponent::FSelectionOverride::CreateRaw(this, &FChaosClothPreviewScene::IsComponentSelected);
-
-		if (SkeletalMeshComponent)
-		{
-			check(SceneActor->GetRootComponent() == SkeletalMeshComponent);
-		}
-		else
-		{
-			SceneActor->SetRootComponent(ClothComponent);
-		}
-
-		SceneActor->RegisterAllComponents();
-	}
-
+	ClothComponent->SetClothAsset(Asset);
 	UpdateClothComponentAttachment();
 
+	// Wait for asset to load and update the component bounds
+	ClothComponent->InvalidateCachedBounds();
 	FSkinnedAssetCompilingManager::Get().FinishCompilation(TArrayView<USkinnedAsset* const>{Asset});
 	ClothComponent->UpdateBounds();
 }
diff --git a/Engine/Plugins/Experimental/ChaosClothAssetEditor/Source/ChaosClothAssetEditor/Private/ChaosClothAsset/ClothEditorPreviewScene.h b/Engine/Plugins/Experimental/ChaosClothAssetEditor/Source/ChaosClothAssetEditor/Private/ChaosClothAsset/ClothEditorPreviewScene.h
--- a/Engine/Plugins/Experimental/ChaosClothAssetEditor/Source/ChaosClothAssetEditor/Private/ChaosClothAsset/ClothEditorPreviewScene.h
+++ b/Engine/Plugins/Experimental/ChaosClothAssetEditor/Source/ChaosClothAssetEditor/Private/ChaosClothAsset/ClothEditorPreviewScene.h
@@ -27,16 +27,21 @@ class CHAOSCLOTHASSETEDITOR_API UChaosClothPreviewSceneDescription : public UObj
 public:
 	GENERATED_BODY()
 
+	UChaosClothPreviewSceneDescription()
+	{
+		SetFlags(RF_Transactional);
+	}
+
 	void SetPreviewScene(UE::Chaos::ClothAsset::FChaosClothPreviewScene* PreviewScene);
 
 	// Skeletal Mesh source asset
-	UPROPERTY(EditAnywhere, Category="SkeletalMesh")
+	UPROPERTY(EditAnywhere, Transient, Category="SkeletalMesh")
 	TObjectPtr<USkeletalMesh> SkeletalMeshAsset;
 
-	UPROPERTY(EditAnywhere, Category = "SkeletalMesh")
+	UPROPERTY(EditAnywhere, Transient, Category = "SkeletalMesh")
 	FTransform SkeletalMeshTransform;
 
-	UPROPERTY(EditAnywhere, Category = "SkeletalMesh")
+	UPROPERTY(EditAnywhere, Transient, Category = "SkeletalMesh")
 	TObjectPtr<UAnimationAsset> AnimationAsset;
 
 private:
@@ -44,6 +49,8 @@ private:
 	// Listen for changes to the scene description members and notify the PreviewScene
 	virtual void PostEditChangeProperty(struct FPropertyChangedEvent& PropertyChangedEvent) override;
 
+	virtual void PostTransacted(const FTransactionObjectEvent& TransactionEvent) override;
+
 	UE::Chaos::ClothAsset::FChaosClothPreviewScene* PreviewScene;
 };
 
@@ -68,7 +75,7 @@ public:
 	void SetClothAsset(UChaosClothAsset* Asset);
 
 	// Update Scene in response to the SceneDescription changing
-	void SceneDescriptionPropertyChanged(struct FPropertyChangedEvent& PropertyChangedEvent);
+	void SceneDescriptionPropertyChanged(const FName& PropertyName);
 
 	UAnimSingleNodeInstance* GetPreviewAnimInstance();
 	const UAnimSingleNodeInstance* const GetPreviewAnimInstance() const;
@@ -92,9 +99,6 @@ private:
 
 	void SkeletalMeshTransformChanged(USceneComponent* UpdatedComponent, EUpdateTransformFlags UpdateTransformFlags, ETeleportType Teleport);
 
-	void CreateSkeletalMeshComponent();
-	void DeleteSkeletalMeshComponent();
-
 	bool IsComponentSelected(const UPrimitiveComponent* InComponent);
 
 	TObjectPtr<UChaosClothPreviewSceneDescription> PreviewSceneDescription;
diff --git a/Engine/Plugins/Experimental/ChaosClothAssetEditor/Source/ChaosClothAssetEditor/Private/ChaosClothAsset/ClothEditorToolkit.cpp b/Engine/Plugins/Experimental/ChaosClothAssetEditor/Source/ChaosClothAssetEditor/Private/ChaosClothAsset/ClothEditorToolkit.cpp
--- a/Engine/Plugins/Experimental/ChaosClothAssetEditor/Source/ChaosClothAssetEditor/Private/ChaosClothAsset/ClothEditorToolkit.cpp
+++ b/Engine/Plugins/Experimental/ChaosClothAssetEditor/Source/ChaosClothAssetEditor/Private/ChaosClothAsset/ClothEditorToolkit.cpp
@@ -993,14 +993,14 @@ void FChaosClothAssetEditorToolkit::OnClothAssetChanged()
 
 	if (UChaosClothAsset* const ClothAsset = Cast<UChaosClothAsset>(ObjectsToEdit[0]))
 	{
-		const bool bHadClothComponent = (ClothPreviewScene->GetClothComponent() != nullptr);
+		const bool bHadClothAsset = (ClothPreviewScene->GetClothComponent()->GetClothAsset() != nullptr);
 
 		ClothPreviewScene->SetClothAsset(ClothAsset);
 
 		ensure(ClothAsset->HasAnyFlags(RF_Transactional));		// Ensure all objects are transactable for undo/redo in the details panel
 		SetEditingObject(ClothAsset);
 
-		if (!bHadClothComponent)
+		if (!bHadClothAsset)
 		{
 			// Focus on the cloth component if this is the first time adding one
 			ClothPreviewViewportClient->FocusViewportOnBox(ClothMode->PreviewBoundingBox());
diff --git a/Engine/Plugins/Experimental/ChaosClothAssetEditor/Source/ChaosClothAssetEditor/Private/ChaosClothAsset/SClothEditor3DViewport.cpp b/Engine/Plugins/Experimental/ChaosClothAssetEditor/Source/ChaosClothAssetEditor/Private/ChaosClothAsset/SClothEditor3DViewport.cpp
--- a/Engine/Plugins/Experimental/ChaosClothAssetEditor/Source/ChaosClothAssetEditor/Private/ChaosClothAsset/SClothEditor3DViewport.cpp
+++ b/Engine/Plugins/Experimental/ChaosClothAssetEditor/Source/ChaosClothAssetEditor/Private/ChaosClothAsset/SClothEditor3DViewport.cpp
@@ -10,6 +10,7 @@
 #include "Animation/AnimSingleNodeInstance.h"
 #include "EditorModeTools.h"
 #include "Widgets/Layout/SBorder.h"
+#include "Components/SkeletalMeshComponent.h"
 
 #define LOCTEXT_NAMESPACE "SChaosClothAssetEditor3DViewport"
 
@@ -169,7 +170,7 @@ EVisibility SChaosClothAssetEditor3DViewport::GetAnimControlVisibility() const
 	using namespace UE::Chaos::ClothAsset;
 
 	const TSharedPtr<const FChaosClothPreviewScene> Scene = GetPreviewScene().Pin();
-	return (Scene && Scene->GetSkeletalMeshComponent() && Scene->GetPreviewAnimInstance()) ? EVisibility::Visible : EVisibility::Hidden;
+	return (Scene && Scene->GetSkeletalMeshComponent() && Scene->GetSkeletalMeshComponent()->GetSkeletalMeshAsset() && Scene->GetPreviewAnimInstance()) ? EVisibility::Visible : EVisibility::Hidden;
 }
 
 #undef LOCTEXT_NAMESPACE
