Cloth Editor: fix some issues when switching between 2D and 3D view modes and selecting a new node
- The first time we see a valid mesh in a particular view mode, refocus on it
- When switching between view modes, save/restore the camera data, rather than relying on refocus
- The overall intended effect is to have persistent 2d and 3d cameras that can be switched back and forth, while still focusing on the mesh when it is first seen

diff --git a/Engine/Plugins/Experimental/ChaosClothAssetEditor/Source/ChaosClothAssetEditor/Private/ChaosClothAsset/ClothEditorMode.cpp b/Engine/Plugins/Experimental/ChaosClothAssetEditor/Source/ChaosClothAssetEditor/Private/ChaosClothAsset/ClothEditorMode.cpp
--- a/Engine/Plugins/Experimental/ChaosClothAssetEditor/Source/ChaosClothAssetEditor/Private/ChaosClothAsset/ClothEditorMode.cpp
+++ b/Engine/Plugins/Experimental/ChaosClothAssetEditor/Source/ChaosClothAssetEditor/Private/ChaosClothAsset/ClothEditorMode.cpp
@@ -382,12 +382,8 @@ void UChaosClothAssetEditorMode::SetSelectedClothCollection(TSharedPtr<FManagedA
 	SelectedClothCollection = Collection;
 	ReinitializeDynamicMeshComponents();
 
-	if (bFirstClothCollection && Collection)
-	{
-		// refocus the viewport if this is the first time a cloth collection has been set
-		RefocusRestSpaceViewportClient();
-		bFirstClothCollection = false;
-	}
+	// The first time we get a valid mesh, refocus the camera on it
+	FirstTimeFocusRestSpaceViewport();
 }
 
 TSharedPtr<FManagedArrayCollection> UChaosClothAssetEditorMode::GetClothCollection()
@@ -600,6 +596,27 @@ void UChaosClothAssetEditorMode::RefocusRestSpaceViewportClient()
 	}
 }
 
+void UChaosClothAssetEditorMode::FirstTimeFocusRestSpaceViewport()
+{
+	// If this is the first time seeing a valid 2D or 3D mesh, refocus the camera on it.
+	const bool bIsValid = (SelectedClothCollection && DynamicMeshComponent && DynamicMeshComponent->GetMesh()->TriangleCount() > 0);
+	const bool bIs2D = ConstructionViewMode == UE::Chaos::ClothAsset::EClothPatternVertexType::Sim2D;
+
+	if (bIsValid)
+	{
+		if (bIs2D && bFirstValid2DMesh)
+		{
+			bFirstValid2DMesh = false;
+			RefocusRestSpaceViewportClient();
+		}
+		else if (!bIs2D && bFirstValid3DMesh)
+		{
+			bFirstValid3DMesh = false;
+			RefocusRestSpaceViewportClient();
+		}
+	}
+}
+
 void UChaosClothAssetEditorMode::InitializeTargets(const TArray<TObjectPtr<UObject>>& AssetsIn)
 {
 	// InitializeContexts needs to have been called first so that we have the 3d preview world ready.
@@ -804,8 +821,6 @@ FBox UChaosClothAssetEditorMode::PreviewBoundingBox() const
 
 void UChaosClothAssetEditorMode::SetConstructionViewMode(UE::Chaos::ClothAsset::EClothPatternVertexType InMode)
 {
-	const bool bSwitching2D3D = (ConstructionViewMode == UE::Chaos::ClothAsset::EClothPatternVertexType::Sim2D) != (InMode == UE::Chaos::ClothAsset::EClothPatternVertexType::Sim2D);
-
 	ConstructionViewMode = InMode;
 	ReinitializeDynamicMeshComponents();
 
@@ -815,10 +830,8 @@ void UChaosClothAssetEditorMode::SetConstructionViewMode(UE::Chaos::ClothAsset::
 		VC->SetConstructionViewMode(ConstructionViewMode);
 	}
 
-	if (bSwitching2D3D)
-	{
-		RefocusRestSpaceViewportClient();
-	}
+	// If we are switching to a mode with a valid mesh for the first time, focus the camera on it
+	FirstTimeFocusRestSpaceViewport();
 }
 
 UE::Chaos::ClothAsset::EClothPatternVertexType UChaosClothAssetEditorMode::GetConstructionViewMode() const
diff --git a/Engine/Plugins/Experimental/ChaosClothAssetEditor/Source/ChaosClothAssetEditor/Private/ChaosClothAsset/ClothEditorMode.h b/Engine/Plugins/Experimental/ChaosClothAssetEditor/Source/ChaosClothAssetEditor/Private/ChaosClothAsset/ClothEditorMode.h
--- a/Engine/Plugins/Experimental/ChaosClothAssetEditor/Source/ChaosClothAssetEditor/Private/ChaosClothAsset/ClothEditorMode.h
+++ b/Engine/Plugins/Experimental/ChaosClothAssetEditor/Source/ChaosClothAssetEditor/Private/ChaosClothAsset/ClothEditorMode.h
@@ -133,6 +133,7 @@ private:
 
 	void SetRestSpaceViewportClient(TWeakPtr<UE::Chaos::ClothAsset::FChaosClothEditorRestSpaceViewportClient, ESPMode::ThreadSafe> ViewportClient);
 	void RefocusRestSpaceViewportClient();
+	void FirstTimeFocusRestSpaceViewport();
 
 	// intended to be called by the toolkit when selected node in the Dataflow graph changes
 	void SetSelectedClothCollection(TSharedPtr<FManagedArrayCollection> Collection);
@@ -202,8 +203,9 @@ private:
 
 	TWeakPtr<UE::Chaos::ClothAsset::FChaosClothEditorRestSpaceViewportClient, ESPMode::ThreadSafe> RestSpaceViewportClient;
 
-	// Whether a cloth collection has ever been set on this Mode object
-	bool bFirstClothCollection = true;
+	// The first time we get a valid mesh, refocus the camera on it
+	bool bFirstValid2DMesh = true;
+	bool bFirstValid3DMesh = true;
 
 	// Whether the rest space viewport should focus on the rest space mesh on the next tick
 	bool bShouldFocusRestSpaceView = true;
diff --git a/Engine/Plugins/Experimental/ChaosClothAssetEditor/Source/ChaosClothAssetEditor/Private/ChaosClothAsset/ClothEditorRestSpaceViewportClient.cpp b/Engine/Plugins/Experimental/ChaosClothAssetEditor/Source/ChaosClothAssetEditor/Private/ChaosClothAsset/ClothEditorRestSpaceViewportClient.cpp
--- a/Engine/Plugins/Experimental/ChaosClothAssetEditor/Source/ChaosClothAssetEditor/Private/ChaosClothAsset/ClothEditorRestSpaceViewportClient.cpp
+++ b/Engine/Plugins/Experimental/ChaosClothAssetEditor/Source/ChaosClothAssetEditor/Private/ChaosClothAsset/ClothEditorRestSpaceViewportClient.cpp
@@ -57,6 +57,12 @@ FChaosClothEditorRestSpaceViewportClient::FChaosClothEditorRestSpaceViewportClie
 
 void FChaosClothEditorRestSpaceViewportClient::SetConstructionViewMode(EClothPatternVertexType InViewMode)
 {
+	const bool bSwitching2D3D = (ConstructionViewMode == UE::Chaos::ClothAsset::EClothPatternVertexType::Sim2D) != (InViewMode == UE::Chaos::ClothAsset::EClothPatternVertexType::Sim2D);
+	if (bSwitching2D3D)
+	{
+		Swap(SavedInactiveViewTransform, ViewTransformPerspective);
+	}
+
 	ConstructionViewMode = InViewMode;
 
 	BehaviorSet->RemoveAll();
diff --git a/Engine/Plugins/Experimental/ChaosClothAssetEditor/Source/ChaosClothAssetEditor/Private/ChaosClothAsset/ClothEditorRestSpaceViewportClient.h b/Engine/Plugins/Experimental/ChaosClothAssetEditor/Source/ChaosClothAssetEditor/Private/ChaosClothAsset/ClothEditorRestSpaceViewportClient.h
--- a/Engine/Plugins/Experimental/ChaosClothAssetEditor/Source/ChaosClothAssetEditor/Private/ChaosClothAsset/ClothEditorRestSpaceViewportClient.h
+++ b/Engine/Plugins/Experimental/ChaosClothAssetEditor/Source/ChaosClothAssetEditor/Private/ChaosClothAsset/ClothEditorRestSpaceViewportClient.h
@@ -49,5 +49,8 @@ private:
 	TUniquePtr<FEditor2DMouseWheelZoomBehaviorTarget> ZoomBehaviorTarget;
 
 	TWeakPtr<FUICommandList> ToolCommandList;
+
+	// Saved view transform for the currently inactive view mode (i.e. store the 3D camera here while in 2D mode and vice-versa)
+	FViewportCameraTransform SavedInactiveViewTransform;
 };
 } // namespace UE::Chaos::ClothAsset
