ClothAssetEditor: Remove ensures when debug drawing that the cloth proxy exists. It won't exist when you load an empty asset.

diff --git a/Engine/Plugins/Experimental/ChaosClothAssetEditor/Source/ChaosClothAssetEditor/Private/ChaosClothAsset/ClothEditorSimulationVisualization.cpp b/Engine/Plugins/Experimental/ChaosClothAssetEditor/Source/ChaosClothAssetEditor/Private/ChaosClothAsset/ClothEditorSimulationVisualization.cpp
--- a/Engine/Plugins/Experimental/ChaosClothAssetEditor/Source/ChaosClothAssetEditor/Private/ChaosClothAsset/ClothEditorSimulationVisualization.cpp
+++ b/Engine/Plugins/Experimental/ChaosClothAssetEditor/Source/ChaosClothAssetEditor/Private/ChaosClothAsset/ClothEditorSimulationVisualization.cpp
@@ -205,7 +205,7 @@ void FClothEditorSimulationVisualization::ExtendViewportShowMenu(FMenuBuilder& M
 void FClothEditorSimulationVisualization::DebugDrawSimulation(const UChaosClothComponent* ClothComponent, FPrimitiveDrawInterface* PDI)
 {
 	const ::Chaos::FClothVisualization* const Visualization = ClothComponent && ClothComponent->GetClothSimulationProxy() ? ClothComponent->GetClothSimulationProxy()->GetVisualization() : nullptr;
-	if (!ensure(Visualization))
+	if (!Visualization)
 	{
 		return;
 	}
@@ -223,7 +223,7 @@ void FClothEditorSimulationVisualization::DebugDrawSimulation(const UChaosClothC
 void FClothEditorSimulationVisualization::DebugDrawSimulationTexts(const UChaosClothComponent* ClothComponent, FCanvas* Canvas, const FSceneView* SceneView)
 {
 	const ::Chaos::FClothVisualization* const Visualization = ClothComponent && ClothComponent->GetClothSimulationProxy() ? ClothComponent->GetClothSimulationProxy()->GetVisualization() : nullptr;
-	if (!ensure(Visualization))
+	if (!Visualization)
 	{
 		return;
 	}
@@ -240,7 +240,7 @@ void FClothEditorSimulationVisualization::DebugDrawSimulationTexts(const UChaosC
 FText FClothEditorSimulationVisualization::GetDisplayString(const UChaosClothComponent* ClothComponent) const
 {
 	const FClothSimulationProxy* const SimProxy = ClothComponent ? ClothComponent->GetClothSimulationProxy() : nullptr;
-	if (!ensure(SimProxy))
+	if (!SimProxy)
 	{
 		return FText();
 	}
