ChaosCloth: Fix SkeletalMeshImport index remapping when merging duplicate vertices.

Also add UVScale option and set default UVIndex to 0 instead of -1.

diff --git a/Engine/Plugins/Experimental/ChaosClothAssetEditor/Source/ChaosClothAssetDataflowNodes/Private/ChaosClothAsset/ClothDataflowTools.cpp b/Engine/Plugins/Experimental/ChaosClothAssetEditor/Source/ChaosClothAssetDataflowNodes/Private/ChaosClothAsset/ClothDataflowTools.cpp
--- a/Engine/Plugins/Experimental/ChaosClothAssetEditor/Source/ChaosClothAssetDataflowNodes/Private/ChaosClothAsset/ClothDataflowTools.cpp
+++ b/Engine/Plugins/Experimental/ChaosClothAssetEditor/Source/ChaosClothAssetDataflowNodes/Private/ChaosClothAsset/ClothDataflowTools.cpp
@@ -42,7 +42,7 @@ namespace UE::Chaos::ClothAsset
 
 				// We need to weld the mesh verts to get rid of duplicates (happens for smoothing groups)
 				TArray<FVector> UniqueVerts;
-				IndexRemap.AddDefaulted(NumVerts);
+				OriginalToMerged.AddDefaulted(NumVerts);
 				constexpr float ThreshSq = UE_THRESH_POINTS_ARE_SAME * UE_THRESH_POINTS_ARE_SAME;
 				for (int32 VertIndex = 0; VertIndex < NumVerts; ++VertIndex)
 				{
@@ -71,11 +71,11 @@ namespace UE::Chaos::ClothAsset
 						// Unique
 						UniqueVerts.Add((FVector)SourceVert.Position);
 						OriginalIndexes.Add(VertIndex);
-						IndexRemap[VertIndex] = UniqueVerts.Num() - 1;
+						OriginalToMerged[VertIndex] = VertIndex;
 					}
 					else
 					{
-						IndexRemap[VertIndex] = RemapIndex;
+						OriginalToMerged[VertIndex] = OriginalIndexes[RemapIndex];
 					}
 				}
 
@@ -123,9 +123,9 @@ namespace UE::Chaos::ClothAsset
 			{
 				if (TriID >= 0 && TriID < (int32)SourceSection.NumTriangles)
 				{
-					VID0 = IndexRemap[IndexBuffer[3 * TriID + 0] - SourceSection.BaseVertexIndex];
-					VID1 = IndexRemap[IndexBuffer[3 * TriID + 1] - SourceSection.BaseVertexIndex];
-					VID2 = IndexRemap[IndexBuffer[3 * TriID + 2] - SourceSection.BaseVertexIndex];
+					VID0 = OriginalToMerged[IndexBuffer[3 * TriID + 0] - SourceSection.BaseVertexIndex];
+					VID1 = OriginalToMerged[IndexBuffer[3 * TriID + 1] - SourceSection.BaseVertexIndex];
+					VID2 = OriginalToMerged[IndexBuffer[3 * TriID + 2] - SourceSection.BaseVertexIndex];
 					return true;
 				}
 				return false;
@@ -350,7 +350,7 @@ namespace UE::Chaos::ClothAsset
 			const FSkelMeshSection& SourceSection;
 			const TConstArrayView<uint32> IndexBuffer;
 			TArray<int32> OriginalIndexes; // UniqueIndex -> OrigIndex
-			TArray<int32> IndexRemap; // OrigIndex -> UniqueIndex
+			TArray<int32> OriginalToMerged; // OriginalIndex -> OriginalIndexes[UniqueVertIndex] 
 			TArray<int32> TriIDs;
 			TArray<int32> EmptyArray;
 
@@ -418,7 +418,7 @@ namespace UE::Chaos::ClothAsset
 		ClothPatternFacade.SetRenderMaterialPathName(RenderMaterialPathName);
 	}
 	
-	void FClothDataflowTools::AddSimPatternsFromSkeletalMeshSection(const TSharedRef<FManagedArrayCollection>& ClothCollection, const FSkeletalMeshLODModel& SkeletalMeshModel, const int32 SectionIndex, const int32 UVChannelIndex)
+	void FClothDataflowTools::AddSimPatternsFromSkeletalMeshSection(const TSharedRef<FManagedArrayCollection>& ClothCollection, const FSkeletalMeshLODModel& SkeletalMeshModel, const int32 SectionIndex, const int32 UVChannelIndex, const FVector2f& UVScale)
 	{
 		check(SectionIndex < SkeletalMeshModel.Sections.Num());
 
@@ -435,7 +435,7 @@ namespace UE::Chaos::ClothAsset
 		UE::Geometry::FNonManifoldMappingSupport::AttachNonManifoldVertexMappingData(SkelMeshSectionToDynamicMesh.ToSrcVertIDMap, DynamicMesh);
 
 		constexpr bool bAppend = true;
-		FClothGeometryTools::BuildSimMeshFromDynamicMesh(ClothCollection, DynamicMesh, UVChannelIndex, FVector2f(1.f), bAppend);
+		FClothGeometryTools::BuildSimMeshFromDynamicMesh(ClothCollection, DynamicMesh, UVChannelIndex, UVScale, bAppend);
 	}
 
 	void FClothDataflowTools::LogAndToastWarning(const FDataflowNode& DataflowNode, const FText& Headline, const FText& Details)
diff --git a/Engine/Plugins/Experimental/ChaosClothAssetEditor/Source/ChaosClothAssetDataflowNodes/Private/ChaosClothAsset/ClothDataflowTools.h b/Engine/Plugins/Experimental/ChaosClothAssetEditor/Source/ChaosClothAssetDataflowNodes/Private/ChaosClothAsset/ClothDataflowTools.h
--- a/Engine/Plugins/Experimental/ChaosClothAssetEditor/Source/ChaosClothAssetDataflowNodes/Private/ChaosClothAsset/ClothDataflowTools.h
+++ b/Engine/Plugins/Experimental/ChaosClothAssetEditor/Source/ChaosClothAssetDataflowNodes/Private/ChaosClothAsset/ClothDataflowTools.h
@@ -19,7 +19,7 @@ namespace UE::Chaos::ClothAsset
 	{
 		static void AddRenderPatternFromSkeletalMeshSection(const TSharedRef<FManagedArrayCollection>& ClothCollection, const FSkeletalMeshLODModel& SkeletalMeshModel, const int32 SectionIndex, const FString& RenderMaterialPathName);
 
-		static void AddSimPatternsFromSkeletalMeshSection(const TSharedRef<FManagedArrayCollection>& ClothCollection, const FSkeletalMeshLODModel& SkeletalMeshModel, const int32 SectionIndex, const int32 UVChannelIndex);
+		static void AddSimPatternsFromSkeletalMeshSection(const TSharedRef<FManagedArrayCollection>& ClothCollection, const FSkeletalMeshLODModel& SkeletalMeshModel, const int32 SectionIndex, const int32 UVChannelIndex, const FVector2f& UVScale);
 
 		static void LogAndToastWarning(const FDataflowNode& DataflowNode, const FText& Headline, const FText& Details);
 
diff --git a/Engine/Plugins/Experimental/ChaosClothAssetEditor/Source/ChaosClothAssetDataflowNodes/Private/ChaosClothAsset/SkeletalMeshImportNode.cpp b/Engine/Plugins/Experimental/ChaosClothAssetEditor/Source/ChaosClothAssetDataflowNodes/Private/ChaosClothAsset/SkeletalMeshImportNode.cpp
--- a/Engine/Plugins/Experimental/ChaosClothAssetEditor/Source/ChaosClothAssetDataflowNodes/Private/ChaosClothAsset/SkeletalMeshImportNode.cpp
+++ b/Engine/Plugins/Experimental/ChaosClothAssetEditor/Source/ChaosClothAssetDataflowNodes/Private/ChaosClothAsset/SkeletalMeshImportNode.cpp
@@ -66,7 +66,7 @@ void FChaosClothAssetSkeletalMeshImportNode::Evaluate(Dataflow::FContext& Contex
 
 			if (bImportSimMesh)
 			{
-				FClothDataflowTools::AddSimPatternsFromSkeletalMeshSection(ClothCollection, LODModel, SectionIndex, UVChannel);
+				FClothDataflowTools::AddSimPatternsFromSkeletalMeshSection(ClothCollection, LODModel, SectionIndex, UVChannel, UVScale);
 			}
 
 			if (bImportRenderMesh)
diff --git a/Engine/Plugins/Experimental/ChaosClothAssetEditor/Source/ChaosClothAssetDataflowNodes/Public/ChaosClothAsset/SkeletalMeshImportNode.h b/Engine/Plugins/Experimental/ChaosClothAssetEditor/Source/ChaosClothAssetDataflowNodes/Public/ChaosClothAsset/SkeletalMeshImportNode.h
--- a/Engine/Plugins/Experimental/ChaosClothAssetEditor/Source/ChaosClothAssetDataflowNodes/Public/ChaosClothAsset/SkeletalMeshImportNode.h
+++ b/Engine/Plugins/Experimental/ChaosClothAssetEditor/Source/ChaosClothAssetDataflowNodes/Public/ChaosClothAsset/SkeletalMeshImportNode.h
@@ -40,11 +40,15 @@ public:
 	bool bImportRenderMesh = true;
 
 	/**
-	 * UV channel of the skeletal meesh to import the 2D simulation mesh patterns from.
+	 * UV channel of the skeletal mesh to import the 2D simulation mesh patterns from.
 	 * If set to -1, then the import will unwrap the 3D simulation mesh into 2D simulation mesh patterns.
 	 */
-	UPROPERTY(EditAnywhere, Category = "Skeletal Mesh Import", Meta = (ClampMin = "-1"))
-	int32 UVChannel = INDEX_NONE;
+	UPROPERTY(EditAnywhere, Category = "Skeletal Mesh Import", Meta = (ClampMin = "-1", EditCondition = "bImportSimMesh"))
+	int32 UVChannel = 0;
+
+	/* Apply this scale to the UVs when populating Sim Mesh positions. */
+	UPROPERTY(EditAnywhere, Category = "Skeletal Mesh Import", Meta = (AllowPreserveRatio, EditCondition = "bImportSimMesh && UVChannel != INDEX_NONE"))
+	FVector2f UVScale = { 1.f, 1.f };
 
 	FChaosClothAssetSkeletalMeshImportNode(const Dataflow::FNodeParameters& InParam, FGuid InGuid = FGuid::NewGuid());
 
