Chaos Cloth: allow authoring bending constraints with flat/arbitrary rest shape.

Updated TriangleMesh to produce a consistent ordering of indices in GetUniqueAdjacentElements so that the resulting angle calculation signs are consistent. Convention is such that the angles calculated by the bending element constraints are 0 when flat, positive when folding away from the normals direction, negative when folding toward the normal direction.

Added 3 options:
* Default is use the 3D mesh, which preserved existing behavior.
* Use the "FlatnessRatio" parameter which is 0-1, 0 = use 3D mesh angle, 1 = use flat (0) rest angle. Lerp between those values
* Use the "RestAngle" parameter which just sets the rest angle directly. As suggested by Kriss, using the absolute min of the two vertex values, which allows us to create sharp creases.

Added debug visualization of the rest angles. Mostly was helpful to verify signs were correct when reordering GetUniqueAdjacentElements.

diff --git a/Engine/Plugins/Experimental/ChaosCloth/Source/Chaos/Private/ChaosCloth/ChaosClothVisualization.cpp b/Engine/Plugins/Experimental/ChaosCloth/Source/Chaos/Private/ChaosCloth/ChaosClothVisualization.cpp
--- a/Engine/Plugins/Experimental/ChaosCloth/Source/Chaos/Private/ChaosCloth/ChaosClothVisualization.cpp
+++ b/Engine/Plugins/Experimental/ChaosCloth/Source/Chaos/Private/ChaosCloth/ChaosClothVisualization.cpp
@@ -52,10 +52,11 @@ enum class EBendingDrawMode : int
 {
 	BuckleStatus = 0,
 	ParallelGraphColor = 1,
-	Anisotropy = 2
+	Anisotropy = 2,
+	RestAngle = 3
 };
 static int32 BendingDrawMode = (int32)EBendingDrawMode::BuckleStatus;
-static FAutoConsoleVariableRef CVarClothVizBendDrawMode(TEXT("p.ChaosClothVisualization.BendingDrawMode"), BendingDrawMode, TEXT("Bending draw mode, 0 = BuckleStatus, 1 = Parallel graph color, 2 = Anisotropy"));
+static FAutoConsoleVariableRef CVarClothVizBendDrawMode(TEXT("p.ChaosClothVisualization.BendingDrawMode"), BendingDrawMode, TEXT("Bending draw mode, 0 = BuckleStatus, 1 = Parallel graph color, 2 = Anisotropy, 3 = RestAngle"));
 
 enum class EStretchBiasDrawMode : int
 {
@@ -1598,6 +1599,22 @@ static FAutoConsoleVariableRef CVarClothVizStretchBiasDrawOutOfRange(TEXT("p.Cha
 			DrawLine(PDI, Pos0, Pos1, FLinearColor(Multiplier[0], Multiplier[1], Multiplier[2]));
 		}
 	}
+	static void DrawBendingElementRestAngle(FPrimitiveDrawInterface* PDI, const TConstArrayView<::Chaos::Softs::FSolverVec3>& Positions, const ::Chaos::FVec3& LocalSpaceLocation, const ::Chaos::Softs::FPBDBendingConstraintsBase* const BendingConstraints)
+	{
+		const TArray<TVec4<int32>>& Constraints = BendingConstraints->GetConstraints();
+		const TArray<Softs::FSolverReal>& RestAngles = BendingConstraints->GetRestAngles();
+		for (int32 ConstraintIndex = 0; ConstraintIndex < Constraints.Num(); ++ConstraintIndex)
+		{
+			const Softs::FSolverVec3& P1 = Positions[Constraints[ConstraintIndex][0]];
+			const Softs::FSolverVec3& P2 = Positions[Constraints[ConstraintIndex][1]];
+			const Softs::FSolverReal RestAngle = RestAngles[ConstraintIndex];
+			const uint8 ColorSat = (uint8)FMath::Clamp(FMath::Abs(RestAngle) / UE_PI * 256, 0, 255);
+
+			const FVec3 Pos0 = FVec3(P1) + LocalSpaceLocation;
+			const FVec3 Pos1 = FVec3(P2) + LocalSpaceLocation;
+			DrawLine(PDI, Pos0, Pos1, FLinearColor::MakeFromHSV8(RestAngle > 0 ? 170 : 0, ColorSat, 255));
+		}
+	}
 
 	void FClothVisualization::DrawBendingConstraint(FPrimitiveDrawInterface* PDI) const
 	{
@@ -1635,6 +1652,9 @@ static FAutoConsoleVariableRef CVarClothVizStretchBiasDrawOutOfRange(TEXT("p.Cha
 			{
 				switch ((Private::EBendingDrawMode)Private::BendingDrawMode)
 				{
+				case Private::EBendingDrawMode::RestAngle:
+					DrawBendingElementRestAngle(PDI, Positions, LocalSpaceLocation, BendingConstraints);
+					break;
 				case Private::EBendingDrawMode::ParallelGraphColor:
 					DrawSpringConstraintColors(PDI, Positions, LocalSpaceLocation, BendingConstraints);
 					break;
@@ -1649,6 +1669,9 @@ static FAutoConsoleVariableRef CVarClothVizStretchBiasDrawOutOfRange(TEXT("p.Cha
 			{
 				switch ((Private::EBendingDrawMode)Private::BendingDrawMode)
 				{
+				case Private::EBendingDrawMode::RestAngle:
+					DrawBendingElementRestAngle(PDI, Positions, LocalSpaceLocation, BendingConstraints);
+					break;
 				case Private::EBendingDrawMode::ParallelGraphColor:
 					DrawSpringConstraintColors(PDI, Positions, LocalSpaceLocation, BendingConstraints);
 					break;
@@ -1663,6 +1686,9 @@ static FAutoConsoleVariableRef CVarClothVizStretchBiasDrawOutOfRange(TEXT("p.Cha
 			{
 				switch ((Private::EBendingDrawMode)Private::BendingDrawMode)
 				{
+				case Private::EBendingDrawMode::RestAngle:
+					DrawBendingElementRestAngle(PDI, Positions, LocalSpaceLocation, BendingConstraints);
+					break;
 				case Private::EBendingDrawMode::ParallelGraphColor:
 					DrawSpringConstraintColors(PDI, Positions, LocalSpaceLocation, BendingConstraints);
 					break;
diff --git a/Engine/Plugins/Experimental/ChaosClothAssetEditor/Source/ChaosClothAssetDataflowNodes/Private/ChaosClothAsset/ChaosClothAssetDataflowNodesModule.cpp b/Engine/Plugins/Experimental/ChaosClothAssetEditor/Source/ChaosClothAssetDataflowNodes/Private/ChaosClothAsset/ChaosClothAssetDataflowNodesModule.cpp
--- a/Engine/Plugins/Experimental/ChaosClothAssetEditor/Source/ChaosClothAssetDataflowNodes/Private/ChaosClothAsset/ChaosClothAssetDataflowNodesModule.cpp
+++ b/Engine/Plugins/Experimental/ChaosClothAssetEditor/Source/ChaosClothAssetDataflowNodes/Private/ChaosClothAsset/ChaosClothAssetDataflowNodesModule.cpp
@@ -98,6 +98,7 @@ void FChaosClothAssetDataflowNodesModule::StartupModule()
 	if (FPropertyEditorModule* const PropertyModule = FModuleManager::GetModulePtr<FPropertyEditorModule>("PropertyEditor"))
 	{
 		PropertyModule->RegisterCustomPropertyTypeLayout("ChaosClothAssetWeightedValue", FOnGetPropertyTypeCustomizationInstance::CreateStatic(&FWeightedValueCustomization::MakeInstance));
+		PropertyModule->RegisterCustomPropertyTypeLayout("ChaosClothAssetWeightedValueNonAnimatable", FOnGetPropertyTypeCustomizationInstance::CreateStatic(&FWeightedValueCustomization::MakeInstance));
 	}
 }
 
@@ -107,6 +108,7 @@ void FChaosClothAssetDataflowNodesModule::ShutdownModule()
 	if (FPropertyEditorModule* const PropertyModule = FModuleManager::GetModulePtr<FPropertyEditorModule>("PropertyEditor"))
 	{
 		PropertyModule->UnregisterCustomPropertyTypeLayout("ChaosClothAssetWeightedValue");
+		PropertyModule->UnregisterCustomPropertyTypeLayout("ChaosClothAssetWeightedValueNonAnimatable");
 	}
 }
 
diff --git a/Engine/Plugins/Experimental/ChaosClothAssetEditor/Source/ChaosClothAssetDataflowNodes/Private/ChaosClothAsset/SimulationBaseConfigNodePrivate.h b/Engine/Plugins/Experimental/ChaosClothAssetEditor/Source/ChaosClothAssetDataflowNodes/Private/ChaosClothAsset/SimulationBaseConfigNodePrivate.h
--- a/Engine/Plugins/Experimental/ChaosClothAssetEditor/Source/ChaosClothAssetDataflowNodes/Private/ChaosClothAsset/SimulationBaseConfigNodePrivate.h
+++ b/Engine/Plugins/Experimental/ChaosClothAssetEditor/Source/ChaosClothAssetDataflowNodes/Private/ChaosClothAsset/SimulationBaseConfigNodePrivate.h
@@ -35,6 +35,12 @@
 		{ FName(TEXT(#SimilarProperty1)), FName(TEXT(#SimilarProperty2)) }); \
 	Properties.SetValue(Property##KeyIndex, Property);
 
+/** Macro for adding and setting enum value as int32 properties from inside AddProperties() and check for two similar properties being overriden. */
+#define UE_CHAOS_CLOTHASSET_SIMULATIONCONFIG_SETPROPERTYENUMCHECKED2(Property, SimilarProperty1, SimilarProperty2) \
+	const int32 Property##KeyIndex = AddPropertyHelper(Properties, FName(TEXT(#Property)), true, \
+		{ FName(TEXT(#SimilarProperty1)), FName(TEXT(#SimilarProperty2) )}); \
+	Properties.SetValue(Property##KeyIndex, (int32)Property);
+
 /** Macro for adding and setting weighted value properties from inside AddProperties() and check for a similar property being overriden. */
 #define UE_CHAOS_CLOTHASSET_SIMULATIONCONFIG_SETPROPERTYWEIGHTEDCHECKED1(Property, SimilarProperty1) \
 	const int32 Property##KeyIndex = AddPropertyHelper(Properties, FName(TEXT(#Property)), Property.bIsAnimatable, \
diff --git a/Engine/Plugins/Experimental/ChaosClothAssetEditor/Source/ChaosClothAssetDataflowNodes/Private/ChaosClothAsset/SimulationPBDBendingElementConfigNode.cpp b/Engine/Plugins/Experimental/ChaosClothAssetEditor/Source/ChaosClothAssetDataflowNodes/Private/ChaosClothAsset/SimulationPBDBendingElementConfigNode.cpp
--- a/Engine/Plugins/Experimental/ChaosClothAssetEditor/Source/ChaosClothAssetDataflowNodes/Private/ChaosClothAsset/SimulationPBDBendingElementConfigNode.cpp
+++ b/Engine/Plugins/Experimental/ChaosClothAssetEditor/Source/ChaosClothAssetDataflowNodes/Private/ChaosClothAsset/SimulationPBDBendingElementConfigNode.cpp
@@ -14,6 +14,18 @@ FChaosClothAssetSimulationPBDBendingElementConfigNode::FChaosClothAssetSimulatio
 
 void FChaosClothAssetSimulationPBDBendingElementConfigNode::AddProperties(::Chaos::Softs::FCollectionPropertyMutableFacade& Properties) const
 {
+	UE_CHAOS_CLOTHASSET_SIMULATIONCONFIG_SETPROPERTYENUMCHECKED2(
+		RestAngleType,
+		XPBDRestAngleType,        // Existing properties to warn against
+		XPBDAnisoRestAngleType);  //
+	UE_CHAOS_CLOTHASSET_SIMULATIONCONFIG_SETPROPERTYWEIGHTEDCHECKED2(
+		FlatnessRatio,
+		XPBDFlatnessRatio,        // Existing properties to warn against
+		XPBDAnisoFlatnessRatio);  //
+	UE_CHAOS_CLOTHASSET_SIMULATIONCONFIG_SETPROPERTYWEIGHTEDCHECKED2(
+		RestAngle,
+		XPBDRestAngle,        // Existing properties to warn against
+		XPBDAnisoRestAngle);  //
 	UE_CHAOS_CLOTHASSET_SIMULATIONCONFIG_SETPROPERTYWEIGHTEDCHECKED4(
 		BendingElementStiffness,
 		BendingSpringStiffness,          // Existing properties to warn against
diff --git a/Engine/Plugins/Experimental/ChaosClothAssetEditor/Source/ChaosClothAssetDataflowNodes/Private/ChaosClothAsset/SimulationXPBDAnisoBendingConfigNode.cpp b/Engine/Plugins/Experimental/ChaosClothAssetEditor/Source/ChaosClothAssetDataflowNodes/Private/ChaosClothAsset/SimulationXPBDAnisoBendingConfigNode.cpp
--- a/Engine/Plugins/Experimental/ChaosClothAssetEditor/Source/ChaosClothAssetDataflowNodes/Private/ChaosClothAsset/SimulationXPBDAnisoBendingConfigNode.cpp
+++ b/Engine/Plugins/Experimental/ChaosClothAssetEditor/Source/ChaosClothAssetDataflowNodes/Private/ChaosClothAsset/SimulationXPBDAnisoBendingConfigNode.cpp
@@ -14,6 +14,18 @@ FChaosClothAssetSimulationXPBDAnisoBendingConfigNode::FChaosClothAssetSimulation
 
 void FChaosClothAssetSimulationXPBDAnisoBendingConfigNode::AddProperties(::Chaos::Softs::FCollectionPropertyMutableFacade& Properties) const
 {
+	UE_CHAOS_CLOTHASSET_SIMULATIONCONFIG_SETPROPERTYENUMCHECKED2(
+		XPBDAnisoRestAngleType,
+		XPBDRestAngleType,  // Existing properties to warn against
+		RestAngleType);     //
+	UE_CHAOS_CLOTHASSET_SIMULATIONCONFIG_SETPROPERTYWEIGHTEDCHECKED2(
+		XPBDAnisoFlatnessRatio,
+		XPBDFlatnessRatio,  // Existing properties to warn against
+		FlatnessRatio);     //
+	UE_CHAOS_CLOTHASSET_SIMULATIONCONFIG_SETPROPERTYWEIGHTEDCHECKED2(
+		XPBDAnisoRestAngle,
+		XPBDRestAngle,  // Existing properties to warn against
+		RestAngle);     //
 	UE_CHAOS_CLOTHASSET_SIMULATIONCONFIG_SETPROPERTYWEIGHTEDCHECKED4(
 		XPBDAnisoBendingStiffnessWarp,
 		BendingSpringStiffness,        // Existing properties to warn against
diff --git a/Engine/Plugins/Experimental/ChaosClothAssetEditor/Source/ChaosClothAssetDataflowNodes/Private/ChaosClothAsset/SimulationXPBDBendingElementConfigNode.cpp b/Engine/Plugins/Experimental/ChaosClothAssetEditor/Source/ChaosClothAssetDataflowNodes/Private/ChaosClothAsset/SimulationXPBDBendingElementConfigNode.cpp
--- a/Engine/Plugins/Experimental/ChaosClothAssetEditor/Source/ChaosClothAssetDataflowNodes/Private/ChaosClothAsset/SimulationXPBDBendingElementConfigNode.cpp
+++ b/Engine/Plugins/Experimental/ChaosClothAssetEditor/Source/ChaosClothAssetDataflowNodes/Private/ChaosClothAsset/SimulationXPBDBendingElementConfigNode.cpp
@@ -14,6 +14,18 @@ FChaosClothAssetSimulationXPBDBendingElementConfigNode::FChaosClothAssetSimulati
 
 void FChaosClothAssetSimulationXPBDBendingElementConfigNode::AddProperties(::Chaos::Softs::FCollectionPropertyMutableFacade& Properties) const
 {
+	UE_CHAOS_CLOTHASSET_SIMULATIONCONFIG_SETPROPERTYENUMCHECKED2(
+		XPBDRestAngleType,
+		RestAngleType,            // Existing properties to warn against
+		XPBDAnisoRestAngleType);  //
+	UE_CHAOS_CLOTHASSET_SIMULATIONCONFIG_SETPROPERTYWEIGHTEDCHECKED2(
+		XPBDFlatnessRatio,
+		FlatnessRatio,            // Existing properties to warn against
+		XPBDAnisoFlatnessRatio);  //
+	UE_CHAOS_CLOTHASSET_SIMULATIONCONFIG_SETPROPERTYWEIGHTEDCHECKED2(
+		XPBDRestAngle,
+		RestAngle,            // Existing properties to warn against
+		XPBDAnisoRestAngle);  //
 	UE_CHAOS_CLOTHASSET_SIMULATIONCONFIG_SETPROPERTYWEIGHTEDCHECKED4(
 		XPBDBendingElementStiffness,
 		BendingSpringStiffness,          // Existing properties to warn against
diff --git a/Engine/Plugins/Experimental/ChaosClothAssetEditor/Source/ChaosClothAssetDataflowNodes/Public/ChaosClothAsset/SimulationPBDBendingElementConfigNode.h b/Engine/Plugins/Experimental/ChaosClothAssetEditor/Source/ChaosClothAssetDataflowNodes/Public/ChaosClothAsset/SimulationPBDBendingElementConfigNode.h
--- a/Engine/Plugins/Experimental/ChaosClothAssetEditor/Source/ChaosClothAssetDataflowNodes/Public/ChaosClothAsset/SimulationPBDBendingElementConfigNode.h
+++ b/Engine/Plugins/Experimental/ChaosClothAssetEditor/Source/ChaosClothAssetDataflowNodes/Public/ChaosClothAsset/SimulationPBDBendingElementConfigNode.h
@@ -4,8 +4,20 @@
 
 #include "ChaosClothAsset/SimulationBaseConfigNode.h"
 #include "ChaosClothAsset/WeightedValue.h"
+#include "Chaos/PBDBendingConstraintsBase.h"
 #include "SimulationPBDBendingElementConfigNode.generated.h"
 
+UENUM()
+enum class EChaosClothAssetRestAngleConstructionType : uint8
+{
+	/** Calculate rest angles using the 3D draped space simulation mesh. */
+	Use3DRestAngles = Chaos::Softs::FPBDBendingConstraintsBase::ERestAngleConstructionType::Use3DRestAngles,
+	/** Calculate rest angles using the FlatnessRatio property. */
+	FlatnessRatio = Chaos::Softs::FPBDBendingConstraintsBase::ERestAngleConstructionType::FlatnessRatio,
+	/** Calculate rest angles using the RestAngle property. */
+	RestAngle = Chaos::Softs::FPBDBendingConstraintsBase::ERestAngleConstructionType::ExplicitRestAngles
+};
+
 /** Bending element constraint property configuration node. */
 USTRUCT(Meta = (DataflowCloth))
 struct FChaosClothAssetSimulationPBDBendingElementConfigNode : public FChaosClothAssetSimulationBaseConfigNode
@@ -14,6 +26,34 @@ struct FChaosClothAssetSimulationPBDBendingElementConfigNode : public FChaosClot
 	DATAFLOW_NODE_DEFINE_INTERNAL(FChaosClothAssetSimulationPBDBendingElementConfigNode, "SimulationPBDBendingElementConfig", "Cloth", "Cloth Simulation PBD Bending Element Config")
 
 public:
+	/**
+	 * Method for calculating the rest angles of the constraints.
+	 */
+	UPROPERTY(EditAnywhere, Category = "PBDBendingElement Properties")
+	EChaosClothAssetRestAngleConstructionType RestAngleType = EChaosClothAssetRestAngleConstructionType::Use3DRestAngles;
+
+	/**
+	 * Calculate rest angles as a ratio between completely flat and whatever is the 3D rest angle.
+	 * When FlatnessRatio = 0, this is equivalent to "Use3DRestAngles".
+	 * When FlatnessRatio = 1, the rest angle will be 0 (completely flat).
+	 * If a valid weight map is found with the given Weight Map name, then both Low and High values
+	 * are interpolated with the per particle weight to make the final value used for the simulation.
+	 * Otherwise all particles are considered to have a zero weight, and only the Low value is meaningful.
+	 */
+	UPROPERTY(EditAnywhere, Category = "PBDBendingElement Properties", Meta = (UIMin = "0", UIMax = "1", ClampMin = "0", ClampMax = "1", EditCondition = "RestAngleType == EChaosClothAssetRestAngleConstructionType::FlatnessRatio"))
+	FChaosClothAssetWeightedValueNonAnimatable FlatnessRatio = { 0.f, 0.f, TEXT("FlatnessRatio") };
+
+	/**
+	 * Set rest angles to be the explicit value set here (in degrees).
+	 * 0 = Flat, Positive values fold away from the edge normal, Negative values fold toward the edge normal.
+	 * When converting vertex weight values to edge values, the value with the smallest absolute value is selected.
+	 * If a valid weight map is found with the given Weight Map name, then both Low and High values
+	 * are interpolated with the per particle weight to make the final value used for the simulation.
+	 * Otherwise all particles are considered to have a zero weight, and only the Low value is meaningful.
+	 */
+	UPROPERTY(EditAnywhere, Category = "PBDBendingElement Properties", Meta = (UIMin = "-180", UIMax = "180", ClampMin = "-180", ClampMax = "180", EditCondition = "RestAngleType == EChaosClothAssetRestAngleConstructionType::RestAngle"))
+	FChaosClothAssetWeightedValueNonAnimatable RestAngle = { 0.f, 0.f, TEXT("RestAngle") };
+
 	/**
 	 * The Stiffness of the bending elements constraints. Increase the iteration count for stiffer materials.
 	 * If a valid weight map is found with the given Weight Map name, then both Low and High values
@@ -31,7 +71,7 @@ public:
 	float BucklingRatio = 0.5f;
 
 	/**
-	 * The stiffness after bucking.
+	 * The stiffness after buckling.
 	 * The constraint will use this stiffness instead of element Stiffness once the cloth has buckled, i.e., bent beyond a certain angle.
 	 * Typically, Buckling Stiffness is set to be less than BendingElement Stiffness.
 	 * Buckling Ratio determines the switch point between using BendingElement Stiffness and Buckling Stiffness.
diff --git a/Engine/Plugins/Experimental/ChaosClothAssetEditor/Source/ChaosClothAssetDataflowNodes/Public/ChaosClothAsset/SimulationXPBDAnisoBendingConfigNode.h b/Engine/Plugins/Experimental/ChaosClothAssetEditor/Source/ChaosClothAssetDataflowNodes/Public/ChaosClothAsset/SimulationXPBDAnisoBendingConfigNode.h
--- a/Engine/Plugins/Experimental/ChaosClothAssetEditor/Source/ChaosClothAssetDataflowNodes/Public/ChaosClothAsset/SimulationXPBDAnisoBendingConfigNode.h
+++ b/Engine/Plugins/Experimental/ChaosClothAssetEditor/Source/ChaosClothAssetDataflowNodes/Public/ChaosClothAsset/SimulationXPBDAnisoBendingConfigNode.h
@@ -4,6 +4,7 @@
 
 #include "ChaosClothAsset/SimulationBaseConfigNode.h"
 #include "ChaosClothAsset/WeightedValue.h"
+#include "SimulationPBDBendingElementConfigNode.h"
 #include "SimulationXPBDAnisoBendingConfigNode.generated.h"
 
 /** XPBD anisotropic bending constraint property configuration node. */
@@ -14,6 +15,34 @@ struct FChaosClothAssetSimulationXPBDAnisoBendingConfigNode : public FChaosCloth
 	DATAFLOW_NODE_DEFINE_INTERNAL(FChaosClothAssetSimulationXPBDAnisoBendingConfigNode, "SimulationXPBDAnisoBendingConfig", "Cloth", "Cloth Simulation XPBD Anisotropic Bending Config")
 
 public:
+	/**
+	 * Method for calculating the rest angles of the constraints.
+	 */
+	UPROPERTY(EditAnywhere, Category = "XPBDAnisoBending Properties")
+	EChaosClothAssetRestAngleConstructionType XPBDAnisoRestAngleType = EChaosClothAssetRestAngleConstructionType::Use3DRestAngles;
+
+	/**
+	 * Calculate rest angles as a ratio between completely flat and whatever is the 3D rest angle.
+	 * When FlatnessRatio = 0, this is equivalent to "Use3DRestAngles".
+	 * When FlatnessRatio = 1, the rest angle will be 0 (completely flat).
+	 * If a valid weight map is found with the given Weight Map name, then both Low and High values
+	 * are interpolated with the per particle weight to make the final value used for the simulation.
+	 * Otherwise all particles are considered to have a zero weight, and only the Low value is meaningful.
+	 */
+	UPROPERTY(EditAnywhere, Category = "XPBDAnisoBending Properties", Meta = (UIMin = "0", UIMax = "1", ClampMin = "0", ClampMax = "1", EditCondition = "XPBDAnisoRestAngleType == EChaosClothAssetRestAngleConstructionType::FlatnessRatio"))
+	FChaosClothAssetWeightedValueNonAnimatable XPBDAnisoFlatnessRatio = { 0.f, 0.f, TEXT("XPBDAnisoFlatnessRatio") };
+
+	/**
+	 * Set rest angles to be the explicit value set here (in degrees).
+	 * 0 = Flat, Positive values fold away from the edge normal, Negative values fold toward the edge normal.
+	 * When converting vertex weight values to edge values, the value with the smallest absolute value is selected.
+	 * If a valid weight map is found with the given Weight Map name, then both Low and High values
+	 * are interpolated with the per particle weight to make the final value used for the simulation.
+	 * Otherwise all particles are considered to have a zero weight, and only the Low value is meaningful.
+	 */
+	UPROPERTY(EditAnywhere, Category = "XPBDAnisoBending Properties", Meta = (UIMin = "-180", UIMax = "180", ClampMin = "-180", ClampMax = "180", EditCondition = "XPBDAnisoRestAngleType == EChaosClothAssetRestAngleConstructionType::RestAngle"))
+	FChaosClothAssetWeightedValueNonAnimatable XPBDAnisoRestAngle = { 0.f, 0.f, TEXT("XPBDAnisoRestAngle") };
+
 	/**
 	 * The stiffness of the bending constraints in the warp (vertical) direction.
 	 * If a valid weight map is found with the given Weight Map name, then both Low and High values
@@ -42,7 +71,7 @@ public:
 	FChaosClothAssetWeightedValue XPBDAnisoBendingStiffnessBias = { true, 100.f, 100.f, TEXT("XPBDAnisoBendingStiffnessBias") };
 
 	/**
-	 * The damping of the bending constraints.
+	 * The damping of the bending constraints, relative to critical damping.
 	 * If a valid weight map is found with the given Weight Map name, then both Low and High values
 	 * are interpolated with the per particle weight to make the final value used for the simulation.
 	 * Otherwise all particles are considered to have a zero weight, and only the Low value is meaningful.
@@ -58,7 +87,7 @@ public:
 	float XPBDAnisoBucklingRatio = 0.5f;
 
 	/**
-	 * The stiffness after bucking in the warp (vertical) direction.
+	 * The stiffness after buckling in the warp (vertical) direction.
 	 * The constraint will use this stiffness instead of element Stiffness once the cloth has buckled, i.e., bent beyond a certain angle.
 	 * Typically, Buckling Stiffness is set to be less than BendingElement Stiffness.
 	 * Buckling Ratio determines the switch point between using BendingElement Stiffness and Buckling Stiffness.
@@ -66,11 +95,11 @@ public:
 	 * are interpolated with the per particle weight to make the final value used for the simulation.
 	 * Otherwise all particles are considered to have a zero weight, and only the Low value is meaningful.
 	 */
-	UPROPERTY(EditAnywhere, Category = "XPBDBendingElement Properties", Meta = (UIMin = "0", UIMax = "10000", ClampMin = "0", ClampMax = "10000000", EditCondition = "XPBDBucklingRatio != 0"))
+	UPROPERTY(EditAnywhere, Category = "XPBDAnisoBending Properties", Meta = (UIMin = "0", UIMax = "10000", ClampMin = "0", ClampMax = "10000000", EditCondition = "XPBDBucklingRatio != 0"))
 	FChaosClothAssetWeightedValue XPBDAnisoBucklingStiffnessWarp = { true, 50.f, 50.f, TEXT("XPBDAnisoBucklingStiffnessWarp") };
 
 	/**
-	 * The stiffness after bucking in the weft (horizontal) direction.
+	 * The stiffness after buckling in the weft (horizontal) direction.
 	 * The constraint will use this stiffness instead of element Stiffness once the cloth has buckled, i.e., bent beyond a certain angle.
 	 * Typically, Buckling Stiffness is set to be less than BendingElement Stiffness.
 	 * Buckling Ratio determines the switch point between using BendingElement Stiffness and Buckling Stiffness.
@@ -78,11 +107,11 @@ public:
 	 * are interpolated with the per particle weight to make the final value used for the simulation.
 	 * Otherwise all particles are considered to have a zero weight, and only the Low value is meaningful.
 	 */
-	UPROPERTY(EditAnywhere, Category = "XPBDBendingElement Properties", Meta = (UIMin = "0", UIMax = "10000", ClampMin = "0", ClampMax = "10000000", EditCondition = "XPBDBucklingRatio != 0"))
+	UPROPERTY(EditAnywhere, Category = "XPBDAnisoBending Properties", Meta = (UIMin = "0", UIMax = "10000", ClampMin = "0", ClampMax = "10000000", EditCondition = "XPBDBucklingRatio != 0"))
 	FChaosClothAssetWeightedValue XPBDAnisoBucklingStiffnessWeft = { true, 50.f, 50.f, TEXT("XPBDAnisoBucklingStiffnessWeft") };
 
 	/**
-	 * The stiffness after bucking in the bias (diagonal) direction.
+	 * The stiffness after buckling in the bias (diagonal) direction.
 	 * The constraint will use this stiffness instead of element Stiffness once the cloth has buckled, i.e., bent beyond a certain angle.
 	 * Typically, Buckling Stiffness is set to be less than BendingElement Stiffness.
 	 * Buckling Ratio determines the switch point between using BendingElement Stiffness and Buckling Stiffness.
@@ -90,7 +119,7 @@ public:
 	 * are interpolated with the per particle weight to make the final value used for the simulation.
 	 * Otherwise all particles are considered to have a zero weight, and only the Low value is meaningful.
 	 */
-	UPROPERTY(EditAnywhere, Category = "XPBDBendingElement Properties", Meta = (UIMin = "0", UIMax = "10000", ClampMin = "0", ClampMax = "10000000", EditCondition = "XPBDBucklingRatio != 0"))
+	UPROPERTY(EditAnywhere, Category = "XPBDAnisoBending Properties", Meta = (UIMin = "0", UIMax = "10000", ClampMin = "0", ClampMax = "10000000", EditCondition = "XPBDBucklingRatio != 0"))
 	FChaosClothAssetWeightedValue XPBDAnisoBucklingStiffnessBias = { true, 50.f, 50.f, TEXT("XPBDAnisoBucklingStiffnessBias") };
 
 	FChaosClothAssetSimulationXPBDAnisoBendingConfigNode(const Dataflow::FNodeParameters& InParam, FGuid InGuid = FGuid::NewGuid());
diff --git a/Engine/Plugins/Experimental/ChaosClothAssetEditor/Source/ChaosClothAssetDataflowNodes/Public/ChaosClothAsset/SimulationXPBDAnisoStretchConfigNode.h b/Engine/Plugins/Experimental/ChaosClothAssetEditor/Source/ChaosClothAssetDataflowNodes/Public/ChaosClothAsset/SimulationXPBDAnisoStretchConfigNode.h
--- a/Engine/Plugins/Experimental/ChaosClothAssetEditor/Source/ChaosClothAssetDataflowNodes/Public/ChaosClothAsset/SimulationXPBDAnisoStretchConfigNode.h
+++ b/Engine/Plugins/Experimental/ChaosClothAssetEditor/Source/ChaosClothAssetDataflowNodes/Public/ChaosClothAsset/SimulationXPBDAnisoStretchConfigNode.h
@@ -14,7 +14,7 @@ struct FChaosClothAssetSimulationXPBDAnisoStretchConfigNode : public FChaosCloth
 	DATAFLOW_NODE_DEFINE_INTERNAL(FChaosClothAssetSimulationXPBDAnisoStretchConfigNode, "SimulationXPBDAnisoStretchConfig", "Cloth", "Cloth Simulation XPBD Anisotropic Stretch Config")
 
 public:
-	/** Whether to use the 3D drapped space as rest lengths, or use the 2D pattern space instead. */
+	/** Whether to use the 3D draped space as rest lengths, or use the 2D pattern space instead. */
 	UPROPERTY(EditAnywhere, Category = "XPBDAnisoStretch Properties")
 	bool bXPBDAnisoStretchUse3dRestLengths = true;
 
@@ -46,7 +46,7 @@ public:
 	FChaosClothAssetWeightedValue XPBDAnisoStretchStiffnessBias = { true, 100.f, 100.f, TEXT("XPBDAnisoStretchStiffnessBias") };
 
 	/**
-	 * The damping of the stretch constraints.
+	 * The damping of the stretch constraints, relative to critical damping.
 	 * If a valid weight map is found with the given Weight Map name, then both Low and High values
 	 * are interpolated with the per particle weight to make the final value used for the simulation.
 	 * Otherwise all particles are considered to have a zero weight, and only the Low value is meaningful.
diff --git a/Engine/Plugins/Experimental/ChaosClothAssetEditor/Source/ChaosClothAssetDataflowNodes/Public/ChaosClothAsset/SimulationXPBDBendingElementConfigNode.h b/Engine/Plugins/Experimental/ChaosClothAssetEditor/Source/ChaosClothAssetDataflowNodes/Public/ChaosClothAsset/SimulationXPBDBendingElementConfigNode.h
--- a/Engine/Plugins/Experimental/ChaosClothAssetEditor/Source/ChaosClothAssetDataflowNodes/Public/ChaosClothAsset/SimulationXPBDBendingElementConfigNode.h
+++ b/Engine/Plugins/Experimental/ChaosClothAssetEditor/Source/ChaosClothAssetDataflowNodes/Public/ChaosClothAsset/SimulationXPBDBendingElementConfigNode.h
@@ -4,6 +4,7 @@
 
 #include "ChaosClothAsset/SimulationBaseConfigNode.h"
 #include "ChaosClothAsset/WeightedValue.h"
+#include "SimulationPBDBendingElementConfigNode.h"
 #include "SimulationXPBDBendingElementConfigNode.generated.h"
 
 /** XPBD bending element constraint property configuration node. */
@@ -14,6 +15,34 @@ struct FChaosClothAssetSimulationXPBDBendingElementConfigNode : public FChaosClo
 	DATAFLOW_NODE_DEFINE_INTERNAL(FChaosClothAssetSimulationXPBDBendingElementConfigNode, "SimulationXPBDBendingElementConfig", "Cloth", "Cloth Simulation XPBD Bending Element Config")
 
 public:
+	/**
+	 * Method for calculating the rest angles of the constraints.
+	 */
+	UPROPERTY(EditAnywhere, Category = "XPBDBendingElement Properties")
+	EChaosClothAssetRestAngleConstructionType XPBDRestAngleType = EChaosClothAssetRestAngleConstructionType::Use3DRestAngles;
+
+	/**
+	 * Calculate rest angles as a ratio between completely flat and whatever is the 3D rest angle.
+	 * When FlatnessRatio = 0, this is equivalent to "Use3DRestAngles".
+	 * When FlatnessRatio = 1, the rest angle will be 0 (completely flat).
+	 * If a valid weight map is found with the given Weight Map name, then both Low and High values
+	 * are interpolated with the per particle weight to make the final value used for the simulation.
+	 * Otherwise all particles are considered to have a zero weight, and only the Low value is meaningful.
+	 */
+	UPROPERTY(EditAnywhere, Category = "XPBDBendingElement Properties", Meta = (UIMin = "0", UIMax = "1", ClampMin = "0", ClampMax = "1", EditCondition = "XPBDRestAngleType == EChaosClothAssetRestAngleConstructionType::FlatnessRatio"))
+	FChaosClothAssetWeightedValueNonAnimatable XPBDFlatnessRatio = { 0.f, 0.f, TEXT("XPBDFlatnessRatio") };
+
+	/**
+	 * Set rest angles to be the explicit value set here (in degrees).
+	 * 0 = Flat, Positive values fold away from the edge normal, Negative values fold toward the edge normal.
+	 * When converting vertex weight values to edge values, the value with the smallest absolute value is selected.
+	 * If a valid weight map is found with the given Weight Map name, then both Low and High values
+	 * are interpolated with the per particle weight to make the final value used for the simulation.
+	 * Otherwise all particles are considered to have a zero weight, and only the Low value is meaningful.
+	 */
+	UPROPERTY(EditAnywhere, Category = "XPBDBendingElement Properties", Meta = (UIMin = "-180", UIMax = "180", ClampMin = "-180", ClampMax = "180", EditCondition = "XPBDRestAngleType == EChaosClothAssetRestAngleConstructionType::RestAngle"))
+	FChaosClothAssetWeightedValueNonAnimatable XPBDRestAngle = { 0.f, 0.f, TEXT("XPBDRestAngle") };
+
 	/**
 	 * The stiffness of the bending element constraints.
 	 * If a valid weight map is found with the given Weight Map name, then both Low and High values
@@ -24,7 +53,7 @@ public:
 	FChaosClothAssetWeightedValue XPBDBendingElementStiffness = { true, 100.f, 100.f, TEXT("XPBDBendingElementStiffness") };
 
 	/**
-	 * The damping of the bending element constraints.
+	 * The damping of the bending element constraints, relative to critical damping.
 	 * If a valid weight map is found with the given Weight Map name, then both Low and High values
 	 * are interpolated with the per particle weight to make the final value used for the simulation.
 	 * Otherwise all particles are considered to have a zero weight, and only the Low value is meaningful.
@@ -40,7 +69,7 @@ public:
 	float XPBDBucklingRatio = 0.5f;
 
 	/**
-	 * The stiffness after bucking.
+	 * The stiffness after buckling.
 	 * The constraint will use this stiffness instead of element Stiffness once the cloth has buckled, i.e., bent beyond a certain angle.
 	 * Typically, Buckling Stiffness is set to be less than BendingElement Stiffness.
 	 * Buckling Ratio determines the switch point between using BendingElement Stiffness and Buckling Stiffness.
diff --git a/Engine/Plugins/Experimental/ChaosClothAssetEditor/Source/ChaosClothAssetDataflowNodes/Public/ChaosClothAsset/WeightedValue.h b/Engine/Plugins/Experimental/ChaosClothAssetEditor/Source/ChaosClothAssetDataflowNodes/Public/ChaosClothAsset/WeightedValue.h
--- a/Engine/Plugins/Experimental/ChaosClothAssetEditor/Source/ChaosClothAssetDataflowNodes/Public/ChaosClothAsset/WeightedValue.h
+++ b/Engine/Plugins/Experimental/ChaosClothAssetEditor/Source/ChaosClothAssetDataflowNodes/Public/ChaosClothAsset/WeightedValue.h
@@ -41,3 +41,35 @@ struct FChaosClothAssetWeightedValue
 	UPROPERTY(EditAnywhere, Category = "Weighted Value")
 	FString WeightMap = TEXT("WeightMap");
 };
+
+USTRUCT()
+struct FChaosClothAssetWeightedValueNonAnimatable
+{
+	GENERATED_BODY()
+
+	static constexpr bool bIsAnimatable = false;
+
+	/**
+	 * Property value corresponding to the lower bound of the Weight Map.
+	 * A Weight Map stores a series of Weight values assigned to each point, all between 0 and 1.
+	 * The weights are used to interpolate the individual values from Low to High assigned to each different point.
+	 * A Weight of 0 always corresponds to the Low parameter value, and a Weight of 1 to the High parameter value.
+	 * The value for Low can be set to be bigger than for High in order to reverse the effect of the Weight Map.
+	 */
+	UPROPERTY(EditAnywhere, Category = "Weighted Value", Meta = (DisplayName = "Low Weight", ChaosClothAssetShortName = "Lo"))
+	float Low = 0.f;
+
+	/**
+	 * Property value corresponding to the upper bound of the Weight Map.
+	 * A Weight Map stores a series of Weight values assigned to each point, all between 0 and 1.
+	 * The weights are used to interpolate the individual values from Low to High assigned to each different point.
+	 * A Weight of 0 always corresponds to the Low parameter value, and a Weight of 1 to the High parameter value.
+	 * The value for Low can be set to be bigger than for High in order to reverse the effect of the Weight Map.
+	 */
+	UPROPERTY(EditAnywhere, Category = "Weighted Value", Meta = (DisplayName = "High Weight", ChaosClothAssetShortName = "Hi"))
+	float High = 1.f;
+
+	/** The name of the weight map for this property. */
+	UPROPERTY(EditAnywhere, Category = "Weighted Value")
+	FString WeightMap = TEXT("WeightMap");
+};
diff --git a/Engine/Source/Runtime/Experimental/Chaos/Private/Chaos/PBDBendingConstraintsBase.cpp b/Engine/Source/Runtime/Experimental/Chaos/Private/Chaos/PBDBendingConstraintsBase.cpp
new file mode 100644
--- /dev/null
+++ b/Engine/Source/Runtime/Experimental/Chaos/Private/Chaos/PBDBendingConstraintsBase.cpp
@@ -0,0 +1,119 @@
+// Copyright Epic Games, Inc. All Rights Reserved.
+#include "Chaos/PBDBendingConstraintsBase.h"
+
+namespace Chaos::Softs 
+{
+	namespace Private
+	{
+		void Calculate3DRestAngles(
+			const FSolverParticles& InParticles,
+			const TArray<TVec4<int32>>& Constraints,
+			TArray<FSolverReal>& RestAngles)
+		{
+			RestAngles.Reset(Constraints.Num());
+			for (const TVec4<int32>& Constraint : Constraints)
+			{
+				const FSolverVec3& P1 = InParticles.X(Constraint[0]);
+				const FSolverVec3& P2 = InParticles.X(Constraint[1]);
+				const FSolverVec3& P3 = InParticles.X(Constraint[2]);
+				const FSolverVec3& P4 = InParticles.X(Constraint[3]);
+				RestAngles.Add(FMath::Clamp(FPBDBendingConstraintsBase::CalcAngle(P1, P2, P3, P4), -UE_PI, UE_PI));
+			}
+		}
+
+		void CalculateFlatnessRestAngles(
+			const FSolverParticles& InParticles,
+			int32 InParticleOffset,
+			int32 InParticleCount,
+			const TConstArrayView<FRealSingle>& RestAngleMap,
+			const FSolverVec2& RestAngleValue,
+			const TArray<TVec4<int32>>& Constraints,
+			TArray<FSolverReal>& RestAngles)
+		{
+			if (RestAngleMap.Num() != InParticleCount && RestAngleValue[0] == 1.f)
+			{
+				// Special case where Flatness = 1, i.e., all RestAngles are 0.
+				RestAngles.SetNumZeroed(Constraints.Num());
+				return;
+			}
+
+			auto FlatnessValue = [InParticleOffset, InParticleCount, &RestAngleMap, &RestAngleValue](const TVec4<int32>& Constraint)->FSolverReal
+			{
+				if (RestAngleMap.Num() == InParticleCount)
+				{
+					return RestAngleValue[0] + (RestAngleValue[1] - RestAngleValue[0]) * 
+						(RestAngleMap[Constraint[0] - InParticleOffset] + RestAngleMap[Constraint[1] - InParticleOffset]) * (FSolverReal).5f;
+				}
+				else
+				{
+					return RestAngleValue[0];
+				}
+			};
+
+			RestAngles.Reset(Constraints.Num());
+			for (const TVec4<int32>& Constraint : Constraints)
+			{
+				const FSolverVec3& P1 = InParticles.X(Constraint[0]);
+				const FSolverVec3& P2 = InParticles.X(Constraint[1]);
+				const FSolverVec3& P3 = InParticles.X(Constraint[2]);
+				const FSolverVec3& P4 = InParticles.X(Constraint[3]);
+				const FSolverReal Flatness = FMath::Clamp(FlatnessValue(Constraint), (FSolverReal)0.f, (FSolverReal)1.f);
+				RestAngles.Add(FMath::Clamp(((FSolverReal)1.f - Flatness) * FPBDBendingConstraintsBase::CalcAngle(P1, P2, P3, P4), -UE_PI, UE_PI));
+			}
+		}
+
+		void CalculateExplicitRestAngles(
+			int32 InParticleOffset,
+			int32 InParticleCount,
+			const TConstArrayView<FRealSingle>& RestAngleMap,
+			const FSolverVec2& RestAngleValue,
+			const TArray<TVec4<int32>>& Constraints,
+			TArray<FSolverReal>& RestAngles)
+		{
+			auto MapAngleRadians = [InParticleOffset, InParticleCount, &RestAngleMap, &RestAngleValue](const TVec4<int32>& Constraint)->FSolverReal
+			{
+				if (RestAngleMap.Num() == InParticleCount)
+				{
+					const FSolverReal RestAngle0 = FMath::UnwindRadians(FMath::DegreesToRadians(
+						RestAngleValue[0] + (RestAngleValue[1] - RestAngleValue[0]) * (RestAngleMap[Constraint[0] - InParticleOffset])));
+					const FSolverReal RestAngle1 = FMath::UnwindRadians(FMath::DegreesToRadians(
+						RestAngleValue[0] + (RestAngleValue[1] - RestAngleValue[0]) * (RestAngleMap[Constraint[1] - InParticleOffset])));
+
+					return FMath::Abs(RestAngle0) < FMath::Abs(RestAngle1) ? RestAngle0 : RestAngle1;
+				}
+				else
+				{
+					return FMath::UnwindRadians(FMath::DegreesToRadians(RestAngleValue[0]));
+				}
+			};
+
+			RestAngles.Reset(Constraints.Num());
+			for (const TVec4<int32>& Constraint : Constraints)
+			{
+				RestAngles.Add(FMath::Clamp(MapAngleRadians(Constraint), -UE_PI, UE_PI));
+			}
+		}
+	}
+
+	void FPBDBendingConstraintsBase::CalculateRestAngles(
+		const FSolverParticles& InParticles,
+		int32 InParticleOffset,
+		int32 InParticleCount,
+		const TConstArrayView<FRealSingle>& RestAngleMap,
+		const FSolverVec2& RestAngleValue,
+		ERestAngleConstructionType RestAngleConstructionType)
+	{
+		switch (RestAngleConstructionType)
+		{
+		default:
+		case ERestAngleConstructionType::Use3DRestAngles:
+			Private::Calculate3DRestAngles(InParticles, Constraints, RestAngles);
+			break;
+		case ERestAngleConstructionType::FlatnessRatio:
+			Private::CalculateFlatnessRestAngles(InParticles, InParticleOffset, InParticleCount, RestAngleMap, RestAngleValue, Constraints, RestAngles);
+			break;
+		case ERestAngleConstructionType::ExplicitRestAngles:
+			Private::CalculateExplicitRestAngles(InParticleOffset, InParticleCount, RestAngleMap, RestAngleValue, Constraints, RestAngles);
+		}
+	}
+}  // End namespace Chaos::Softs
diff --git a/Engine/Source/Runtime/Experimental/Chaos/Private/Chaos/TriangleMesh.cpp b/Engine/Source/Runtime/Experimental/Chaos/Private/Chaos/TriangleMesh.cpp
--- a/Engine/Source/Runtime/Experimental/Chaos/Private/Chaos/TriangleMesh.cpp
+++ b/Engine/Source/Runtime/Experimental/Chaos/Private/Chaos/TriangleMesh.cpp
@@ -250,29 +250,50 @@ TArray<TVec2<int32>> FTriangleMesh::GetUniqueAdjacentPoints() const
 TArray<TVec4<int32>> FTriangleMesh::GetUniqueAdjacentElements() const
 {
 	// Build a map with a list of opposite points for every edges
-	TMap<TVec2<int32> /*Edge*/, TArray<int32> /*OppositePoints*/> EdgeMap;
+	// OppositePoint = { Index, bEdgeFlipped }
+	TMap<TVec2<int32> /*Edge*/, TArray<TTuple<int32, bool>> /*OppositePoints*/> EdgeMap;
 
 	auto SortedEdge = [](int32 P0, int32 P1) { return P0 <= P1 ? TVec2<int32>(P0, P1) : TVec2<int32>(P1, P0); };
 
 	for (const TVec3<int32>& Element : MElements)
 	{
-		EdgeMap.FindOrAdd(SortedEdge(Element[0], Element[1])).AddUnique(Element[2]);
-		EdgeMap.FindOrAdd(SortedEdge(Element[1], Element[2])).AddUnique(Element[0]);
-		EdgeMap.FindOrAdd(SortedEdge(Element[2], Element[0])).AddUnique(Element[1]);
+		EdgeMap.FindOrAdd(SortedEdge(Element[0], Element[1])).AddUnique({ Element[2], Element[0] > Element[1] });
+		EdgeMap.FindOrAdd(SortedEdge(Element[1], Element[2])).AddUnique({ Element[0], Element[1] > Element[2] });
+		EdgeMap.FindOrAdd(SortedEdge(Element[2], Element[0])).AddUnique({ Element[1], Element[2] > Element[0] });
 	}
 
+	auto OrderOppositePoints = [](const TVec2<int32>& Edge, const TTuple<int32, bool>& Opposite0, const TTuple<int32, bool>& Opposite1)->TVec4<int32>
+	{
+		// Unflipped opposite point before flipped
+		if (!Opposite0.Get<1>() && Opposite1.Get<1>())
+		{
+			return TVec4<int32>(Edge[0], Edge[1], Opposite0.Get<0>(), Opposite1.Get<0>());
+		}
+		if (Opposite0.Get<1>() && !Opposite1.Get<1>())
+		{
+			return TVec4<int32>(Edge[0], Edge[1], Opposite1.Get<0>(), Opposite0.Get<0>());
+		}
+		// Both same flipped, just order lowest to highest
+		if (Opposite0.Get<0>() < Opposite1.Get<0>())
+		{
+			return TVec4<int32>(Edge[0], Edge[1], Opposite0.Get<0>(), Opposite1.Get<0>());
+		}
+
+		return TVec4<int32>(Edge[0], Edge[1], Opposite1.Get<0>(), Opposite0.Get<0>());
+	};
+
 	// Build constraints
 	TArray<TVec4<int32>> BendingConstraints;
-	for (const TPair<TVec2<int32>, TArray<int32>>& EdgeOppositePoints : EdgeMap)
+	for (const TPair<TVec2<int32>, TArray<TTuple<int32, bool>>>& EdgeOppositePoints : EdgeMap)
 	{
 		const TVec2<int32>& Edge = EdgeOppositePoints.Key;
-		const TArray<int32>& OppositePoints = EdgeOppositePoints.Value;
+		const TArray<TTuple<int32, bool>>& OppositePoints = EdgeOppositePoints.Value;
 
 		for (int32 Index0 = 0; Index0 < OppositePoints.Num(); ++Index0)
 		{
 			for (int32 Index1 = Index0 + 1; Index1 < OppositePoints.Num(); ++Index1)
 			{
-				BendingConstraints.Add({ Edge[0], Edge[1], OppositePoints[Index0], OppositePoints[Index1] });
+				BendingConstraints.Add(OrderOppositePoints(Edge, OppositePoints[Index0], OppositePoints[Index1] ));
 			}
 		}
 	}
diff --git a/Engine/Source/Runtime/Experimental/Chaos/Private/Chaos/XPBDAnisotropicBendingConstraints.cpp b/Engine/Source/Runtime/Experimental/Chaos/Private/Chaos/XPBDAnisotropicBendingConstraints.cpp
--- a/Engine/Source/Runtime/Experimental/Chaos/Private/Chaos/XPBDAnisotropicBendingConstraints.cpp
+++ b/Engine/Source/Runtime/Experimental/Chaos/Private/Chaos/XPBDAnisotropicBendingConstraints.cpp
@@ -37,9 +37,12 @@ FXPBDAnisotropicBendingConstraints::FXPBDAnisotropicBendingConstraints(const FSo
 		TriangleMesh.GetUniqueAdjacentElements(),
 		WeightMaps.FindRef(GetXPBDAnisoBendingStiffnessWarpString(PropertyCollection, XPBDAnisoBendingStiffnessWarpName.ToString())),
 		WeightMaps.FindRef(GetXPBDAnisoBucklingStiffnessWarpString(PropertyCollection, XPBDAnisoBucklingStiffnessWarpName.ToString())),
+		GetRestAngleMapFromCollection(WeightMaps, PropertyCollection),
 		FSolverVec2(GetWeightedFloatXPBDAnisoBendingStiffnessWarp(PropertyCollection, MaxStiffness)),
 		(FSolverReal)GetXPBDAnisoBucklingRatio(PropertyCollection, 0.f),
 		FSolverVec2(GetWeightedFloatXPBDAnisoBucklingStiffnessWarp(PropertyCollection, MaxStiffness)),
+		GetRestAngleValueFromCollection(PropertyCollection),
+		(ERestAngleConstructionType)GetXPBDAnisoRestAngleType(PropertyCollection, (int32)ERestAngleConstructionType::Use3DRestAngles),
 		bTrimKinematicConstraints,
 		MaxStiffness)
 	, StiffnessWeft(
diff --git a/Engine/Source/Runtime/Experimental/Chaos/Public/Chaos/PBDBendingConstraints.h b/Engine/Source/Runtime/Experimental/Chaos/Public/Chaos/PBDBendingConstraints.h
--- a/Engine/Source/Runtime/Experimental/Chaos/Public/Chaos/PBDBendingConstraints.h
+++ b/Engine/Source/Runtime/Experimental/Chaos/Public/Chaos/PBDBendingConstraints.h
@@ -31,9 +31,12 @@ public:
 			MoveTemp(InConstraints),
 			WeightMaps.FindRef(GetBendingElementStiffnessString(PropertyCollection, BendingElementStiffnessName.ToString())),
 			WeightMaps.FindRef(GetBucklingStiffnessString(PropertyCollection, BucklingStiffnessName.ToString())),
+			GetRestAngleMapFromCollection(WeightMaps, PropertyCollection),
 			FSolverVec2(GetWeightedFloatBendingElementStiffness(PropertyCollection, 1.f)),
 			(FSolverReal)GetBucklingRatio(PropertyCollection, 0.f),  // BucklingRatio is clamped in base class
 			FSolverVec2(GetWeightedFloatBucklingStiffness(PropertyCollection, 1.f)),
+			GetRestAngleValueFromCollection(PropertyCollection),
+			(ERestAngleConstructionType)GetRestAngleType(PropertyCollection, (int32)ERestAngleConstructionType::Use3DRestAngles),
 			bTrimKinematicConstraints)
 	{
 		InitColor(InParticles);
@@ -119,6 +122,40 @@ private:
 	CHAOS_API void InitColor(const FSolverParticles& InParticles);
 	CHAOS_API void ApplyHelper(FSolverParticles& InParticles, const FSolverReal Dt, const int32 ConstraintIndex, const FSolverReal ExpStiffnessValue, const FSolverReal ExpBucklingValue) const;
 
+	TConstArrayView<FRealSingle> GetRestAngleMapFromCollection(
+		const TMap<FString, TConstArrayView<FRealSingle>>& WeightMaps,
+		const FCollectionPropertyConstFacade& PropertyCollection)
+	{
+		const ERestAngleConstructionType ConstructionType = (ERestAngleConstructionType)GetRestAngleType(PropertyCollection, (int32)ERestAngleConstructionType::Use3DRestAngles);
+
+		switch (ConstructionType)
+		{
+		default:
+		case ERestAngleConstructionType::Use3DRestAngles:
+			return TConstArrayView<FRealSingle>(); // Unused
+		case ERestAngleConstructionType::FlatnessRatio:
+			return WeightMaps.FindRef(GetFlatnessRatioString(PropertyCollection, FlatnessRatioName.ToString()));
+		case ERestAngleConstructionType::ExplicitRestAngles:
+			return WeightMaps.FindRef(GetRestAngleString(PropertyCollection, RestAngleName.ToString()));
+		}
+	}
+
+	FSolverVec2 GetRestAngleValueFromCollection(const FCollectionPropertyConstFacade& PropertyCollection)
+	{
+		const ERestAngleConstructionType ConstructionType = (ERestAngleConstructionType)GetRestAngleType(PropertyCollection, (int32)ERestAngleConstructionType::Use3DRestAngles);
+
+		switch (ConstructionType)
+		{
+		default:
+		case ERestAngleConstructionType::Use3DRestAngles:
+			return FSolverVec2(0.f); // Unused
+		case ERestAngleConstructionType::FlatnessRatio:
+			return FSolverVec2(GetWeightedFloatFlatnessRatio(PropertyCollection, 0.f));
+		case ERestAngleConstructionType::ExplicitRestAngles:
+			return FSolverVec2(GetWeightedFloatRestAngle(PropertyCollection, 0.f));
+		}
+	}
+
 	using Base::Constraints;
 	using Base::ParticleOffset;
 	using Base::ParticleCount;
@@ -131,6 +168,9 @@ private:
 	UE_CHAOS_DECLARE_PROPERTYCOLLECTION_NAME(BendingElementStiffness, float);
 	UE_CHAOS_DECLARE_PROPERTYCOLLECTION_NAME(BucklingRatio, float);
 	UE_CHAOS_DECLARE_PROPERTYCOLLECTION_NAME(BucklingStiffness, float);
+	UE_CHAOS_DECLARE_PROPERTYCOLLECTION_NAME(FlatnessRatio, float);
+	UE_CHAOS_DECLARE_PROPERTYCOLLECTION_NAME(RestAngle, float);
+	UE_CHAOS_DECLARE_PROPERTYCOLLECTION_NAME(RestAngleType, int32);
 };
 
 }  // End namespace Chaos::Softs
diff --git a/Engine/Source/Runtime/Experimental/Chaos/Public/Chaos/PBDBendingConstraintsBase.h b/Engine/Source/Runtime/Experimental/Chaos/Public/Chaos/PBDBendingConstraintsBase.h
--- a/Engine/Source/Runtime/Experimental/Chaos/Public/Chaos/PBDBendingConstraintsBase.h
+++ b/Engine/Source/Runtime/Experimental/Chaos/Public/Chaos/PBDBendingConstraintsBase.h
@@ -13,15 +13,25 @@ namespace Chaos::Softs
 class FPBDBendingConstraintsBase
 {
 public:
+	enum class ERestAngleConstructionType : uint8
+	{
+		Use3DRestAngles,
+		FlatnessRatio,
+		ExplicitRestAngles
+	};
+
 	FPBDBendingConstraintsBase(const FSolverParticles& InParticles,
 		int32 InParticleOffset,
-		int32 InParticleCount, 
+		int32 InParticleCount,
 		TArray<TVec4<int32>>&& InConstraints,
 		const TConstArrayView<FRealSingle>& StiffnessMultipliers,
 		const TConstArrayView<FRealSingle>& BucklingStiffnessMultipliers,
+		const TConstArrayView<FRealSingle>& RestAngleMap,
 		const FSolverVec2& InStiffness,
 		const FSolverReal InBucklingRatio,
 		const FSolverVec2& InBucklingStiffness,
+		const FSolverVec2& RestAngleValue,
+		ERestAngleConstructionType RestAngleConstructionType,
 		bool bTrimKinematicConstraints = false,
 		FSolverReal MaxStiffness = FPBDStiffness::DefaultPBDMaxStiffness)
 		: Constraints(bTrimKinematicConstraints ? TrimKinematicConstraints(InConstraints, InParticles): MoveTemp(InConstraints))
@@ -48,16 +58,26 @@ public:
 			FPBDStiffness::DefaultParameterFitBase,
 			MaxStiffness)
 	{
-		for (const TVec4<int32>& Constraint : Constraints)
-		{
-			const FSolverVec3& P1 = InParticles.X(Constraint[0]);
-			const FSolverVec3& P2 = InParticles.X(Constraint[1]);
-			const FSolverVec3& P3 = InParticles.X(Constraint[2]);
-			const FSolverVec3& P4 = InParticles.X(Constraint[3]);
-			RestAngles.Add(CalcAngle(P1, P2, P3, P4));
-		}
+		CalculateRestAngles(InParticles, ParticleOffset, ParticleCount, RestAngleMap, RestAngleValue, RestAngleConstructionType);
 	}
 
+	FPBDBendingConstraintsBase(const FSolverParticles& InParticles,
+		int32 InParticleOffset,
+		int32 InParticleCount,
+		TArray<TVec4<int32>>&& InConstraints,
+		const TConstArrayView<FRealSingle>& StiffnessMultipliers,
+		const TConstArrayView<FRealSingle>& BucklingStiffnessMultipliers,
+		const FSolverVec2& InStiffness,
+		const FSolverReal InBucklingRatio,
+		const FSolverVec2& InBucklingStiffness,
+		bool bTrimKinematicConstraints = false,
+		FSolverReal MaxStiffness = FPBDStiffness::DefaultPBDMaxStiffness)
+		:FPBDBendingConstraintsBase(InParticles, InParticleOffset, InParticleCount, MoveTemp(InConstraints),
+			StiffnessMultipliers, BucklingStiffnessMultipliers, TConstArrayView<FRealSingle>(),
+			InStiffness, InBucklingRatio, InBucklingStiffness, FSolverVec2((FSolverReal)0.f),
+			ERestAngleConstructionType::Use3DRestAngles, bTrimKinematicConstraints, MaxStiffness)
+	{}
+
 	UE_DEPRECATED(5.2, "Use one of the other constructors instead.")
 	FPBDBendingConstraintsBase(const FSolverParticles& InParticles, TArray<TVec4<int32>>&& InConstraints, const FSolverReal InStiffness = (FSolverReal)1.)
 		: Constraints(MoveTemp(InConstraints))
@@ -68,14 +88,7 @@ public:
 		, BucklingRatio(0.f)
 		, BucklingStiffness(FSolverVec2(InStiffness))
 	{
-		for (const TVec4<int32>& Constraint : Constraints)
-		{
-			const FSolverVec3& P1 = InParticles.X(Constraint[0]);
-			const FSolverVec3& P2 = InParticles.X(Constraint[1]);
-			const FSolverVec3& P3 = InParticles.X(Constraint[2]);
-			const FSolverVec3& P4 = InParticles.X(Constraint[3]);
-			RestAngles.Add(CalcAngle(P1, P2, P3, P4));
-		}
+		CalculateRestAngles(InParticles, ParticleOffset, ParticleCount, TConstArrayView<FRealSingle>(), FSolverVec2(0.f), ERestAngleConstructionType::Use3DRestAngles);
 	}
 
 	virtual ~FPBDBendingConstraintsBase() {}
@@ -257,6 +270,13 @@ private:
 		return ExtractedEdges;
 	}
 
+	CHAOS_API void CalculateRestAngles(const FSolverParticles& InParticles,
+		int32 InParticleOffset,
+		int32 InParticleCount,
+		const TConstArrayView<FRealSingle>& RestAngleMap,
+		const FSolverVec2& RestAngleValue,
+		ERestAngleConstructionType RestAngleConstructionType);
+	
 protected:
 	TArray<TVec4<int32>> Constraints;
 	TArray<TVec2<int32>> ConstraintSharedEdges; // Only shared edges are used for calculating weighted stiffnesses.
diff --git a/Engine/Source/Runtime/Experimental/Chaos/Public/Chaos/XPBDAnisotropicBendingConstraints.h b/Engine/Source/Runtime/Experimental/Chaos/Public/Chaos/XPBDAnisotropicBendingConstraints.h
--- a/Engine/Source/Runtime/Experimental/Chaos/Public/Chaos/XPBDAnisotropicBendingConstraints.h
+++ b/Engine/Source/Runtime/Experimental/Chaos/Public/Chaos/XPBDAnisotropicBendingConstraints.h
@@ -103,6 +103,40 @@ private:
 
 	CHAOS_API TArray<FSolverVec3> GenerateWarpWeftBiasBaseMultipliers(const TArray<TVec3<FVec2f>>& FaceVertexPatternPositions, const FTriangleMesh& TriangleMesh) const;
 
+	TConstArrayView<FRealSingle> GetRestAngleMapFromCollection(
+		const TMap<FString, TConstArrayView<FRealSingle>>& WeightMaps,
+		const FCollectionPropertyConstFacade& PropertyCollection)
+	{
+		const ERestAngleConstructionType ConstructionType = (ERestAngleConstructionType)GetXPBDAnisoRestAngleType(PropertyCollection, (int32)ERestAngleConstructionType::Use3DRestAngles);
+
+		switch (ConstructionType)
+		{
+		default:
+		case ERestAngleConstructionType::Use3DRestAngles:
+			return TConstArrayView<FRealSingle>(); // Unused
+		case ERestAngleConstructionType::FlatnessRatio:
+			return WeightMaps.FindRef(GetXPBDAnisoFlatnessRatioString(PropertyCollection, XPBDAnisoFlatnessRatioName.ToString()));
+		case ERestAngleConstructionType::ExplicitRestAngles:
+			return WeightMaps.FindRef(GetXPBDAnisoRestAngleString(PropertyCollection, XPBDAnisoRestAngleName.ToString()));
+		}
+	}
+
+	FSolverVec2 GetRestAngleValueFromCollection(const FCollectionPropertyConstFacade& PropertyCollection)
+	{
+		const ERestAngleConstructionType ConstructionType = (ERestAngleConstructionType)GetXPBDAnisoRestAngleType(PropertyCollection, (int32)ERestAngleConstructionType::Use3DRestAngles);
+
+		switch (ConstructionType)
+		{
+		default:
+		case ERestAngleConstructionType::Use3DRestAngles:
+			return FSolverVec2(0.f); // Unused
+		case ERestAngleConstructionType::FlatnessRatio:
+			return FSolverVec2(GetWeightedFloatXPBDAnisoFlatnessRatio(PropertyCollection, 0.f));
+		case ERestAngleConstructionType::ExplicitRestAngles:
+			return FSolverVec2(GetWeightedFloatXPBDAnisoRestAngle(PropertyCollection, 0.f));
+		}
+	}
+
 	using Base::Constraints;
 	using Base::ParticleOffset;
 	using Base::ParticleCount;
@@ -129,6 +163,9 @@ private:
 	UE_CHAOS_DECLARE_PROPERTYCOLLECTION_NAME(XPBDAnisoBucklingStiffnessWarp, float);
 	UE_CHAOS_DECLARE_PROPERTYCOLLECTION_NAME(XPBDAnisoBucklingStiffnessWeft, float);
 	UE_CHAOS_DECLARE_PROPERTYCOLLECTION_NAME(XPBDAnisoBucklingStiffnessBias, float);
+	UE_CHAOS_DECLARE_PROPERTYCOLLECTION_NAME(XPBDAnisoFlatnessRatio, float);
+	UE_CHAOS_DECLARE_PROPERTYCOLLECTION_NAME(XPBDAnisoRestAngle, float);
+	UE_CHAOS_DECLARE_PROPERTYCOLLECTION_NAME(XPBDAnisoRestAngleType, int32);
 };
 
 }  // End namespace Chaos::Softs
diff --git a/Engine/Source/Runtime/Experimental/Chaos/Public/Chaos/XPBDBendingConstraints.h b/Engine/Source/Runtime/Experimental/Chaos/Public/Chaos/XPBDBendingConstraints.h
--- a/Engine/Source/Runtime/Experimental/Chaos/Public/Chaos/XPBDBendingConstraints.h
+++ b/Engine/Source/Runtime/Experimental/Chaos/Public/Chaos/XPBDBendingConstraints.h
@@ -44,9 +44,12 @@ public:
 			MoveTemp(InConstraints),
 			WeightMaps.FindRef(GetXPBDBendingElementStiffnessString(PropertyCollection, XPBDBendingElementStiffnessName.ToString())),
 			WeightMaps.FindRef(GetXPBDBucklingStiffnessString(PropertyCollection, XPBDBucklingStiffnessName.ToString())),
+			GetRestAngleMapFromCollection(WeightMaps, PropertyCollection),
 			FSolverVec2(GetWeightedFloatXPBDBendingElementStiffness(PropertyCollection, MaxStiffness)),
 			(FSolverReal)GetXPBDBucklingRatio(PropertyCollection, 0.f),
 			FSolverVec2(GetWeightedFloatXPBDBucklingStiffness(PropertyCollection, MaxStiffness)),
+			GetRestAngleValueFromCollection(PropertyCollection),
+			(ERestAngleConstructionType)GetXPBDRestAngleType(PropertyCollection, (int32)ERestAngleConstructionType::Use3DRestAngles),
 			bTrimKinematicConstraints,
 			MaxStiffness)
 		, DampingRatio(
@@ -200,6 +203,40 @@ private:
 	CHAOS_API void InitColor(const FSolverParticles& InParticles);
 	CHAOS_API void ApplyHelper(FSolverParticles& Particles, const FSolverReal Dt, const int32 ConstraintIndex, const FSolverReal ExpStiffnessValue, const FSolverReal ExpBucklingValue, const FSolverReal DampingRatioValue) const;
 
+	TConstArrayView<FRealSingle> GetRestAngleMapFromCollection(
+		const TMap<FString, TConstArrayView<FRealSingle>>& WeightMaps,
+		const FCollectionPropertyConstFacade& PropertyCollection)
+	{
+		const ERestAngleConstructionType ConstructionType = (ERestAngleConstructionType)GetXPBDRestAngleType(PropertyCollection, (int32)ERestAngleConstructionType::Use3DRestAngles);
+
+		switch (ConstructionType)
+		{
+		default:
+		case ERestAngleConstructionType::Use3DRestAngles:
+			return TConstArrayView<FRealSingle>(); // Unused
+		case ERestAngleConstructionType::FlatnessRatio:
+			return WeightMaps.FindRef(GetXPBDFlatnessRatioString(PropertyCollection, XPBDFlatnessRatioName.ToString()));
+		case ERestAngleConstructionType::ExplicitRestAngles:
+			return WeightMaps.FindRef(GetXPBDRestAngleString(PropertyCollection, XPBDRestAngleName.ToString()));
+		}
+	}
+
+	FSolverVec2 GetRestAngleValueFromCollection(const FCollectionPropertyConstFacade& PropertyCollection)
+	{
+		const ERestAngleConstructionType ConstructionType = (ERestAngleConstructionType)GetXPBDRestAngleType(PropertyCollection, (int32)ERestAngleConstructionType::Use3DRestAngles);
+
+		switch (ConstructionType)
+		{
+		default:
+		case ERestAngleConstructionType::Use3DRestAngles:
+			return FSolverVec2(0.f); // Unused
+		case ERestAngleConstructionType::FlatnessRatio:
+			return FSolverVec2(GetWeightedFloatXPBDFlatnessRatio(PropertyCollection, 0.f));
+		case ERestAngleConstructionType::ExplicitRestAngles:
+			return FSolverVec2(GetWeightedFloatXPBDRestAngle(PropertyCollection, 0.f));
+		}
+	}
+
 	using Base::Constraints;
 	using Base::ParticleOffset;
 	using Base::ParticleCount;
@@ -215,6 +252,9 @@ private:
 	UE_CHAOS_DECLARE_PROPERTYCOLLECTION_NAME(XPBDBendingElementDamping, float);
 	UE_CHAOS_DECLARE_PROPERTYCOLLECTION_NAME(XPBDBucklingRatio, float);
 	UE_CHAOS_DECLARE_PROPERTYCOLLECTION_NAME(XPBDBucklingStiffness, float);
+	UE_CHAOS_DECLARE_PROPERTYCOLLECTION_NAME(XPBDFlatnessRatio, float);
+	UE_CHAOS_DECLARE_PROPERTYCOLLECTION_NAME(XPBDRestAngle, float);
+	UE_CHAOS_DECLARE_PROPERTYCOLLECTION_NAME(XPBDRestAngleType, int32);
 };
 
 }  // End namespace Chaos::Softs
