Fix Clang CIS error

diff --git a/Engine/Plugins/Experimental/ChaosClothAssetEditor/Source/ChaosClothAssetDataflowNodes/Public/ChaosClothAsset/SimulationPBDBendingElementConfigNode.h b/Engine/Plugins/Experimental/ChaosClothAssetEditor/Source/ChaosClothAssetDataflowNodes/Public/ChaosClothAsset/SimulationPBDBendingElementConfigNode.h
--- a/Engine/Plugins/Experimental/ChaosClothAssetEditor/Source/ChaosClothAssetDataflowNodes/Public/ChaosClothAsset/SimulationPBDBendingElementConfigNode.h
+++ b/Engine/Plugins/Experimental/ChaosClothAssetEditor/Source/ChaosClothAssetDataflowNodes/Public/ChaosClothAsset/SimulationPBDBendingElementConfigNode.h
@@ -11,11 +11,11 @@ UENUM()
 enum class EChaosClothAssetRestAngleConstructionType : uint8
 {
 	/** Calculate rest angles using the 3D draped space simulation mesh. */
-	Use3DRestAngles = Chaos::Softs::FPBDBendingConstraintsBase::ERestAngleConstructionType::Use3DRestAngles,
+	Use3DRestAngles = (uint8)Chaos::Softs::FPBDBendingConstraintsBase::ERestAngleConstructionType::Use3DRestAngles,
 	/** Calculate rest angles using the FlatnessRatio property. */
-	FlatnessRatio = Chaos::Softs::FPBDBendingConstraintsBase::ERestAngleConstructionType::FlatnessRatio,
+	FlatnessRatio = (uint8)Chaos::Softs::FPBDBendingConstraintsBase::ERestAngleConstructionType::FlatnessRatio,
 	/** Calculate rest angles using the RestAngle property. */
-	RestAngle = Chaos::Softs::FPBDBendingConstraintsBase::ERestAngleConstructionType::ExplicitRestAngles
+	RestAngle = (uint8)Chaos::Softs::FPBDBendingConstraintsBase::ERestAngleConstructionType::ExplicitRestAngles
 };
 
 /** Bending element constraint property configuration node. */
