ChaosCloth: Update TransferSkinWeightsNode to understand cloth collections with non-manifold and otherwise non-dynamicmesh friendly data.

diff --git a/Engine/Plugins/Experimental/ChaosClothAssetEditor/Source/ChaosClothAssetDataflowNodes/ChaosClothAssetDataflowNodes.Build.cs b/Engine/Plugins/Experimental/ChaosClothAssetEditor/Source/ChaosClothAssetDataflowNodes/ChaosClothAssetDataflowNodes.Build.cs
--- a/Engine/Plugins/Experimental/ChaosClothAssetEditor/Source/ChaosClothAssetDataflowNodes/ChaosClothAssetDataflowNodes.Build.cs
+++ b/Engine/Plugins/Experimental/ChaosClothAssetEditor/Source/ChaosClothAssetDataflowNodes/ChaosClothAssetDataflowNodes.Build.cs
@@ -21,6 +21,7 @@ public class ChaosClothAssetDataflowNodes : ModuleRules
 				"ChaosCloth",
 				"ChaosClothAsset",
 				"ChaosClothAssetEngine",
+				"ChaosClothAssetTools",
 				"ClothingSystemRuntimeCommon",
 				"CoreUObject",
 				"DataflowCore",
diff --git a/Engine/Plugins/Experimental/ChaosClothAssetEditor/Source/ChaosClothAssetDataflowNodes/Private/ChaosClothAsset/TransferSkinWeightsNode.cpp b/Engine/Plugins/Experimental/ChaosClothAssetEditor/Source/ChaosClothAssetDataflowNodes/Private/ChaosClothAsset/TransferSkinWeightsNode.cpp
--- a/Engine/Plugins/Experimental/ChaosClothAssetEditor/Source/ChaosClothAssetDataflowNodes/Private/ChaosClothAsset/TransferSkinWeightsNode.cpp
+++ b/Engine/Plugins/Experimental/ChaosClothAssetEditor/Source/ChaosClothAssetDataflowNodes/Private/ChaosClothAsset/TransferSkinWeightsNode.cpp
@@ -1,7 +1,9 @@
 // Copyright Epic Games, Inc. All Rights Reserved.
 
 #include "ChaosClothAsset/TransferSkinWeightsNode.h"
+#include "ChaosClothAsset/ClothPatternToDynamicMesh.h"
 #include "ChaosClothAsset/ClothDataflowTools.h"
+#include "ChaosClothAsset/ClothGeometryTools.h"
 #include "ChaosClothAsset/CollectionClothFacade.h"
 #include "Animation/Skeleton.h"
 #include "Dataflow/DataflowInputOutput.h"
@@ -10,6 +12,7 @@
 #include "DynamicMesh/DynamicVertexSkinWeightsAttribute.h"
 #include "DynamicMesh/MeshNormals.h"
 #include "DynamicMesh/MeshTransforms.h"
+#include "DynamicMesh/NonManifoldMappingSupport.h"
 #include "Engine/SkeletalMesh.h"
 #include "MeshDescriptionToDynamicMesh.h"
 #include "Operations/TransferBoneWeights.h"
@@ -34,29 +37,19 @@ namespace UE::Chaos::ClothAsset::Private
 	// Convert the cloth into DynamicMesh
 	// @todo: This should instead be handled by a cloth lod to dynamic mesh converter similar to functions in the ClothPatternToDynamicMesh.h
 	static bool SimClothToDynamicMesh(
-		const UE::Chaos::ClothAsset::FCollectionClothConstFacade& ClothFacade,
+		const TSharedRef<FManagedArrayCollection>& ClothCollection,
 		const FReferenceSkeleton& TargetRefSkeleton, // the reference skeleton to add to the dynamic mesh
 		UE::Geometry::FDynamicMesh3& WeldedSimMesh)
 	{
 		using namespace UE::Geometry;
+		using namespace UE::Chaos::ClothAsset;
 
 		// Convert the sim mesh to DynamicMesh. 
 		// @todo: FTransferBoneWeights should accept raw data arrays (vertices/triangles) to avoid this conversion
-		WeldedSimMesh.Clear();
-		for (const FVector3f& Pos : ClothFacade.GetSimPosition3D())
-		{
-			WeldedSimMesh.AppendVertex(FVector3d(Pos));
-		}
+		FClothPatternToDynamicMesh PatternToDynamicMesh;
+		constexpr bool bDisableAttributes = true;
+		PatternToDynamicMesh.Convert(ClothCollection, INDEX_NONE, EClothPatternVertexType::Sim3D, WeldedSimMesh, bDisableAttributes);
 
-		for (const FIntVector& Indices : ClothFacade.GetSimIndices3D())
-		{
-			const int TID = WeldedSimMesh.AppendTriangle(Indices[0], Indices[1], Indices[2]);
-			if (TID < 0) // Failed to add the triangle (non-manifold/duplicate)
-			{
-				return false;
-			}
-		}
-		
 		WeldedSimMesh.EnableAttributes();
 		FMeshNormals NormalsUtil(&WeldedSimMesh);
 		NormalsUtil.InitializeOverlayToPerVertexNormals(WeldedSimMesh.Attributes()->PrimaryNormals());
@@ -78,7 +71,7 @@ namespace UE::Chaos::ClothAsset::Private
 		const double NormalThreshold,
 		const double RadiusPercentage,
 		bool bUseParallel,
-		UE::Chaos::ClothAsset::FCollectionClothFacade& ClothFacade,
+		const TSharedRef<FManagedArrayCollection>& ClothCollection,
 		UE::Geometry::FTransferBoneWeights& TransferBoneWeights)
 	{
 		using namespace UE::Geometry;
@@ -86,8 +79,12 @@ namespace UE::Chaos::ClothAsset::Private
 		//
 		// Convert cloth sim mesh LOD to the welded dynamic sim mesh.
 		//
+		
+		// We cannot handle it if there are orphaned vertices, so strip those now.
+		UE::Chaos::ClothAsset::FClothGeometryTools::CleanupAndCompactMesh(ClothCollection);
+
 		FDynamicMesh3 WeldedSimMesh;
-		if (!ensure(SimClothToDynamicMesh(ClothFacade, TargetRefSkeleton, WeldedSimMesh)))
+		if (!ensure(SimClothToDynamicMesh(ClothCollection, TargetRefSkeleton, WeldedSimMesh)))
 		{
 			UE_LOG(LogChaosClothAssetDataflowNodes, Warning, TEXT("TransferSkinWeightsNode: Failed to weld the simulation mesh for LOD."));
 			return;
@@ -109,19 +106,115 @@ namespace UE::Chaos::ClothAsset::Private
 			return;
 		}
 
+		UE::Chaos::ClothAsset::FCollectionClothFacade ClothFacade(ClothCollection);
+
 		//
 		// Copy the new bone weight data from the welded sim mesh back to the cloth patterns.
 		//
-		ParallelFor(WeldedSimMesh.MaxVertexID(), [&ClothFacade, &WeldedSimMesh](int32 WeldedID)
+		const FNonManifoldMappingSupport NonManifoldMapping(WeldedSimMesh);
+		TArray<TArray<int32>> SimMeshToDynamicMesh;
+		if (NonManifoldMapping.IsNonManifoldVertexInSource())
 		{
-			const FDynamicMeshVertexSkinWeightsAttribute* OutAttribute = WeldedSimMesh.Attributes()->GetSkinWeightsAttribute(FSkeletalMeshAttributes::DefaultSkinWeightProfileName);
+			// WeldedSimMesh indices don't match cloth collection. 
+			SimMeshToDynamicMesh.SetNum(ClothFacade.GetNumSimVertices3D());
+			for (int32 DynamicMeshVert = 0; DynamicMeshVert < WeldedSimMesh.VertexCount(); ++DynamicMeshVert)
+			{
+				SimMeshToDynamicMesh[NonManifoldMapping.GetOriginalNonManifoldVertexID(DynamicMeshVert)].Add(DynamicMeshVert);
+			}
 
-			checkSlow(OutAttribute);
-			checkSlow(WeldedSimMesh.IsVertex(WeldedID));
-			checkSlow(WeldedID < ClothFacade.GetNumSimVertices3D());
-			OutAttribute->GetValue(WeldedID, ClothFacade.GetSimBoneIndices()[WeldedID],
-				ClothFacade.GetSimBoneWeights()[WeldedID]);
-		}, bUseParallel ? EParallelForFlags::None : EParallelForFlags::ForceSingleThread);
+			ParallelFor(ClothFacade.GetNumSimVertices3D(), [&ClothFacade, &WeldedSimMesh, &SimMeshToDynamicMesh](int32 SimVertexID)
+			{
+				const FDynamicMeshVertexSkinWeightsAttribute* const OutAttribute = WeldedSimMesh.Attributes()->GetSkinWeightsAttribute(FSkeletalMeshAttributes::DefaultSkinWeightProfileName);
+				checkSlow(OutAttribute);
+				if (!ensure(SimMeshToDynamicMesh[SimVertexID].Num() > 0))
+				{
+					ClothFacade.GetSimBoneIndices()[SimVertexID].Reset();
+					ClothFacade.GetSimBoneWeights()[SimVertexID].Reset();
+					return;
+				}
+				if (SimMeshToDynamicMesh[SimVertexID].Num() == 1)
+				{
+					// Simple most common case, one-to-one correspondence. Just copy over.
+					const int32 WeldedID = SimMeshToDynamicMesh[SimVertexID][0];
+					checkSlow(WeldedSimMesh.IsVertex(WeldedID));
+					OutAttribute->GetValue(WeldedID, ClothFacade.GetSimBoneIndices()[SimVertexID],
+						ClothFacade.GetSimBoneWeights()[SimVertexID]);
+				}
+				else
+				{
+					// Need to merge data because dynamic mesh split the original vertex
+					TMap<int32, TPair<float, int32>> CombinedData;
+					for (const int32 WeldedID : SimMeshToDynamicMesh[SimVertexID])
+					{
+						TArray<int32> Indices;
+						TArray<float> Weights;
+						checkSlow(WeldedSimMesh.IsVertex(WeldedID));
+						OutAttribute->GetValue(WeldedID, Indices, Weights);
+						check(Indices.Num() == Weights.Num());
+						for (int32 Idx = 0; Idx < Indices.Num(); ++Idx)
+						{
+							TPair<float, int32>& WeightedFloat = CombinedData.FindOrAdd(Indices[Idx]);
+							WeightedFloat.Get<0>() += Weights[Idx];
+							WeightedFloat.Get<1>() += 1;
+						}
+					}
+					TArray<int32>& BoneIndices = ClothFacade.GetSimBoneIndices()[SimVertexID];
+					TArray<float>& BoneWeights = ClothFacade.GetSimBoneWeights()[SimVertexID];
+					BoneIndices.Reset(CombinedData.Num());
+					BoneWeights.Reset(CombinedData.Num());
+					float WeightsSum = 0.f;
+					for (TMap<int32, TPair<float, int32>>::TConstIterator CombinedDataIter = CombinedData.CreateConstIterator(); CombinedDataIter; ++CombinedDataIter)
+					{
+						check(CombinedDataIter.Value().Get<1>() > 0);
+						BoneIndices.Add(CombinedDataIter.Key());
+						const float FloatVal = CombinedDataIter.Value().Get<0>() / (float)CombinedDataIter.Value().Get<1>();
+						BoneWeights.Add(FloatVal);
+						WeightsSum += FloatVal;
+					}
+					if (BoneIndices.Num() > MAX_TOTAL_INFLUENCES)
+					{
+						// Choose MAX_TOTAL_INFLUENCES highest weighted bones.
+						TArray<TPair<float, int32>> SortableData;
+						SortableData.Reserve(BoneIndices.Num());
+						for (int32 Idx = 0; Idx < BoneIndices.Num(); ++Idx)
+						{
+							SortableData.Emplace(BoneWeights[Idx], BoneIndices[Idx]);
+						}
+						SortableData.Sort([](const TPair<float, int32>& A, const TPair<float, int32>& B) { return A > B; });
+
+						BoneIndices.SetNum(MAX_TOTAL_INFLUENCES);
+						BoneWeights.SetNum(MAX_TOTAL_INFLUENCES);
+						WeightsSum = 0.f;
+						for (int32 Idx = 0; Idx < MAX_TOTAL_INFLUENCES; ++Idx)
+						{
+							BoneIndices[Idx] = SortableData[Idx].Get<1>();
+							BoneWeights[Idx] = SortableData[Idx].Get<0>();
+							WeightsSum += SortableData[Idx].Get<0>();
+						}
+					}
+
+					// Normalize weights
+					const float WeightsSumRecip = WeightsSum > UE_SMALL_NUMBER ? 1.f / WeightsSum : 0.f;
+					for (float& Weight : BoneWeights)
+					{
+						Weight *= WeightsSumRecip;
+					}
+				}
+			}, bUseParallel ? EParallelForFlags::None : EParallelForFlags::ForceSingleThread);
+		}
+		else
+		{
+			ParallelFor(WeldedSimMesh.MaxVertexID(), [&ClothFacade, &WeldedSimMesh, &NonManifoldMapping](int32 WeldedID)
+			{
+				const FDynamicMeshVertexSkinWeightsAttribute* OutAttribute = WeldedSimMesh.Attributes()->GetSkinWeightsAttribute(FSkeletalMeshAttributes::DefaultSkinWeightProfileName);
+
+				checkSlow(OutAttribute);
+				checkSlow(WeldedSimMesh.IsVertex(WeldedID));
+				checkSlow(WeldedID < ClothFacade.GetNumSimVertices3D());
+				OutAttribute->GetValue(WeldedID, ClothFacade.GetSimBoneIndices()[WeldedID],
+					ClothFacade.GetSimBoneWeights()[WeldedID]);
+			}, bUseParallel ? EParallelForFlags::None : EParallelForFlags::ForceSingleThread);
+		}
 
 		//
 		// Compute the bone weights for the render mesh by transferring weights from the sim mesh
@@ -249,7 +342,7 @@ void FChaosClothAssetTransferSkinWeightsNode::Evaluate(Dataflow::FContext& Conte
 
 			if (TransferMethod == EChaosClothAssetTransferSkinWeightsMethod::InpaintWeights)
 			{
-				TransferInpaintWeights(TargetRefSkeleton, NormalThreshold, RadiusPercentage, bUseParallel, ClothFacade, TransferBoneWeights);
+				TransferInpaintWeights(TargetRefSkeleton, NormalThreshold, RadiusPercentage, bUseParallel, ClothCollection, TransferBoneWeights);
 			}
 			else if (TransferMethod == EChaosClothAssetTransferSkinWeightsMethod::ClosestPointOnSurface)
 			{
diff --git a/Engine/Plugins/Experimental/ChaosClothAssetEditor/Source/ChaosClothAssetTools/Private/ChaosClothAsset/ClothPatternToDynamicMesh.cpp b/Engine/Plugins/Experimental/ChaosClothAssetEditor/Source/ChaosClothAssetTools/Private/ChaosClothAsset/ClothPatternToDynamicMesh.cpp
--- a/Engine/Plugins/Experimental/ChaosClothAssetEditor/Source/ChaosClothAssetTools/Private/ChaosClothAsset/ClothPatternToDynamicMesh.cpp
+++ b/Engine/Plugins/Experimental/ChaosClothAssetEditor/Source/ChaosClothAssetTools/Private/ChaosClothAsset/ClothPatternToDynamicMesh.cpp
@@ -586,7 +586,7 @@ private:
 };
 
 
-void FClothPatternToDynamicMesh::Convert(const TSharedRef<const FManagedArrayCollection> ClothCollection, int32 PatternIndex, EClothPatternVertexType VertexDataType, UE::Geometry::FDynamicMesh3& MeshOut)
+void FClothPatternToDynamicMesh::Convert(const TSharedRef<const FManagedArrayCollection> ClothCollection, int32 PatternIndex, EClothPatternVertexType VertexDataType, UE::Geometry::FDynamicMesh3& MeshOut, bool bDisableAttributes)
 {
 	const FCollectionClothConstFacade ClothFacade(ClothCollection);
 
@@ -595,26 +595,34 @@ void FClothPatternToDynamicMesh::Convert(const TSharedRef<const FManagedArrayCol
 	FClothPatternWrapper PatternWrapper(ClothFacade, PatternIndex, VertexDataType);
 
 	auto TriangleToGroupFunction = [](FClothPatternWrapper::TriIDType) { return 0; };
-	const bool bIsRenderType = VertexDataType == EClothPatternVertexType::Render;
-	auto TriangleToMaterialFunction = [PatternIndex, bIsRenderType, &ClothFacade](FClothPatternWrapper::TriIDType TriID)->int32
+
+	if (bDisableAttributes)
+	{
+		PatternToDynamicMesh.ConvertWOAttributes(MeshOut, PatternWrapper, TriangleToGroupFunction);
+	}
+	else
 	{
-		if (bIsRenderType)
+		constexpr bool bCopyTangents = false;
+		const bool bIsRenderType = VertexDataType == EClothPatternVertexType::Render;
+		auto TriangleToMaterialFunction = [PatternIndex, bIsRenderType, &ClothFacade](FClothPatternWrapper::TriIDType TriID)->int32
 		{
-			if (PatternIndex != INDEX_NONE)
+			if (bIsRenderType)
 			{
-				return PatternIndex;
+				if (PatternIndex != INDEX_NONE)
+				{
+					return PatternIndex;
+				}
+				const int32 FoundPattern = ClothFacade.FindRenderPatternByFaceIndex(TriID);
+				check(FoundPattern != INDEX_NONE);
+				return FoundPattern;
 			}
-			const int32 FoundPattern = ClothFacade.FindRenderPatternByFaceIndex(TriID);
-			check(FoundPattern != INDEX_NONE);
-			return FoundPattern;
-		}
 
-		// Sim meshes will have a default material with MaterialID zero applied
-		return 0;
-	};
+			// Sim meshes will have a default material with MaterialID zero applied
+			return 0;
+		};
 
-	constexpr bool bCopyTangents = false;
-	PatternToDynamicMesh.Convert(MeshOut, PatternWrapper, TriangleToGroupFunction, TriangleToMaterialFunction, bCopyTangents);
+		PatternToDynamicMesh.Convert(MeshOut, PatternWrapper, TriangleToGroupFunction, TriangleToMaterialFunction, bCopyTangents);
+	}
 
 	// Add non-manifold mapping data if DynamicMesh indices don't match orig indices
 	const int32 NumVertices = VertexDataType == EClothPatternVertexType::Render ? ClothFacade.GetNumRenderVertices() :
diff --git a/Engine/Plugins/Experimental/ChaosClothAssetEditor/Source/ChaosClothAssetTools/Public/ChaosClothAsset/ClothPatternToDynamicMesh.h b/Engine/Plugins/Experimental/ChaosClothAssetEditor/Source/ChaosClothAssetTools/Public/ChaosClothAsset/ClothPatternToDynamicMesh.h
--- a/Engine/Plugins/Experimental/ChaosClothAssetEditor/Source/ChaosClothAssetTools/Public/ChaosClothAsset/ClothPatternToDynamicMesh.h
+++ b/Engine/Plugins/Experimental/ChaosClothAssetEditor/Source/ChaosClothAssetTools/Public/ChaosClothAsset/ClothPatternToDynamicMesh.h
@@ -23,7 +23,7 @@ class CHAOSCLOTHASSETTOOLS_API FClothPatternToDynamicMesh
 {
 public:
 
-	void Convert(const TSharedRef<const FManagedArrayCollection> ClothCollection, int32 PatternIndex, EClothPatternVertexType VertexDataType, UE::Geometry::FDynamicMesh3& MeshOut);
+	void Convert(const TSharedRef<const FManagedArrayCollection> ClothCollection, int32 PatternIndex, EClothPatternVertexType VertexDataType, UE::Geometry::FDynamicMesh3& MeshOut, bool bDisableAttributes = false);
 
 	void Convert(const UChaosClothAsset* ClothAssetMeshIn, int32 LODIndex, int32 PatternIndex, EClothPatternVertexType VertexDataType, UE::Geometry::FDynamicMesh3& MeshOut);
 };
