ChaosCloth: Replace "DeleteRenderMesh" with "DeleteElement" dataflow node. All types of elements can now be deleted, as well as the entire sim mesh or render mesh. Currently can delete individual elements by manually inputting the indices in an array. Selection support will be added in the future.

The cloth facade methods for deleting elements was added in a previous CL. This just calls those methods. Note that elements that reference the deleted element will still exist (e.g., if you delete a 2D vertex, the associated 3D vertex will still exist, but point to INDEX_NONE in lookups). This is all cleaned up in the terminal node when constructing the final asset.

diff --git a/Engine/Plugins/Experimental/ChaosClothAssetEditor/Source/ChaosClothAssetDataflowNodes/Private/ChaosClothAsset/ChaosClothAssetDataflowNodesModule.cpp b/Engine/Plugins/Experimental/ChaosClothAssetEditor/Source/ChaosClothAssetDataflowNodes/Private/ChaosClothAsset/ChaosClothAssetDataflowNodesModule.cpp
--- a/Engine/Plugins/Experimental/ChaosClothAssetEditor/Source/ChaosClothAssetDataflowNodes/Private/ChaosClothAsset/ChaosClothAssetDataflowNodesModule.cpp
+++ b/Engine/Plugins/Experimental/ChaosClothAssetEditor/Source/ChaosClothAssetDataflowNodes/Private/ChaosClothAsset/ChaosClothAssetDataflowNodesModule.cpp
@@ -5,7 +5,7 @@
 #include "ChaosClothAsset/BindToRootBoneNode.h"
 #include "ChaosClothAsset/CopySimulationToRenderMeshNode.h"
 #include "ChaosClothAsset/DatasmithImportNode.h"
-#include "ChaosClothAsset/DeleteRenderMeshNode.h"
+#include "ChaosClothAsset/DeleteElementNode.h"
 #include "ChaosClothAsset/ImportNode.h"
 #include "ChaosClothAsset/MergeClothCollectionsNode.h"
 #include "ChaosClothAsset/ReverseNormalsNode.h"
@@ -52,7 +52,7 @@ namespace UE::Chaos::ClothAsset::Private
 		DATAFLOW_NODE_REGISTER_CREATION_FACTORY(FChaosClothAssetBindToRootBoneNode);
 		DATAFLOW_NODE_REGISTER_CREATION_FACTORY(FChaosClothAssetCopySimulationToRenderMeshNode);
 		DATAFLOW_NODE_REGISTER_CREATION_FACTORY(FChaosClothAssetDatasmithImportNode);
-		DATAFLOW_NODE_REGISTER_CREATION_FACTORY(FChaosClothAssetDeleteRenderMeshNode);
+		DATAFLOW_NODE_REGISTER_CREATION_FACTORY(FChaosClothAssetDeleteElementNode);
 		DATAFLOW_NODE_REGISTER_CREATION_FACTORY(FChaosClothAssetImportNode);
 		DATAFLOW_NODE_REGISTER_CREATION_FACTORY(FChaosClothAssetMergeClothCollectionsNode);
 		DATAFLOW_NODE_REGISTER_CREATION_FACTORY(FChaosClothAssetReverseNormalsNode);
diff --git a/Engine/Plugins/Experimental/ChaosClothAssetEditor/Source/ChaosClothAssetDataflowNodes/Private/ChaosClothAsset/DeleteElementNode.cpp b/Engine/Plugins/Experimental/ChaosClothAssetEditor/Source/ChaosClothAssetDataflowNodes/Private/ChaosClothAsset/DeleteElementNode.cpp
new file mode 100644
--- /dev/null
+++ b/Engine/Plugins/Experimental/ChaosClothAssetEditor/Source/ChaosClothAssetDataflowNodes/Private/ChaosClothAsset/DeleteElementNode.cpp
@@ -0,0 +1,275 @@
+// Copyright Epic Games, Inc. All Rights Reserved.
+
+#include "ChaosClothAsset/DeleteElementNode.h"
+#include "ChaosClothAsset/ClothGeometryTools.h"
+#include "ChaosClothAsset/CollectionClothFacade.h"
+#include "Dataflow/DataflowInputOutput.h"
+
+#include UE_INLINE_GENERATED_CPP_BY_NAME(DeleteElementNode)
+
+#define LOCTEXT_NAMESPACE "ChaosClothAssetDeleteElementNode"
+
+FChaosClothAssetDeleteElementNode::FChaosClothAssetDeleteElementNode(const Dataflow::FNodeParameters& InParam, FGuid InGuid)
+	: FDataflowNode(InParam, InGuid)
+{
+	RegisterInputConnection(&Collection);
+	RegisterOutputConnection(&Collection, &Collection);
+}
+
+void FChaosClothAssetDeleteElementNode::Evaluate(Dataflow::FContext& Context, const FDataflowOutput* Out) const
+{
+	if (Out->IsA<FManagedArrayCollection>(&Collection))
+	{
+		using namespace UE::Chaos::ClothAsset;
+
+		// Evaluate in collection
+		FManagedArrayCollection InCollection = GetValue<FManagedArrayCollection>(Context, &Collection);
+		const TSharedRef<FManagedArrayCollection> ClothCollection = MakeShared<FManagedArrayCollection>(MoveTemp(InCollection));
+
+		FCollectionClothFacade Cloth(ClothCollection);
+		if (Cloth.IsValid())
+		{
+			switch (ElementType)
+			{
+			case EChaosClothAssetElementType::None:
+			default:
+				break;
+			case EChaosClothAssetElementType::SimMesh:
+			{
+				FClothGeometryTools::DeleteSimMesh(ClothCollection);
+			} break;
+			case EChaosClothAssetElementType::RenderMesh:
+			{
+				FClothGeometryTools::DeleteRenderMesh(ClothCollection);
+			} break;
+			case EChaosClothAssetElementType::SimPattern:
+			{
+				if (Elements.IsEmpty())
+				{
+					// Remove all sim patterns
+					Cloth.SetNumSimPatterns(0);
+				}
+				else
+				{
+					TArray<int32> PatternsToRemove;
+					PatternsToRemove.Reserve(Elements.Num());
+					for (int32 ElementIndex = 0; ElementIndex < Cloth.GetNumSimPatterns(); ++ElementIndex)
+					{
+						if (Elements.Find(ElementIndex) != INDEX_NONE)
+						{
+							PatternsToRemove.Add(ElementIndex);
+						}
+					}
+					if (PatternsToRemove.Num())
+					{
+						Cloth.RemoveSimPatterns(PatternsToRemove);
+					}
+				}
+			} break;
+			case EChaosClothAssetElementType::RenderPattern:
+			{
+				if (Elements.IsEmpty())
+				{
+					// Remove all render patterns
+					Cloth.SetNumRenderPatterns(0);
+				}
+				else
+				{
+					TArray<int32> PatternsToRemove;
+					PatternsToRemove.Reserve(Elements.Num());
+					for (int32 ElementIndex = 0; ElementIndex < Cloth.GetNumRenderPatterns(); ++ElementIndex)
+					{
+						if (Elements.Find(ElementIndex) != INDEX_NONE)
+						{
+							PatternsToRemove.Add(ElementIndex);
+						}
+					}
+					if (PatternsToRemove.Num())
+					{
+						Cloth.RemoveRenderPatterns(PatternsToRemove);
+					}
+				}
+			} break;
+			case EChaosClothAssetElementType::SimVertex2D:
+			{
+				if (Elements.IsEmpty())
+				{
+					for (int32 ElementIndex = 0; ElementIndex < Cloth.GetNumSimPatterns(); ++ElementIndex)
+					{
+						FCollectionClothSimPatternFacade Pattern = Cloth.GetSimPattern(ElementIndex);
+						Pattern.RemoveAllSimVertices2D();
+					}
+				}
+				else
+				{
+					TArray<int32> SortedToDeleteList;
+					SortedToDeleteList.Reserve(Elements.Num());
+					for (int32 PatternIndex = 0; PatternIndex < Cloth.GetNumSimPatterns(); ++PatternIndex)
+					{
+						FCollectionClothSimPatternFacade Pattern = Cloth.GetSimPattern(PatternIndex);
+						const int32 VertexOffset = Pattern.GetSimVertices2DOffset();
+						SortedToDeleteList.Reset();
+						for (int32 VertexIndex = 0; VertexIndex < Pattern.GetNumSimVertices2D(); ++VertexIndex)
+						{
+							if (Elements.Find(VertexIndex + VertexOffset) != INDEX_NONE)
+							{
+								SortedToDeleteList.Add(VertexIndex);
+							}
+						}
+						if (SortedToDeleteList.Num())
+						{
+							Pattern.RemoveSimVertices2D(SortedToDeleteList);
+						}
+					}
+				}
+			} break;
+			case EChaosClothAssetElementType::SimVertex3D:
+			{
+				if (Elements.IsEmpty())
+				{
+					Cloth.RemoveAllSimVertices3D();
+				}
+				else
+				{
+					TArray<int32> SortedToDeleteList;
+					SortedToDeleteList.Reserve(Elements.Num());
+					for (int32 VertexIndex = 0; VertexIndex < Cloth.GetNumSimVertices3D(); ++VertexIndex)
+					{
+						if (Elements.Find(VertexIndex) != INDEX_NONE)
+						{
+							SortedToDeleteList.Add(VertexIndex);
+						}
+					}
+					Cloth.RemoveSimVertices3D(SortedToDeleteList);
+				}
+			} break;
+			case EChaosClothAssetElementType::RenderVertex:
+			{
+				if (Elements.IsEmpty())
+				{
+					for (int32 ElementIndex = 0; ElementIndex < Cloth.GetNumRenderPatterns(); ++ElementIndex)
+					{
+						FCollectionClothRenderPatternFacade Pattern = Cloth.GetRenderPattern(ElementIndex);
+						Pattern.SetNumRenderVertices(0);
+					}
+				}
+				else
+				{
+					TArray<int32> SortedToDeleteList;
+					SortedToDeleteList.Reserve(Elements.Num());
+					for (int32 PatternIndex = 0; PatternIndex < Cloth.GetNumRenderPatterns(); ++PatternIndex)
+					{
+						FCollectionClothRenderPatternFacade Pattern = Cloth.GetRenderPattern(PatternIndex);
+						const int32 VertexOffset = Pattern.GetRenderVerticesOffset();
+						SortedToDeleteList.Reset();
+						for (int32 VertexIndex = 0; VertexIndex < Pattern.GetNumRenderVertices(); ++VertexIndex)
+						{
+							if (Elements.Find(VertexIndex + VertexOffset) != INDEX_NONE)
+							{
+								SortedToDeleteList.Add(VertexIndex);
+							}
+						}
+						if (SortedToDeleteList.Num())
+						{
+							Pattern.RemoveRenderVertices(SortedToDeleteList);
+						}
+					}
+				}
+			} break;
+			case EChaosClothAssetElementType::SimFace:
+			{
+				if (Elements.IsEmpty())
+				{
+					for (int32 ElementIndex = 0; ElementIndex < Cloth.GetNumSimPatterns(); ++ElementIndex)
+					{
+						FCollectionClothSimPatternFacade Pattern = Cloth.GetSimPattern(ElementIndex);
+						Pattern.SetNumSimFaces(0);
+					}
+				}
+				else
+				{
+					TArray<int32> SortedToDeleteList;
+					SortedToDeleteList.Reserve(Elements.Num());
+					for (int32 PatternIndex = 0; PatternIndex < Cloth.GetNumSimPatterns(); ++PatternIndex)
+					{
+						FCollectionClothSimPatternFacade Pattern = Cloth.GetSimPattern(PatternIndex);
+						const int32 FaceOffset = Pattern.GetSimFacesOffset();
+						SortedToDeleteList.Reset();
+						for (int32 FaceIndex = 0; FaceIndex < Pattern.GetNumSimFaces(); ++FaceIndex)
+						{
+							if (Elements.Find(FaceIndex + FaceOffset) != INDEX_NONE)
+							{
+								SortedToDeleteList.Add(FaceIndex);
+							}
+						}
+						if (SortedToDeleteList.Num())
+						{
+							Pattern.RemoveSimFaces(SortedToDeleteList);
+						}
+					}
+				}
+			} break;
+			case EChaosClothAssetElementType::RenderFace:
+			{
+				if (Elements.IsEmpty())
+				{
+					for (int32 PatternIndex = 0; PatternIndex < Cloth.GetNumRenderPatterns(); ++PatternIndex)
+					{
+						FCollectionClothRenderPatternFacade Pattern = Cloth.GetRenderPattern(PatternIndex);
+						Pattern.SetNumRenderFaces(0);
+					}
+				}
+				else
+				{
+					TArray<int32> SortedToDeleteList;
+					SortedToDeleteList.Reserve(Elements.Num());
+					for (int32 PatternIndex = 0; PatternIndex < Cloth.GetNumRenderPatterns(); ++PatternIndex)
+					{
+						FCollectionClothRenderPatternFacade Pattern = Cloth.GetRenderPattern(PatternIndex);
+						const int32 FaceOffset = Pattern.GetRenderFacesOffset();
+						SortedToDeleteList.Reset();
+						for (int32 FaceIndex = 0; FaceIndex < Pattern.GetNumRenderFaces(); ++FaceIndex)
+						{
+							if (Elements.Find(FaceIndex + FaceOffset) != INDEX_NONE)
+							{
+								SortedToDeleteList.Add(FaceIndex);
+							}
+						}
+						if (SortedToDeleteList.Num())
+						{
+							Pattern.RemoveRenderFaces(SortedToDeleteList);
+						}
+					}
+				}
+			} break;
+			case EChaosClothAssetElementType::Seam:
+			{
+				if (Elements.IsEmpty())
+				{
+					// Remove all seams
+					Cloth.SetNumSeams(0);
+				}
+				else
+				{
+					TArray<int32> SeamsToRemove;
+					SeamsToRemove.Reserve(Elements.Num());
+					for (int32 SeamIndex = 0; SeamIndex < Cloth.GetNumSeams(); ++SeamIndex)
+					{
+						if (Elements.Find(SeamIndex) != INDEX_NONE)
+						{
+							SeamsToRemove.Add(SeamIndex);
+						}
+					}
+					if (SeamsToRemove.Num())
+					{
+						Cloth.RemoveSeams(SeamsToRemove);
+					}
+				}
+			}break;
+			}
+		}
+		SetValue(Context, MoveTemp(*ClothCollection), &Collection);
+	}
+}
+
+#undef LOCTEXT_NAMESPACE
diff --git a/Engine/Plugins/Experimental/ChaosClothAssetEditor/Source/ChaosClothAssetDataflowNodes/Private/ChaosClothAsset/DeleteRenderMeshNode.cpp b/Engine/Plugins/Experimental/ChaosClothAssetEditor/Source/ChaosClothAssetDataflowNodes/Private/ChaosClothAsset/DeleteRenderMeshNode.cpp
deleted file mode 100644
diff --git a/Engine/Plugins/Experimental/ChaosClothAssetEditor/Source/ChaosClothAssetDataflowNodes/Public/ChaosClothAsset/DeleteElementNode.h b/Engine/Plugins/Experimental/ChaosClothAssetEditor/Source/ChaosClothAssetDataflowNodes/Public/ChaosClothAsset/DeleteElementNode.h
new file mode 100644
--- /dev/null
+++ b/Engine/Plugins/Experimental/ChaosClothAssetEditor/Source/ChaosClothAssetDataflowNodes/Public/ChaosClothAsset/DeleteElementNode.h
@@ -0,0 +1,47 @@
+// Copyright Epic Games, Inc. All Rights Reserved.
+
+#pragma once 
+
+#include "Dataflow/DataflowNode.h"
+#include "GeometryCollection/ManagedArrayCollection.h"
+#include "DeleteElementNode.generated.h"
+
+UENUM()
+enum class EChaosClothAssetElementType : uint8
+{
+	None,
+	SimMesh,
+	RenderMesh,
+	SimPattern,
+	RenderPattern,
+	SimVertex2D,
+	SimVertex3D,
+	RenderVertex,
+	SimFace,
+	RenderFace,
+	Seam
+};
+
+USTRUCT(Meta = (DataflowCloth))
+struct FChaosClothAssetDeleteElementNode : public FDataflowNode
+{
+	GENERATED_USTRUCT_BODY()
+	DATAFLOW_NODE_DEFINE_INTERNAL(FChaosClothAssetDeleteElementNode, "DeleteElement", "Cloth", "Cloth Simulation Delete Element")
+
+public:
+
+	UPROPERTY(Meta = (Dataflowinput, DataflowOutput, DataflowPassthrough = "Collection"))
+	FManagedArrayCollection Collection;
+
+	/** Element type to delete.*/
+	UPROPERTY(EditAnywhere, Category = "Delete Element")
+	EChaosClothAssetElementType ElementType = EChaosClothAssetElementType::None;
+
+	/** List of Elements to apply the operation on. All Elements will be used if left empty. */
+	UPROPERTY(EditAnywhere, Category = "Delete Element", Meta = (EditCondition = "ElementType != EChaosClothAssetElementType::SimMesh && ElementType != EChaosClothAssetElementType::RenderMesh"))
+	TArray<int32> Elements;
+
+	FChaosClothAssetDeleteElementNode(const Dataflow::FNodeParameters& InParam, FGuid InGuid = FGuid::NewGuid());
+
+	virtual void Evaluate(Dataflow::FContext& Context, const FDataflowOutput* Out) const override;
+};
diff --git a/Engine/Plugins/Experimental/ChaosClothAssetEditor/Source/ChaosClothAssetDataflowNodes/Public/ChaosClothAsset/DeleteRenderMeshNode.h b/Engine/Plugins/Experimental/ChaosClothAssetEditor/Source/ChaosClothAssetDataflowNodes/Public/ChaosClothAsset/DeleteRenderMeshNode.h
deleted file mode 100644
