Cloth Editor: fix sync between transform gizmo and numeric input in the Preview Scene

diff --git a/Engine/Plugins/Experimental/ChaosClothAssetEditor/Source/ChaosClothAssetEditor/Private/ChaosClothAsset/ClothEditor3DViewportClient.cpp b/Engine/Plugins/Experimental/ChaosClothAssetEditor/Source/ChaosClothAssetEditor/Private/ChaosClothAsset/ClothEditor3DViewportClient.cpp
--- a/Engine/Plugins/Experimental/ChaosClothAssetEditor/Source/ChaosClothAssetEditor/Private/ChaosClothAsset/ClothEditor3DViewportClient.cpp
+++ b/Engine/Plugins/Experimental/ChaosClothAssetEditor/Source/ChaosClothAssetEditor/Private/ChaosClothAsset/ClothEditor3DViewportClient.cpp
@@ -19,6 +19,7 @@
 #include "Animation/SkeletalMeshActor.h"
 #include "AssetViewerSettings.h"
 #include "Editor/EditorPerProjectUserSettings.h"
+#include "Transforms/TransformGizmoDataBinder.h"
 
 namespace UE::Chaos::ClothAsset
 {
@@ -53,6 +54,12 @@ FChaosClothAssetEditor3DViewportClient::FChaosClothAssetEditor3DViewportClient(F
 	Gizmo->bUseContextCoordinateSystem = false;
 	Gizmo->ActiveGizmoMode = EToolContextTransformGizmoMode::Combined;
 
+	UChaosClothPreviewSceneDescription* const SceneDescription = InPreviewScene->GetPreviewSceneDescription();
+	DataBinder = MakeShared<FTransformGizmoDataBinder>();
+	DataBinder->InitializeBoundVectors(&SceneDescription->Translation, &SceneDescription->Rotation, &SceneDescription->Scale);
+
+	InPreviewScene->SetGizmoDataBinder(DataBinder);
+
 	// Set correct flags according to current profile settings
 	SetAdvancedShowFlagsForScene(UAssetViewerSettings::Get()->Profiles[GetMutableDefault<UEditorPerProjectUserSettings>()->AssetViewerProfileIndex].bPostProcessingEnabled);
 }
@@ -81,12 +88,18 @@ FChaosClothAssetEditor3DViewportClient::~FChaosClothAssetEditor3DViewportClient(
 
 void FChaosClothAssetEditor3DViewportClient::DeleteViewportGizmo()
 {
+	if (DataBinder && Gizmo && Gizmo->ActiveTarget)
+	{
+		DataBinder->UnbindFromGizmo(Gizmo, TransformProxy);
+	}
+
 	if (Gizmo && ModeTools && ModeTools->GetInteractiveToolsContext() && ModeTools->GetInteractiveToolsContext()->GizmoManager)
 	{
 		ModeTools->GetInteractiveToolsContext()->GizmoManager->DestroyGizmo(Gizmo);
 	}
 	Gizmo = nullptr;
 	TransformProxy = nullptr;
+	DataBinder = nullptr;
 }
 
 void FChaosClothAssetEditor3DViewportClient::AddReferencedObjects(FReferenceCollector& Collector)
@@ -350,6 +363,12 @@ void FChaosClothAssetEditor3DViewportClient::ComponentSelectionChanged(UObject*
 
 	// Update TransformProxy
 
+	if (Gizmo && Gizmo->ActiveTarget)
+	{
+		DataBinder->UnbindFromGizmo(Gizmo, TransformProxy);
+		Gizmo->ClearActiveTarget();
+	}
+
 	TransformProxy = NewObject<UTransformProxy>();
 	TArray<USceneComponent*> Components;
 	SelectedComponents->GetSelectedObjects(Components);
@@ -365,12 +384,16 @@ void FChaosClothAssetEditor3DViewportClient::ComponentSelectionChanged(UObject*
 		{
 			Gizmo->SetActiveTarget(TransformProxy);
 			Gizmo->SetVisibility(true);
+			DataBinder->BindToInitializedGizmo(Gizmo, TransformProxy);
 		}
 		else
 		{
-			Gizmo->ClearActiveTarget();
 			Gizmo->SetVisibility(false);
 		}
+
+		// TODO: Set UChaosClothPreviewSceneDescription::bValidSelectionForTransform here once we figure out why it's not
+		// properly affecting the EditCondition on the other properties (UE-189504)
+
 	}
 }
 
diff --git a/Engine/Plugins/Experimental/ChaosClothAssetEditor/Source/ChaosClothAssetEditor/Private/ChaosClothAsset/ClothEditor3DViewportClient.h b/Engine/Plugins/Experimental/ChaosClothAssetEditor/Source/ChaosClothAssetEditor/Private/ChaosClothAsset/ClothEditor3DViewportClient.h
--- a/Engine/Plugins/Experimental/ChaosClothAssetEditor/Source/ChaosClothAssetEditor/Private/ChaosClothAsset/ClothEditor3DViewportClient.h
+++ b/Engine/Plugins/Experimental/ChaosClothAssetEditor/Source/ChaosClothAssetEditor/Private/ChaosClothAsset/ClothEditor3DViewportClient.h
@@ -10,6 +10,7 @@ class UChaosClothComponent;
 class UChaosClothAssetEditorMode;
 class UTransformProxy;
 class UCombinedTransformGizmo;
+class FTransformGizmoDataBinder;
 
 namespace UE::Chaos::ClothAsset
 {
@@ -104,6 +105,7 @@ private:
 	// Gizmo support
 	TObjectPtr<UTransformProxy> TransformProxy = nullptr;
 	TObjectPtr<UCombinedTransformGizmo> Gizmo = nullptr;
+	TSharedPtr<FTransformGizmoDataBinder> DataBinder = nullptr;
 
 };
 } // namespace UE::Chaos::ClothAsset
diff --git a/Engine/Plugins/Experimental/ChaosClothAssetEditor/Source/ChaosClothAssetEditor/Private/ChaosClothAsset/ClothEditorPreviewScene.cpp b/Engine/Plugins/Experimental/ChaosClothAssetEditor/Source/ChaosClothAssetEditor/Private/ChaosClothAsset/ClothEditorPreviewScene.cpp
--- a/Engine/Plugins/Experimental/ChaosClothAssetEditor/Source/ChaosClothAssetEditor/Private/ChaosClothAsset/ClothEditorPreviewScene.cpp
+++ b/Engine/Plugins/Experimental/ChaosClothAssetEditor/Source/ChaosClothAssetEditor/Private/ChaosClothAsset/ClothEditorPreviewScene.cpp
@@ -11,6 +11,7 @@
 #include "ChaosClothAsset/ClothAsset.h"
 #include "SkinnedAssetCompiler.h"
 #include "Misc/TransactionObjectEvent.h"
+#include "Transforms/TransformGizmoDataBinder.h"
 
 
 #define LOCTEXT_NAMESPACE "UChaosClothEditorPreviewScene"
@@ -26,7 +27,7 @@ void UChaosClothPreviewSceneDescription::PostEditChangeProperty(FPropertyChanged
 
 	if (PreviewScene)
 	{
-		PreviewScene->SceneDescriptionPropertyChanged(PropertyChangedEvent.GetPropertyName());
+		PreviewScene->SceneDescriptionPropertyChanged(PropertyChangedEvent.GetMemberPropertyName());
 	}
 }
 
@@ -57,7 +58,6 @@ FChaosClothPreviewScene::FChaosClothPreviewScene(FPreviewScene::ConstructionValu
 	SceneActor = GetWorld()->SpawnActor<AActor>(AActor::StaticClass());
 
 	SkeletalMeshComponent = NewObject<USkeletalMeshComponent>(SceneActor);
-	SkeletalMeshComponent->TransformUpdated.AddRaw(this, &FChaosClothPreviewScene::SkeletalMeshTransformChanged);
 	SkeletalMeshComponent->SelectionOverrideDelegate = UPrimitiveComponent::FSelectionOverride::CreateRaw(this, &FChaosClothPreviewScene::IsComponentSelected);
 	SkeletalMeshComponent->SetDisablePostProcessBlueprint(true);
 
@@ -151,11 +151,14 @@ void FChaosClothPreviewScene::SceneDescriptionPropertyChanged(const FName& Prope
 		UpdateClothComponentAttachment();
 	}
 
-	if (PropertyName == GET_MEMBER_NAME_CHECKED(UChaosClothPreviewSceneDescription, SkeletalMeshTransform))
+	if (PropertyName == GET_MEMBER_NAME_CHECKED(UChaosClothPreviewSceneDescription, Translation) || 
+		PropertyName == GET_MEMBER_NAME_CHECKED(UChaosClothPreviewSceneDescription, Rotation) ||
+		PropertyName == GET_MEMBER_NAME_CHECKED(UChaosClothPreviewSceneDescription, Scale))
 	{
-		check(SkeletalMeshComponent);
-
-		SkeletalMeshComponent->SetComponentToWorld(PreviewSceneDescription->SkeletalMeshTransform);
+		if (DataBinder)
+		{
+			DataBinder->UpdateAfterDataEdit();
+		}
 	}
 
 	if (PropertyName == GET_MEMBER_NAME_CHECKED(UChaosClothPreviewSceneDescription, AnimationAsset))
@@ -194,11 +197,6 @@ const TSharedPtr<const FAssetEditorModeManager> FChaosClothPreviewScene::GetClot
 	return ClothPreviewEditorModeManager;
 }
 
-void FChaosClothPreviewScene::SkeletalMeshTransformChanged(USceneComponent* UpdatedComponent, EUpdateTransformFlags UpdateTransformFlags, ETeleportType Teleport)
-{
-	ensure(UpdatedComponent == SkeletalMeshComponent);
-	PreviewSceneDescription->SkeletalMeshTransform = UpdatedComponent->GetComponentToWorld();
-}
 
 bool FChaosClothPreviewScene::IsComponentSelected(const UPrimitiveComponent* InComponent)
 {
@@ -238,6 +236,12 @@ const UAnimSingleNodeInstance* const FChaosClothPreviewScene::GetPreviewAnimInst
 {
 	return PreviewAnimInstance;
 }
+
+void FChaosClothPreviewScene::SetGizmoDataBinder(TSharedPtr<FTransformGizmoDataBinder> InDataBinder)
+{
+	DataBinder = InDataBinder;
+}
+
 } // namespace UE::Chaos::ClothAsset
 
 #undef LOCTEXT_NAMESPACE
diff --git a/Engine/Plugins/Experimental/ChaosClothAssetEditor/Source/ChaosClothAssetEditor/Private/ChaosClothAsset/ClothEditorPreviewScene.h b/Engine/Plugins/Experimental/ChaosClothAssetEditor/Source/ChaosClothAssetEditor/Private/ChaosClothAsset/ClothEditorPreviewScene.h
--- a/Engine/Plugins/Experimental/ChaosClothAssetEditor/Source/ChaosClothAssetEditor/Private/ChaosClothAsset/ClothEditorPreviewScene.h
+++ b/Engine/Plugins/Experimental/ChaosClothAssetEditor/Source/ChaosClothAssetEditor/Private/ChaosClothAsset/ClothEditorPreviewScene.h
@@ -12,6 +12,7 @@ class UChaosClothComponent;
 class FAssetEditorModeManager;
 class UAnimationAsset;
 class UAnimSingleNodeInstance;
+class FTransformGizmoDataBinder;
 
 namespace UE::Chaos::ClothAsset
 {
@@ -38,15 +39,27 @@ public:
 	UPROPERTY(EditAnywhere, Transient, Category="SkeletalMesh")
 	TObjectPtr<USkeletalMesh> SkeletalMeshAsset;
 
-	UPROPERTY(EditAnywhere, Transient, Category = "SkeletalMesh")
-	FTransform SkeletalMeshTransform;
-
 	UPROPERTY(EditAnywhere, Transient, Category = "SkeletalMesh")
 	TObjectPtr<UAnimationAsset> AnimationAsset;
 
+	UPROPERTY(EditAnywhere, Transient, Category = "Transform")
+	FVector3d Translation = FVector3d::ZeroVector;
+
+	UPROPERTY(EditAnywhere, Transient, Category = "Transform")
+	FVector3d Rotation = FVector3d::ZeroVector;
+
+	UPROPERTY(EditAnywhere, Transient, Category = "Transform", Meta = (AllowPreserveRatio))
+	FVector3d Scale = FVector3d::OneVector;
+
+	// TODO: We should be able to hook this boolean property up to the EditCondition meta tag for the properties above and toggle it
+	// on and off when the selection changes in the scene. However the EditCondition does not seem to propagate for some reason, 
+	// even if we manually call PostEditChangeProperty() after toggling it. It will take some more digging to figure out exactly
+	// what's going on. (UE-189504)
+	//UPROPERTY(Transient)
+	//bool bValidSelectionForTransform = true;
+
 private:
 
-	// Listen for changes to the scene description members and notify the PreviewScene
 	virtual void PostEditChangeProperty(struct FPropertyChangedEvent& PropertyChangedEvent) override;
 
 	virtual void PostTransacted(const FTransactionObjectEvent& TransactionEvent) override;
@@ -88,6 +101,8 @@ public:
 	void SetModeManager(TSharedPtr<FAssetEditorModeManager> InClothPreviewEditorModeManager);
 	const TSharedPtr<const FAssetEditorModeManager> GetClothPreviewEditorModeManager() const;
 
+	void SetGizmoDataBinder(TSharedPtr<FTransformGizmoDataBinder> InDataBinder);
+
 private:
 
 	// Create the PreviewAnimationInstance if the AnimationAsset and SkeletalMesh both exist, and set the animation to run on the SkeletalMeshComponent
@@ -96,9 +111,6 @@ private:
 	// Attach the cloth component to the skeletal mesh component, if it exists
 	void UpdateClothComponentAttachment();
 
-
-	void SkeletalMeshTransformChanged(USceneComponent* UpdatedComponent, EUpdateTransformFlags UpdateTransformFlags, ETeleportType Teleport);
-
 	bool IsComponentSelected(const UPrimitiveComponent* InComponent);
 
 	TObjectPtr<UChaosClothPreviewSceneDescription> PreviewSceneDescription;
@@ -113,6 +125,7 @@ private:
 
 	TObjectPtr<USkeletalMeshComponent> SkeletalMeshComponent;
 
+	TSharedPtr<FTransformGizmoDataBinder> DataBinder = nullptr;
 };
 } // namespace UE::Chaos::ClothAsset
 
