Cloth Asset Delete Node: fix for deleting elements from different patterns

diff --git a/Engine/Plugins/Experimental/ChaosClothAssetEditor/Source/ChaosClothAssetDataflowNodes/Private/ChaosClothAsset/DeleteElementNode.cpp b/Engine/Plugins/Experimental/ChaosClothAssetEditor/Source/ChaosClothAssetDataflowNodes/Private/ChaosClothAsset/DeleteElementNode.cpp
--- a/Engine/Plugins/Experimental/ChaosClothAssetEditor/Source/ChaosClothAssetDataflowNodes/Private/ChaosClothAsset/DeleteElementNode.cpp
+++ b/Engine/Plugins/Experimental/ChaosClothAssetEditor/Source/ChaosClothAssetDataflowNodes/Private/ChaosClothAsset/DeleteElementNode.cpp
@@ -104,7 +104,7 @@ void FChaosClothAssetDeleteElementNode::Evaluate(Dataflow::FContext& Context, co
 				{
 					TArray<int32> SortedToDeleteList;
 					SortedToDeleteList.Reserve(Elements.Num());
-					for (int32 PatternIndex = 0; PatternIndex < Cloth.GetNumSimPatterns(); ++PatternIndex)
+					for (int32 PatternIndex = Cloth.GetNumSimPatterns()-1; PatternIndex >= 0; --PatternIndex)
 					{
 						FCollectionClothSimPatternFacade Pattern = Cloth.GetSimPattern(PatternIndex);
 						const int32 VertexOffset = Pattern.GetSimVertices2DOffset();
@@ -157,7 +157,7 @@ void FChaosClothAssetDeleteElementNode::Evaluate(Dataflow::FContext& Context, co
 				{
 					TArray<int32> SortedToDeleteList;
 					SortedToDeleteList.Reserve(Elements.Num());
-					for (int32 PatternIndex = 0; PatternIndex < Cloth.GetNumRenderPatterns(); ++PatternIndex)
+					for (int32 PatternIndex = Cloth.GetNumRenderPatterns()-1; PatternIndex >= 0; --PatternIndex)
 					{
 						FCollectionClothRenderPatternFacade Pattern = Cloth.GetRenderPattern(PatternIndex);
 						const int32 VertexOffset = Pattern.GetRenderVerticesOffset();
@@ -190,7 +190,7 @@ void FChaosClothAssetDeleteElementNode::Evaluate(Dataflow::FContext& Context, co
 				{
 					TArray<int32> SortedToDeleteList;
 					SortedToDeleteList.Reserve(Elements.Num());
-					for (int32 PatternIndex = 0; PatternIndex < Cloth.GetNumSimPatterns(); ++PatternIndex)
+					for (int32 PatternIndex = Cloth.GetNumSimPatterns()-1; PatternIndex >= 0; --PatternIndex)
 					{
 						FCollectionClothSimPatternFacade Pattern = Cloth.GetSimPattern(PatternIndex);
 						const int32 FaceOffset = Pattern.GetSimFacesOffset();
@@ -223,7 +223,7 @@ void FChaosClothAssetDeleteElementNode::Evaluate(Dataflow::FContext& Context, co
 				{
 					TArray<int32> SortedToDeleteList;
 					SortedToDeleteList.Reserve(Elements.Num());
-					for (int32 PatternIndex = 0; PatternIndex < Cloth.GetNumRenderPatterns(); ++PatternIndex)
+					for (int32 PatternIndex = Cloth.GetNumRenderPatterns()-1; PatternIndex >= 0; --PatternIndex)
 					{
 						FCollectionClothRenderPatternFacade Pattern = Cloth.GetRenderPattern(PatternIndex);
 						const int32 FaceOffset = Pattern.GetRenderFacesOffset();
