Chaos Cloth Asset: use our own viewport IsVisible, so we can draw even when the simulations are long.

diff --git a/Engine/Plugins/Experimental/ChaosClothAssetEditor/Source/ChaosClothAssetEditor/Private/ChaosClothAsset/SClothEditor3DViewport.cpp b/Engine/Plugins/Experimental/ChaosClothAssetEditor/Source/ChaosClothAssetEditor/Private/ChaosClothAsset/SClothEditor3DViewport.cpp
--- a/Engine/Plugins/Experimental/ChaosClothAssetEditor/Source/ChaosClothAssetEditor/Private/ChaosClothAsset/SClothEditor3DViewport.cpp
+++ b/Engine/Plugins/Experimental/ChaosClothAssetEditor/Source/ChaosClothAssetEditor/Private/ChaosClothAsset/SClothEditor3DViewport.cpp
@@ -24,6 +24,7 @@ void SChaosClothAssetEditor3DViewport::Construct(const FArguments& InArgs, const
 	}
 	ToolkitCommandList = InArgs._ToolkitCommandList;
 	SAssetEditorViewport::Construct(ParentArgs, InViewportConstructionArgs);
+	Client->VisibilityDelegate.BindSP(this, &SChaosClothAssetEditor3DViewport::IsVisible);
 
 	ViewportOverlay->AddSlot()
 	[
@@ -125,6 +126,12 @@ TSharedPtr<SWidget> SChaosClothAssetEditor3DViewport::MakeViewportToolbar()
 		.CommandList(ToolkitCommandList);
 }
 
+bool SChaosClothAssetEditor3DViewport::IsVisible() const
+{
+	// Intentionally not calling SEditorViewport::IsVisible because it will return false if our simulation is more than 250ms.
+	return ViewportWidget.IsValid();
+}
+
 void SChaosClothAssetEditor3DViewport::OnFocusViewportToSelection()
 {
 	const FBox PreviewBoundingBox = StaticCastSharedPtr<UE::Chaos::ClothAsset::FChaosClothAssetEditor3DViewportClient>(Client)->PreviewBoundingBox();
diff --git a/Engine/Plugins/Experimental/ChaosClothAssetEditor/Source/ChaosClothAssetEditor/Private/ChaosClothAsset/SClothEditor3DViewport.h b/Engine/Plugins/Experimental/ChaosClothAssetEditor/Source/ChaosClothAssetEditor/Private/ChaosClothAsset/SClothEditor3DViewport.h
--- a/Engine/Plugins/Experimental/ChaosClothAssetEditor/Source/ChaosClothAssetEditor/Private/ChaosClothAsset/SClothEditor3DViewport.h
+++ b/Engine/Plugins/Experimental/ChaosClothAssetEditor/Source/ChaosClothAssetEditor/Private/ChaosClothAsset/SClothEditor3DViewport.h
@@ -32,6 +32,7 @@ public:
 	// SAssetEditorViewport
 	virtual void BindCommands() override;
 	virtual TSharedPtr<SWidget> MakeViewportToolbar() override;
+	virtual bool IsVisible() const override;
 
 	virtual void OnFocusViewportToSelection() override;
 
diff --git a/Engine/Plugins/Experimental/ChaosClothAssetEditor/Source/ChaosClothAssetEditor/Private/ChaosClothAsset/SClothEditorRestSpaceViewport.cpp b/Engine/Plugins/Experimental/ChaosClothAssetEditor/Source/ChaosClothAssetEditor/Private/ChaosClothAsset/SClothEditorRestSpaceViewport.cpp
--- a/Engine/Plugins/Experimental/ChaosClothAssetEditor/Source/ChaosClothAssetEditor/Private/ChaosClothAsset/SClothEditorRestSpaceViewport.cpp
+++ b/Engine/Plugins/Experimental/ChaosClothAssetEditor/Source/ChaosClothAssetEditor/Private/ChaosClothAsset/SClothEditorRestSpaceViewport.cpp
@@ -9,6 +9,18 @@
 
 #define LOCTEXT_NAMESPACE "SChaosClothAssetEditorRestSpaceViewport"
 
+void SChaosClothAssetEditorRestSpaceViewport::Construct(const FArguments& InArgs, const FAssetEditorViewportConstructionArgs& InViewportConstructionArgs)
+{
+	SAssetEditorViewport::FArguments ParentArgs;
+	ParentArgs._EditorViewportClient = InArgs._EditorViewportClient;
+	if (InArgs._ViewportSize.IsSet())
+	{
+		ParentArgs._ViewportSize = InArgs._ViewportSize;
+	}
+	SAssetEditorViewport::Construct(ParentArgs, InViewportConstructionArgs);
+	Client->VisibilityDelegate.BindSP(this, &SChaosClothAssetEditorRestSpaceViewport::IsVisible);
+}
+
 UChaosClothAssetEditorMode* SChaosClothAssetEditorRestSpaceViewport::GetEdMode() const
 {
 	if (const FEditorModeTools* const EditorModeTools = Client->GetModeTools())
@@ -159,6 +171,12 @@ void SChaosClothAssetEditorRestSpaceViewport::OnFocusViewportToSelection()
 	}
 }
 
+bool SChaosClothAssetEditorRestSpaceViewport::IsVisible() const
+{
+	// Intentionally not calling SEditorViewport::IsVisible because it will return false if our simulation is more than 250ms.
+	return ViewportWidget.IsValid();
+}
+
 TSharedRef<class SEditorViewport> SChaosClothAssetEditorRestSpaceViewport::GetViewportWidget()
 {
 	return SharedThis(this);
diff --git a/Engine/Plugins/Experimental/ChaosClothAssetEditor/Source/ChaosClothAssetEditor/Private/ChaosClothAsset/SClothEditorRestSpaceViewport.h b/Engine/Plugins/Experimental/ChaosClothAssetEditor/Source/ChaosClothAssetEditor/Private/ChaosClothAsset/SClothEditorRestSpaceViewport.h
--- a/Engine/Plugins/Experimental/ChaosClothAssetEditor/Source/ChaosClothAssetEditor/Private/ChaosClothAsset/SClothEditorRestSpaceViewport.h
+++ b/Engine/Plugins/Experimental/ChaosClothAssetEditor/Source/ChaosClothAssetEditor/Private/ChaosClothAsset/SClothEditorRestSpaceViewport.h
@@ -12,10 +12,18 @@ class CHAOSCLOTHASSETEDITOR_API SChaosClothAssetEditorRestSpaceViewport : public
 {
 public:
 
+	SLATE_BEGIN_ARGS(SChaosClothAssetEditorRestSpaceViewport) {}
+		SLATE_ATTRIBUTE(FVector2D, ViewportSize);
+		SLATE_ARGUMENT(TSharedPtr<FEditorViewportClient>, EditorViewportClient)
+	SLATE_END_ARGS()
+
+	void Construct(const FArguments& InArgs, const FAssetEditorViewportConstructionArgs& InViewportConstructionArgs);
+
 	// SEditorViewport
 	virtual void BindCommands() override;
 	virtual TSharedPtr<SWidget> MakeViewportToolbar() override;
 	virtual void OnFocusViewportToSelection() override;
+	virtual bool IsVisible() const override;
 
 	// ICommonEditorViewportToolbarInfoProvider
 	virtual TSharedRef<class SEditorViewport> GetViewportWidget() override;
