Chaos Cloth Editor - Update UI when a node invalidates.

#rnx

diff --git a/Engine/Plugins/Experimental/ChaosClothAssetEditor/Source/ChaosClothAssetEditor/Private/ChaosClothAsset/ClothEditorToolkit.cpp b/Engine/Plugins/Experimental/ChaosClothAssetEditor/Source/ChaosClothAssetEditor/Private/ChaosClothAsset/ClothEditorToolkit.cpp
--- a/Engine/Plugins/Experimental/ChaosClothAssetEditor/Source/ChaosClothAssetEditor/Private/ChaosClothAsset/ClothEditorToolkit.cpp
+++ b/Engine/Plugins/Experimental/ChaosClothAssetEditor/Source/ChaosClothAssetEditor/Private/ChaosClothAsset/ClothEditorToolkit.cpp
@@ -206,6 +206,12 @@ FChaosClothAssetEditorToolkit::FChaosClothAssetEditorToolkit(UAssetEditor* InOwn
 
 FChaosClothAssetEditorToolkit::~FChaosClothAssetEditorToolkit()
 {
+	if (DataflowNode && OnNodeInvalidatedDelegateHandle.IsValid())
+	{
+		DataflowNode->GetOnNodeInvalidatedDelegate().Remove(OnNodeInvalidatedDelegateHandle);
+	}
+	DataflowNode.Reset();
+
 	if (ClothPreviewViewportClient)
 	{
 		// Delete the gizmo in the viewport before deleting the EditorModeManager. The Gizmo Manager can get tripped up if it gets deleted
@@ -887,7 +893,7 @@ void FChaosClothAssetEditorToolkit::OnNodeTitleCommitted(const FText& InNewText,
 	FDataflowEditorCommands::OnNodeTitleCommitted(InNewText, InCommitType, GraphNode);
 }
 
-void FChaosClothAssetEditorToolkit::OnNodeSelectionChanged(const TSet<UObject*>& NewSelection) const
+void FChaosClothAssetEditorToolkit::OnNodeSelectionChanged(const TSet<UObject*>& NewSelection)
 {
 	auto GetClothCollectionIfPossible = [](const TSharedPtr<const FDataflowNode> DataflowNode, const TSharedPtr<Dataflow::FEngineContext> Context) -> TSharedPtr<FManagedArrayCollection>
 	{
@@ -927,15 +933,32 @@ void FChaosClothAssetEditorToolkit::OnNodeSelectionChanged(const TSet<UObject*>&
 	{
 		Dataflow->RenderTargets.Reset();
 
-		for (const UObject* const Selected : NewSelection)
+		for (UObject* const Selected : NewSelection)
 		{
-			if (const UDataflowEdNode* const Node = Cast<UDataflowEdNode>(Selected))
+			if (UDataflowEdNode* const Node = Cast<UDataflowEdNode>(Selected))
 			{
 				Dataflow->RenderTargets.Add(Node);
 
-				if (const TSharedPtr<const FDataflowNode> DataflowNode = Node->GetDataflowNode())
+				if (DataflowNode && OnNodeInvalidatedDelegateHandle.IsValid())
+				{
+					DataflowNode->GetOnNodeInvalidatedDelegate().Remove(OnNodeInvalidatedDelegateHandle);
+				}
+				DataflowNode = Node->GetDataflowNode();
+
+				if (DataflowNode)
 				{
 					Collection = GetClothCollectionIfPossible(DataflowNode, this->DataflowContext);
+					DataflowNode->GetOnNodeInvalidatedDelegate();
+
+					// Set a callback to re-evaluate the node if it is invalidated
+					OnNodeInvalidatedDelegateHandle = DataflowNode->GetOnNodeInvalidatedDelegate().AddLambda(
+						[this, &GetClothCollectionIfPossible](FDataflowNode* InDataflowNode)
+						{
+							if (DataflowNode.Get() == InDataflowNode)
+							{
+								GetClothCollectionIfPossible(DataflowNode, DataflowContext);
+							}
+						});
 				}
 			}
 		}
diff --git a/Engine/Plugins/Experimental/ChaosClothAssetEditor/Source/ChaosClothAssetEditor/Private/ChaosClothAsset/ClothEditorToolkit.h b/Engine/Plugins/Experimental/ChaosClothAssetEditor/Source/ChaosClothAssetEditor/Private/ChaosClothAsset/ClothEditorToolkit.h
--- a/Engine/Plugins/Experimental/ChaosClothAssetEditor/Source/ChaosClothAssetEditor/Private/ChaosClothAsset/ClothEditorToolkit.h
+++ b/Engine/Plugins/Experimental/ChaosClothAssetEditor/Source/ChaosClothAssetEditor/Private/ChaosClothAsset/ClothEditorToolkit.h
@@ -122,7 +122,7 @@ private:
 	void OnPropertyValueChanged(const FPropertyChangedEvent& PropertyChangedEvent);
 	bool OnNodeVerifyTitleCommit(const FText& NewText, UEdGraphNode* GraphNode, FText& OutErrorMessage) const;
 	void OnNodeTitleCommitted(const FText& InNewText, ETextCommit::Type InCommitType, UEdGraphNode* GraphNode) const;
-	void OnNodeSelectionChanged(const TSet<UObject*>& NewSelection) const;
+	void OnNodeSelectionChanged(const TSet<UObject*>& NewSelection);
 	void OnNodeDeleted(const TSet<UObject*>& DeletedNodes) const;
 
 	/** Scene in which the 3D sim space preview meshes live. Ownership shared with AdvancedPreviewSettingsWidget*/
@@ -147,6 +147,8 @@ private:
 	FString DataflowTerminalPath = "";
 	TSharedPtr<Dataflow::FEngineContext> DataflowContext;
 	Dataflow::FTimestamp LastDataflowNodeTimestamp = Dataflow::FTimestamp::Invalid;
+	FDelegateHandle OnNodeInvalidatedDelegateHandle;
+	TSharedPtr<FDataflowNode> DataflowNode;
 
 	static const FName GraphCanvasTabId;
 	TSharedPtr<SDockTab> GraphEditorTab;
