Chaos Cloth:
- Added weight map inputs to all simulation config nodes.
- Added a new weight map override property to display connected weight map names.
- Changed the call to AddProperties so that it is called with Context and even with an empty collection as to update the override.

#rnx

diff --git a/Engine/Plugins/Experimental/ChaosClothAssetEditor/Source/ChaosClothAssetDataflowNodes/Private/ChaosClothAsset/SimulationAerodynamicsConfigNode.cpp b/Engine/Plugins/Experimental/ChaosClothAssetEditor/Source/ChaosClothAssetDataflowNodes/Private/ChaosClothAsset/SimulationAerodynamicsConfigNode.cpp
--- a/Engine/Plugins/Experimental/ChaosClothAssetEditor/Source/ChaosClothAssetDataflowNodes/Private/ChaosClothAsset/SimulationAerodynamicsConfigNode.cpp
+++ b/Engine/Plugins/Experimental/ChaosClothAssetEditor/Source/ChaosClothAssetDataflowNodes/Private/ChaosClothAsset/SimulationAerodynamicsConfigNode.cpp
@@ -10,9 +10,11 @@ FChaosClothAssetSimulationAerodynamicsConfigNode::FChaosClothAssetSimulationAero
 	: FChaosClothAssetSimulationBaseConfigNode(InParam, InGuid)
 {
 	RegisterCollectionConnections();
+	RegisterInputConnection(&Drag.WeightMap);
+	RegisterInputConnection(&Lift.WeightMap);
 }
 
-void FChaosClothAssetSimulationAerodynamicsConfigNode::AddProperties(::Chaos::Softs::FCollectionPropertyMutableFacade& Properties) const
+void FChaosClothAssetSimulationAerodynamicsConfigNode::AddProperties(Dataflow::FContext& Context, ::Chaos::Softs::FCollectionPropertyMutableFacade& Properties) const
 {
 	UE_CHAOS_CLOTHASSET_SIMULATIONCONFIG_SETPROPERTY(FluidDensity);
 	UE_CHAOS_CLOTHASSET_SIMULATIONCONFIG_SETPROPERTYWEIGHTED(Drag);
diff --git a/Engine/Plugins/Experimental/ChaosClothAssetEditor/Source/ChaosClothAssetDataflowNodes/Private/ChaosClothAsset/SimulationAnimDriveConfigNode.cpp b/Engine/Plugins/Experimental/ChaosClothAssetEditor/Source/ChaosClothAssetDataflowNodes/Private/ChaosClothAsset/SimulationAnimDriveConfigNode.cpp
--- a/Engine/Plugins/Experimental/ChaosClothAssetEditor/Source/ChaosClothAssetDataflowNodes/Private/ChaosClothAsset/SimulationAnimDriveConfigNode.cpp
+++ b/Engine/Plugins/Experimental/ChaosClothAssetEditor/Source/ChaosClothAssetDataflowNodes/Private/ChaosClothAsset/SimulationAnimDriveConfigNode.cpp
@@ -10,9 +10,11 @@ FChaosClothAssetSimulationAnimDriveConfigNode::FChaosClothAssetSimulationAnimDri
 	: FChaosClothAssetSimulationBaseConfigNode(InParam, InGuid)
 {
 	RegisterCollectionConnections();
+	RegisterInputConnection(&AnimDriveStiffness.WeightMap);
+	RegisterInputConnection(&AnimDriveDamping.WeightMap);
 }
 
-void FChaosClothAssetSimulationAnimDriveConfigNode::AddProperties(::Chaos::Softs::FCollectionPropertyMutableFacade& Properties) const
+void FChaosClothAssetSimulationAnimDriveConfigNode::AddProperties(Dataflow::FContext& Context, ::Chaos::Softs::FCollectionPropertyMutableFacade& Properties) const
 {
 	UE_CHAOS_CLOTHASSET_SIMULATIONCONFIG_SETPROPERTYWEIGHTED(AnimDriveStiffness);
 	UE_CHAOS_CLOTHASSET_SIMULATIONCONFIG_SETPROPERTYWEIGHTED(AnimDriveDamping);
diff --git a/Engine/Plugins/Experimental/ChaosClothAssetEditor/Source/ChaosClothAssetDataflowNodes/Private/ChaosClothAsset/SimulationBackstopConfigNode.cpp b/Engine/Plugins/Experimental/ChaosClothAssetEditor/Source/ChaosClothAssetDataflowNodes/Private/ChaosClothAsset/SimulationBackstopConfigNode.cpp
--- a/Engine/Plugins/Experimental/ChaosClothAssetEditor/Source/ChaosClothAssetDataflowNodes/Private/ChaosClothAsset/SimulationBackstopConfigNode.cpp
+++ b/Engine/Plugins/Experimental/ChaosClothAssetEditor/Source/ChaosClothAssetDataflowNodes/Private/ChaosClothAsset/SimulationBackstopConfigNode.cpp
@@ -10,9 +10,11 @@ FChaosClothAssetSimulationBackstopConfigNode::FChaosClothAssetSimulationBackstop
 	: FChaosClothAssetSimulationBaseConfigNode(InParam, InGuid)
 {
 	RegisterCollectionConnections();
+	RegisterInputConnection(&BackstopDistance.WeightMap);
+	RegisterInputConnection(&BackstopRadius.WeightMap);
 }
 
-void FChaosClothAssetSimulationBackstopConfigNode::AddProperties(::Chaos::Softs::FCollectionPropertyMutableFacade& Properties) const
+void FChaosClothAssetSimulationBackstopConfigNode::AddProperties(Dataflow::FContext& Context, ::Chaos::Softs::FCollectionPropertyMutableFacade& Properties) const
 {
 	UE_CHAOS_CLOTHASSET_SIMULATIONCONFIG_SETPROPERTYWEIGHTED(BackstopDistance);
 	UE_CHAOS_CLOTHASSET_SIMULATIONCONFIG_SETPROPERTYWEIGHTED(BackstopRadius);
diff --git a/Engine/Plugins/Experimental/ChaosClothAssetEditor/Source/ChaosClothAssetDataflowNodes/Private/ChaosClothAsset/SimulationBaseConfigNode.cpp b/Engine/Plugins/Experimental/ChaosClothAssetEditor/Source/ChaosClothAssetDataflowNodes/Private/ChaosClothAsset/SimulationBaseConfigNode.cpp
--- a/Engine/Plugins/Experimental/ChaosClothAssetEditor/Source/ChaosClothAssetDataflowNodes/Private/ChaosClothAsset/SimulationBaseConfigNode.cpp
+++ b/Engine/Plugins/Experimental/ChaosClothAssetEditor/Source/ChaosClothAssetDataflowNodes/Private/ChaosClothAsset/SimulationBaseConfigNode.cpp
@@ -64,13 +64,13 @@ void FChaosClothAssetSimulationBaseConfigNode::Evaluate(Dataflow::FContext& Cont
 		FManagedArrayCollection InCollection = GetValue<FManagedArrayCollection>(Context, &Collection);
 		const TSharedRef<FManagedArrayCollection> ClothCollection = MakeShared<FManagedArrayCollection>(MoveTemp(InCollection));
 
-		if (FCollectionClothFacade(ClothCollection).IsValid())  // Can only act on the collection if it is a valid cloth collection
-		{
-			FCollectionPropertyMutableFacade Properties(ClothCollection);
-			Properties.DefineSchema();
+		FCollectionPropertyMutableFacade Properties(ClothCollection);
+		Properties.DefineSchema();
 
-			AddProperties(Properties);
+		AddProperties(Context, Properties);
 
+		if (FCollectionClothFacade(ClothCollection).IsValid())  // Can only act on the collection if it is a valid cloth collection
+		{
 			EvaluateClothCollection(Context, ClothCollection);
 		}
 
diff --git a/Engine/Plugins/Experimental/ChaosClothAssetEditor/Source/ChaosClothAssetDataflowNodes/Private/ChaosClothAsset/SimulationBaseConfigNodePrivate.h b/Engine/Plugins/Experimental/ChaosClothAssetEditor/Source/ChaosClothAssetDataflowNodes/Private/ChaosClothAsset/SimulationBaseConfigNodePrivate.h
--- a/Engine/Plugins/Experimental/ChaosClothAssetEditor/Source/ChaosClothAssetDataflowNodes/Private/ChaosClothAsset/SimulationBaseConfigNodePrivate.h
+++ b/Engine/Plugins/Experimental/ChaosClothAssetEditor/Source/ChaosClothAssetDataflowNodes/Private/ChaosClothAsset/SimulationBaseConfigNodePrivate.h
@@ -1,6 +1,8 @@
 // Copyright Epic Games, Inc. All Rights Reserved.
 
-#pragma once 
+#pragma once
+
+#include "Dataflow/DataflowInputOutput.h"
 
 /** Macro for adding and setting value properties from inside AddProperties() */
 #define UE_CHAOS_CLOTHASSET_SIMULATIONCONFIG_SETPROPERTY(Property) \
@@ -21,7 +23,10 @@
 #define UE_CHAOS_CLOTHASSET_SIMULATIONCONFIG_SETPROPERTYWEIGHTED(Property) \
 	const int32 Property##KeyIndex = AddPropertyHelper(Properties, FName(TEXT(#Property)), Property.bIsAnimatable); \
 	Properties.SetWeightedValue(Property##KeyIndex, Property.Low, Property.High); \
-	Properties.SetStringValue(Property##KeyIndex, Property.WeightMap);
+	Properties.SetStringValue(Property##KeyIndex, GetValue<FString>(Context, &Property.WeightMap)); \
+PRAGMA_DISABLE_DEPRECATION_WARNINGS \
+	Property.WeightMap_Override = GetValue<FString>(Context, &Property.WeightMap, FString()); \
+PRAGMA_ENABLE_DEPRECATION_WARNINGS
 
 /** Macro for adding and setting weighted value properties from inside AddProperties() and check for a similar property being overriden. */
 #define UE_CHAOS_CLOTHASSET_SIMULATIONCONFIG_SETPROPERTYCHECKED1(Property, SimilarProperty1) \
@@ -46,25 +51,37 @@
 	const int32 Property##KeyIndex = AddPropertyHelper(Properties, FName(TEXT(#Property)), Property.bIsAnimatable, \
 		{ FName(TEXT(#SimilarProperty1)) }); \
 	Properties.SetWeightedValue(Property##KeyIndex, Property.Low, Property.High); \
-	Properties.SetStringValue(Property##KeyIndex, Property.WeightMap);
+	Properties.SetStringValue(Property##KeyIndex, GetValue<FString>(Context, &Property.WeightMap)); \
+PRAGMA_DISABLE_DEPRECATION_WARNINGS \
+	Property.WeightMap_Override = GetValue<FString>(Context, &Property.WeightMap, FString()); \
+PRAGMA_ENABLE_DEPRECATION_WARNINGS
 
 /** Macro for adding and setting weighted value properties from inside AddProperties() and check for two similar properties being overriden. */
 #define UE_CHAOS_CLOTHASSET_SIMULATIONCONFIG_SETPROPERTYWEIGHTEDCHECKED2(Property, SimilarProperty1, SimilarProperty2) \
 	const int32 Property##KeyIndex = AddPropertyHelper(Properties, FName(TEXT(#Property)), Property.bIsAnimatable, \
 		{ FName(TEXT(#SimilarProperty1)), FName(TEXT(#SimilarProperty2)) }); \
 	Properties.SetWeightedValue(Property##KeyIndex, Property.Low, Property.High); \
-	Properties.SetStringValue(Property##KeyIndex, Property.WeightMap);
+	Properties.SetStringValue(Property##KeyIndex, GetValue<FString>(Context, &Property.WeightMap)); \
+PRAGMA_DISABLE_DEPRECATION_WARNINGS \
+	Property.WeightMap_Override = GetValue<FString>(Context, &Property.WeightMap, FString()); \
+PRAGMA_ENABLE_DEPRECATION_WARNINGS
 
 /** Macro for adding and setting weighted value properties from inside AddProperties() and check for three similar properties being overriden. */
 #define UE_CHAOS_CLOTHASSET_SIMULATIONCONFIG_SETPROPERTYWEIGHTEDCHECKED3(Property, SimilarProperty1, SimilarProperty2, SimilarProperty3) \
 	const int32 Property##KeyIndex = AddPropertyHelper(Properties, FName(TEXT(#Property)), Property.bIsAnimatable, \
 		{ FName(TEXT(#SimilarProperty1)), FName(TEXT(#SimilarProperty2)), FName(TEXT(#SimilarProperty3)) }); \
 	Properties.SetWeightedValue(Property##KeyIndex, Property.Low, Property.High); \
-	Properties.SetStringValue(Property##KeyIndex, Property.WeightMap);
+	Properties.SetStringValue(Property##KeyIndex, GetValue<FString>(Context, &Property.WeightMap)); \
+PRAGMA_DISABLE_DEPRECATION_WARNINGS \
+	Property.WeightMap_Override = GetValue<FString>(Context, &Property.WeightMap, FString()); \
+PRAGMA_ENABLE_DEPRECATION_WARNINGS
 
 /** Macro for adding and setting weighted value properties from inside AddProperties() and check for four similar properties being overriden. */
 #define UE_CHAOS_CLOTHASSET_SIMULATIONCONFIG_SETPROPERTYWEIGHTEDCHECKED4(Property, SimilarProperty1, SimilarProperty2, SimilarProperty3, SimilarProperty4) \
 	const int32 Property##KeyIndex = AddPropertyHelper(Properties, FName(TEXT(#Property)), Property.bIsAnimatable, \
 		{ FName(TEXT(#SimilarProperty1)), FName(TEXT(#SimilarProperty2)), FName(TEXT(#SimilarProperty3)), FName(TEXT(#SimilarProperty4)) }); \
 	Properties.SetWeightedValue(Property##KeyIndex, Property.Low, Property.High); \
-	Properties.SetStringValue(Property##KeyIndex, Property.WeightMap);
+	Properties.SetStringValue(Property##KeyIndex, GetValue<FString>(Context, &Property.WeightMap)); \
+PRAGMA_DISABLE_DEPRECATION_WARNINGS \
+	Property.WeightMap_Override = GetValue<FString>(Context, &Property.WeightMap, FString()); \
+PRAGMA_ENABLE_DEPRECATION_WARNINGS
diff --git a/Engine/Plugins/Experimental/ChaosClothAssetEditor/Source/ChaosClothAssetDataflowNodes/Private/ChaosClothAsset/SimulationCollisionConfigNode.cpp b/Engine/Plugins/Experimental/ChaosClothAssetEditor/Source/ChaosClothAssetDataflowNodes/Private/ChaosClothAsset/SimulationCollisionConfigNode.cpp
--- a/Engine/Plugins/Experimental/ChaosClothAssetEditor/Source/ChaosClothAssetDataflowNodes/Private/ChaosClothAsset/SimulationCollisionConfigNode.cpp
+++ b/Engine/Plugins/Experimental/ChaosClothAssetEditor/Source/ChaosClothAssetDataflowNodes/Private/ChaosClothAsset/SimulationCollisionConfigNode.cpp
@@ -12,7 +12,7 @@ FChaosClothAssetSimulationCollisionConfigNode::FChaosClothAssetSimulationCollisi
 	RegisterCollectionConnections();
 }
 
-void FChaosClothAssetSimulationCollisionConfigNode::AddProperties(::Chaos::Softs::FCollectionPropertyMutableFacade& Properties) const
+void FChaosClothAssetSimulationCollisionConfigNode::AddProperties(Dataflow::FContext& Context, ::Chaos::Softs::FCollectionPropertyMutableFacade& Properties) const
 {
 	UE_CHAOS_CLOTHASSET_SIMULATIONCONFIG_SETPROPERTY(CollisionThickness);
 	UE_CHAOS_CLOTHASSET_SIMULATIONCONFIG_SETPROPERTY(FrictionCoefficient);
diff --git a/Engine/Plugins/Experimental/ChaosClothAssetEditor/Source/ChaosClothAssetDataflowNodes/Private/ChaosClothAsset/SimulationDampingConfigNode.cpp b/Engine/Plugins/Experimental/ChaosClothAssetEditor/Source/ChaosClothAssetDataflowNodes/Private/ChaosClothAsset/SimulationDampingConfigNode.cpp
--- a/Engine/Plugins/Experimental/ChaosClothAssetEditor/Source/ChaosClothAssetDataflowNodes/Private/ChaosClothAsset/SimulationDampingConfigNode.cpp
+++ b/Engine/Plugins/Experimental/ChaosClothAssetEditor/Source/ChaosClothAssetDataflowNodes/Private/ChaosClothAsset/SimulationDampingConfigNode.cpp
@@ -12,7 +12,7 @@ FChaosClothAssetSimulationDampingConfigNode::FChaosClothAssetSimulationDampingCo
 	RegisterCollectionConnections();
 }
 
-void FChaosClothAssetSimulationDampingConfigNode::AddProperties(::Chaos::Softs::FCollectionPropertyMutableFacade& Properties) const
+void FChaosClothAssetSimulationDampingConfigNode::AddProperties(Dataflow::FContext& Context, ::Chaos::Softs::FCollectionPropertyMutableFacade& Properties) const
 {
 	UE_CHAOS_CLOTHASSET_SIMULATIONCONFIG_SETPROPERTY(DampingCoefficient);
 	UE_CHAOS_CLOTHASSET_SIMULATIONCONFIG_SETPROPERTY(LocalDampingCoefficient);
diff --git a/Engine/Plugins/Experimental/ChaosClothAssetEditor/Source/ChaosClothAssetDataflowNodes/Private/ChaosClothAsset/SimulationGravityConfigNode.cpp b/Engine/Plugins/Experimental/ChaosClothAssetEditor/Source/ChaosClothAssetDataflowNodes/Private/ChaosClothAsset/SimulationGravityConfigNode.cpp
--- a/Engine/Plugins/Experimental/ChaosClothAssetEditor/Source/ChaosClothAssetDataflowNodes/Private/ChaosClothAsset/SimulationGravityConfigNode.cpp
+++ b/Engine/Plugins/Experimental/ChaosClothAssetEditor/Source/ChaosClothAssetDataflowNodes/Private/ChaosClothAsset/SimulationGravityConfigNode.cpp
@@ -12,7 +12,7 @@ FChaosClothAssetSimulationGravityConfigNode::FChaosClothAssetSimulationGravityCo
 	RegisterCollectionConnections();
 }
 
-void FChaosClothAssetSimulationGravityConfigNode::AddProperties(::Chaos::Softs::FCollectionPropertyMutableFacade& Properties) const
+void FChaosClothAssetSimulationGravityConfigNode::AddProperties(Dataflow::FContext& Context, ::Chaos::Softs::FCollectionPropertyMutableFacade& Properties) const
 {
 	UE_CHAOS_CLOTHASSET_SIMULATIONCONFIG_SETPROPERTYBOOL(UseGravityOverride);
 	UE_CHAOS_CLOTHASSET_SIMULATIONCONFIG_SETPROPERTY(GravityScale);
diff --git a/Engine/Plugins/Experimental/ChaosClothAssetEditor/Source/ChaosClothAssetDataflowNodes/Private/ChaosClothAsset/SimulationLongRangeAttachmentConfigNode.cpp b/Engine/Plugins/Experimental/ChaosClothAssetEditor/Source/ChaosClothAssetDataflowNodes/Private/ChaosClothAsset/SimulationLongRangeAttachmentConfigNode.cpp
--- a/Engine/Plugins/Experimental/ChaosClothAssetEditor/Source/ChaosClothAssetDataflowNodes/Private/ChaosClothAsset/SimulationLongRangeAttachmentConfigNode.cpp
+++ b/Engine/Plugins/Experimental/ChaosClothAssetEditor/Source/ChaosClothAssetDataflowNodes/Private/ChaosClothAsset/SimulationLongRangeAttachmentConfigNode.cpp
@@ -13,9 +13,11 @@ FChaosClothAssetSimulationLongRangeAttachmentConfigNode::FChaosClothAssetSimulat
 {
 	RegisterCollectionConnections();
 	RegisterInputConnection(&FixedEndWeightMap);
+	RegisterInputConnection(&TetherStiffness.WeightMap);
+	RegisterInputConnection(&TetherScale.WeightMap);
 }
 
-void FChaosClothAssetSimulationLongRangeAttachmentConfigNode::AddProperties(::Chaos::Softs::FCollectionPropertyMutableFacade& Properties) const
+void FChaosClothAssetSimulationLongRangeAttachmentConfigNode::AddProperties(Dataflow::FContext& Context, ::Chaos::Softs::FCollectionPropertyMutableFacade& Properties) const
 {
 	UE_CHAOS_CLOTHASSET_SIMULATIONCONFIG_SETPROPERTYWEIGHTED(TetherStiffness);
 	UE_CHAOS_CLOTHASSET_SIMULATIONCONFIG_SETPROPERTYWEIGHTED(TetherScale);
diff --git a/Engine/Plugins/Experimental/ChaosClothAssetEditor/Source/ChaosClothAssetDataflowNodes/Private/ChaosClothAsset/SimulationMassConfigNode.cpp b/Engine/Plugins/Experimental/ChaosClothAssetEditor/Source/ChaosClothAssetDataflowNodes/Private/ChaosClothAsset/SimulationMassConfigNode.cpp
--- a/Engine/Plugins/Experimental/ChaosClothAssetEditor/Source/ChaosClothAssetDataflowNodes/Private/ChaosClothAsset/SimulationMassConfigNode.cpp
+++ b/Engine/Plugins/Experimental/ChaosClothAssetEditor/Source/ChaosClothAssetDataflowNodes/Private/ChaosClothAsset/SimulationMassConfigNode.cpp
@@ -12,7 +12,7 @@ FChaosClothAssetSimulationMassConfigNode::FChaosClothAssetSimulationMassConfigNo
 	RegisterCollectionConnections();
 }
 
-void FChaosClothAssetSimulationMassConfigNode::AddProperties(::Chaos::Softs::FCollectionPropertyMutableFacade& Properties) const
+void FChaosClothAssetSimulationMassConfigNode::AddProperties(Dataflow::FContext& Context, ::Chaos::Softs::FCollectionPropertyMutableFacade& Properties) const
 {
 	float MassValue;
 	switch (MassMode)
diff --git a/Engine/Plugins/Experimental/ChaosClothAssetEditor/Source/ChaosClothAssetDataflowNodes/Private/ChaosClothAsset/SimulationMaxDistanceConfigNode.cpp b/Engine/Plugins/Experimental/ChaosClothAssetEditor/Source/ChaosClothAssetDataflowNodes/Private/ChaosClothAsset/SimulationMaxDistanceConfigNode.cpp
--- a/Engine/Plugins/Experimental/ChaosClothAssetEditor/Source/ChaosClothAssetDataflowNodes/Private/ChaosClothAsset/SimulationMaxDistanceConfigNode.cpp
+++ b/Engine/Plugins/Experimental/ChaosClothAssetEditor/Source/ChaosClothAssetDataflowNodes/Private/ChaosClothAsset/SimulationMaxDistanceConfigNode.cpp
@@ -10,9 +10,10 @@ FChaosClothAssetSimulationMaxDistanceConfigNode::FChaosClothAssetSimulationMaxDi
 	: FChaosClothAssetSimulationBaseConfigNode(InParam, InGuid)
 {
 	RegisterCollectionConnections();
+	RegisterInputConnection(&MaxDistance.WeightMap);
 }
 
-void FChaosClothAssetSimulationMaxDistanceConfigNode::AddProperties(::Chaos::Softs::FCollectionPropertyMutableFacade& Properties) const
+void FChaosClothAssetSimulationMaxDistanceConfigNode::AddProperties(Dataflow::FContext& Context, ::Chaos::Softs::FCollectionPropertyMutableFacade& Properties) const
 {
 	UE_CHAOS_CLOTHASSET_SIMULATIONCONFIG_SETPROPERTYWEIGHTED(MaxDistance);
 }
diff --git a/Engine/Plugins/Experimental/ChaosClothAssetEditor/Source/ChaosClothAssetDataflowNodes/Private/ChaosClothAsset/SimulationPBDAreaSpringConfigNode.cpp b/Engine/Plugins/Experimental/ChaosClothAssetEditor/Source/ChaosClothAssetDataflowNodes/Private/ChaosClothAsset/SimulationPBDAreaSpringConfigNode.cpp
--- a/Engine/Plugins/Experimental/ChaosClothAssetEditor/Source/ChaosClothAssetDataflowNodes/Private/ChaosClothAsset/SimulationPBDAreaSpringConfigNode.cpp
+++ b/Engine/Plugins/Experimental/ChaosClothAssetEditor/Source/ChaosClothAssetDataflowNodes/Private/ChaosClothAsset/SimulationPBDAreaSpringConfigNode.cpp
@@ -10,9 +10,10 @@ FChaosClothAssetSimulationPBDAreaSpringConfigNode::FChaosClothAssetSimulationPBD
 	: FChaosClothAssetSimulationBaseConfigNode(InParam, InGuid)
 {
 	RegisterCollectionConnections();
+	RegisterInputConnection(&AreaSpringStiffness.WeightMap);
 }
 
-void FChaosClothAssetSimulationPBDAreaSpringConfigNode::AddProperties(::Chaos::Softs::FCollectionPropertyMutableFacade& Properties) const
+void FChaosClothAssetSimulationPBDAreaSpringConfigNode::AddProperties(Dataflow::FContext& Context, ::Chaos::Softs::FCollectionPropertyMutableFacade& Properties) const
 {
 	UE_CHAOS_CLOTHASSET_SIMULATIONCONFIG_SETPROPERTYWEIGHTEDCHECKED1(
 		AreaSpringStiffness,
diff --git a/Engine/Plugins/Experimental/ChaosClothAssetEditor/Source/ChaosClothAssetDataflowNodes/Private/ChaosClothAsset/SimulationPBDBendingElementConfigNode.cpp b/Engine/Plugins/Experimental/ChaosClothAssetEditor/Source/ChaosClothAssetDataflowNodes/Private/ChaosClothAsset/SimulationPBDBendingElementConfigNode.cpp
--- a/Engine/Plugins/Experimental/ChaosClothAssetEditor/Source/ChaosClothAssetDataflowNodes/Private/ChaosClothAsset/SimulationPBDBendingElementConfigNode.cpp
+++ b/Engine/Plugins/Experimental/ChaosClothAssetEditor/Source/ChaosClothAssetDataflowNodes/Private/ChaosClothAsset/SimulationPBDBendingElementConfigNode.cpp
@@ -10,9 +10,13 @@ FChaosClothAssetSimulationPBDBendingElementConfigNode::FChaosClothAssetSimulatio
 	: FChaosClothAssetSimulationBaseConfigNode(InParam, InGuid)
 {
 	RegisterCollectionConnections();
+	RegisterInputConnection(&FlatnessRatio.WeightMap);
+	RegisterInputConnection(&RestAngle.WeightMap);
+	RegisterInputConnection(&BendingElementStiffness.WeightMap);
+	RegisterInputConnection(&BucklingStiffness.WeightMap);
 }
 
-void FChaosClothAssetSimulationPBDBendingElementConfigNode::AddProperties(::Chaos::Softs::FCollectionPropertyMutableFacade& Properties) const
+void FChaosClothAssetSimulationPBDBendingElementConfigNode::AddProperties(Dataflow::FContext& Context, ::Chaos::Softs::FCollectionPropertyMutableFacade& Properties) const
 {
 	UE_CHAOS_CLOTHASSET_SIMULATIONCONFIG_SETPROPERTYENUMCHECKED2(
 		RestAngleType,
diff --git a/Engine/Plugins/Experimental/ChaosClothAssetEditor/Source/ChaosClothAssetDataflowNodes/Private/ChaosClothAsset/SimulationPBDBendingSpringConfigNode.cpp b/Engine/Plugins/Experimental/ChaosClothAssetEditor/Source/ChaosClothAssetDataflowNodes/Private/ChaosClothAsset/SimulationPBDBendingSpringConfigNode.cpp
--- a/Engine/Plugins/Experimental/ChaosClothAssetEditor/Source/ChaosClothAssetDataflowNodes/Private/ChaosClothAsset/SimulationPBDBendingSpringConfigNode.cpp
+++ b/Engine/Plugins/Experimental/ChaosClothAssetEditor/Source/ChaosClothAssetDataflowNodes/Private/ChaosClothAsset/SimulationPBDBendingSpringConfigNode.cpp
@@ -10,9 +10,10 @@ FChaosClothAssetSimulationPBDBendingSpringConfigNode::FChaosClothAssetSimulation
 	: FChaosClothAssetSimulationBaseConfigNode(InParam, InGuid)
 {
 	RegisterCollectionConnections();
+	RegisterInputConnection(&BendingSpringStiffness.WeightMap);
 }
 
-void FChaosClothAssetSimulationPBDBendingSpringConfigNode::AddProperties(::Chaos::Softs::FCollectionPropertyMutableFacade& Properties) const
+void FChaosClothAssetSimulationPBDBendingSpringConfigNode::AddProperties(Dataflow::FContext& Context, ::Chaos::Softs::FCollectionPropertyMutableFacade& Properties) const
 {
 	UE_CHAOS_CLOTHASSET_SIMULATIONCONFIG_SETPROPERTYWEIGHTEDCHECKED4(
 		BendingSpringStiffness,
diff --git a/Engine/Plugins/Experimental/ChaosClothAssetEditor/Source/ChaosClothAssetDataflowNodes/Private/ChaosClothAsset/SimulationPBDEdgeSpringConfigNode.cpp b/Engine/Plugins/Experimental/ChaosClothAssetEditor/Source/ChaosClothAssetDataflowNodes/Private/ChaosClothAsset/SimulationPBDEdgeSpringConfigNode.cpp
--- a/Engine/Plugins/Experimental/ChaosClothAssetEditor/Source/ChaosClothAssetDataflowNodes/Private/ChaosClothAsset/SimulationPBDEdgeSpringConfigNode.cpp
+++ b/Engine/Plugins/Experimental/ChaosClothAssetEditor/Source/ChaosClothAssetDataflowNodes/Private/ChaosClothAsset/SimulationPBDEdgeSpringConfigNode.cpp
@@ -10,9 +10,10 @@ FChaosClothAssetSimulationPBDEdgeSpringConfigNode::FChaosClothAssetSimulationPBD
 	: FChaosClothAssetSimulationBaseConfigNode(InParam, InGuid)
 {
 	RegisterCollectionConnections();
+	RegisterInputConnection(&EdgeSpringStiffness.WeightMap);
 }
 
-void FChaosClothAssetSimulationPBDEdgeSpringConfigNode::AddProperties(::Chaos::Softs::FCollectionPropertyMutableFacade& Properties) const
+void FChaosClothAssetSimulationPBDEdgeSpringConfigNode::AddProperties(Dataflow::FContext& Context, ::Chaos::Softs::FCollectionPropertyMutableFacade& Properties) const
 {
 	UE_CHAOS_CLOTHASSET_SIMULATIONCONFIG_SETPROPERTYWEIGHTEDCHECKED2(
 		EdgeSpringStiffness,
diff --git a/Engine/Plugins/Experimental/ChaosClothAssetEditor/Source/ChaosClothAssetDataflowNodes/Private/ChaosClothAsset/SimulationPressureConfigNode.cpp b/Engine/Plugins/Experimental/ChaosClothAssetEditor/Source/ChaosClothAssetDataflowNodes/Private/ChaosClothAsset/SimulationPressureConfigNode.cpp
--- a/Engine/Plugins/Experimental/ChaosClothAssetEditor/Source/ChaosClothAssetDataflowNodes/Private/ChaosClothAsset/SimulationPressureConfigNode.cpp
+++ b/Engine/Plugins/Experimental/ChaosClothAssetEditor/Source/ChaosClothAssetDataflowNodes/Private/ChaosClothAsset/SimulationPressureConfigNode.cpp
@@ -10,9 +10,10 @@ FChaosClothAssetSimulationPressureConfigNode::FChaosClothAssetSimulationPressure
 	: FChaosClothAssetSimulationBaseConfigNode(InParam, InGuid)
 {
 	RegisterCollectionConnections();
+	RegisterInputConnection(&Pressure.WeightMap);
 }
 
-void FChaosClothAssetSimulationPressureConfigNode::AddProperties(::Chaos::Softs::FCollectionPropertyMutableFacade& Properties) const
+void FChaosClothAssetSimulationPressureConfigNode::AddProperties(Dataflow::FContext& Context, ::Chaos::Softs::FCollectionPropertyMutableFacade& Properties) const
 {
 	UE_CHAOS_CLOTHASSET_SIMULATIONCONFIG_SETPROPERTYWEIGHTED(Pressure);
 }
diff --git a/Engine/Plugins/Experimental/ChaosClothAssetEditor/Source/ChaosClothAssetDataflowNodes/Private/ChaosClothAsset/SimulationSelfCollisionConfigNode.cpp b/Engine/Plugins/Experimental/ChaosClothAssetEditor/Source/ChaosClothAssetDataflowNodes/Private/ChaosClothAsset/SimulationSelfCollisionConfigNode.cpp
--- a/Engine/Plugins/Experimental/ChaosClothAssetEditor/Source/ChaosClothAssetDataflowNodes/Private/ChaosClothAsset/SimulationSelfCollisionConfigNode.cpp
+++ b/Engine/Plugins/Experimental/ChaosClothAssetEditor/Source/ChaosClothAssetDataflowNodes/Private/ChaosClothAsset/SimulationSelfCollisionConfigNode.cpp
@@ -12,8 +12,9 @@ FChaosClothAssetSimulationSelfCollisionConfigNode::FChaosClothAssetSimulationSel
 	RegisterCollectionConnections();
 }
 
-void FChaosClothAssetSimulationSelfCollisionConfigNode::AddProperties(::Chaos::Softs::FCollectionPropertyMutableFacade& Properties) const
+void FChaosClothAssetSimulationSelfCollisionConfigNode::AddProperties(Dataflow::FContext& Context, ::Chaos::Softs::FCollectionPropertyMutableFacade& Properties) const
 {
+	constexpr bool bUseSelfCollisions = true;
 	UE_CHAOS_CLOTHASSET_SIMULATIONCONFIG_SETPROPERTYBOOL(UseSelfCollisions);
 	UE_CHAOS_CLOTHASSET_SIMULATIONCONFIG_SETPROPERTY(SelfCollisionThickness);
 	UE_CHAOS_CLOTHASSET_SIMULATIONCONFIG_SETPROPERTY(SelfCollisionStiffness);
diff --git a/Engine/Plugins/Experimental/ChaosClothAssetEditor/Source/ChaosClothAssetDataflowNodes/Private/ChaosClothAsset/SimulationSolverConfigNode.cpp b/Engine/Plugins/Experimental/ChaosClothAssetEditor/Source/ChaosClothAssetDataflowNodes/Private/ChaosClothAsset/SimulationSolverConfigNode.cpp
--- a/Engine/Plugins/Experimental/ChaosClothAssetEditor/Source/ChaosClothAssetDataflowNodes/Private/ChaosClothAsset/SimulationSolverConfigNode.cpp
+++ b/Engine/Plugins/Experimental/ChaosClothAssetEditor/Source/ChaosClothAssetDataflowNodes/Private/ChaosClothAsset/SimulationSolverConfigNode.cpp
@@ -12,7 +12,7 @@ FChaosClothAssetSimulationSolverConfigNode::FChaosClothAssetSimulationSolverConf
 	RegisterCollectionConnections();
 }
 
-void FChaosClothAssetSimulationSolverConfigNode::AddProperties(::Chaos::Softs::FCollectionPropertyMutableFacade& Properties) const
+void FChaosClothAssetSimulationSolverConfigNode::AddProperties(Dataflow::FContext& Context, ::Chaos::Softs::FCollectionPropertyMutableFacade& Properties) const
 {
 	UE_CHAOS_CLOTHASSET_SIMULATIONCONFIG_SETPROPERTY(NumIterations);
 	UE_CHAOS_CLOTHASSET_SIMULATIONCONFIG_SETPROPERTY(MaxNumIterations);
diff --git a/Engine/Plugins/Experimental/ChaosClothAssetEditor/Source/ChaosClothAssetDataflowNodes/Private/ChaosClothAsset/SimulationVelocityScaleConfigNode.cpp b/Engine/Plugins/Experimental/ChaosClothAssetEditor/Source/ChaosClothAssetDataflowNodes/Private/ChaosClothAsset/SimulationVelocityScaleConfigNode.cpp
--- a/Engine/Plugins/Experimental/ChaosClothAssetEditor/Source/ChaosClothAssetDataflowNodes/Private/ChaosClothAsset/SimulationVelocityScaleConfigNode.cpp
+++ b/Engine/Plugins/Experimental/ChaosClothAssetEditor/Source/ChaosClothAssetDataflowNodes/Private/ChaosClothAsset/SimulationVelocityScaleConfigNode.cpp
@@ -12,7 +12,7 @@ FChaosClothAssetSimulationVelocityScaleConfigNode::FChaosClothAssetSimulationVel
 	RegisterCollectionConnections();
 }
 
-void FChaosClothAssetSimulationVelocityScaleConfigNode::AddProperties(::Chaos::Softs::FCollectionPropertyMutableFacade& Properties) const
+void FChaosClothAssetSimulationVelocityScaleConfigNode::AddProperties(Dataflow::FContext& Context, ::Chaos::Softs::FCollectionPropertyMutableFacade& Properties) const
 {
 	UE_CHAOS_CLOTHASSET_SIMULATIONCONFIG_SETPROPERTY(LinearVelocityScale);
 	UE_CHAOS_CLOTHASSET_SIMULATIONCONFIG_SETPROPERTY(AngularVelocityScale);
diff --git a/Engine/Plugins/Experimental/ChaosClothAssetEditor/Source/ChaosClothAssetDataflowNodes/Private/ChaosClothAsset/SimulationXPBDAnisoBendingConfigNode.cpp b/Engine/Plugins/Experimental/ChaosClothAssetEditor/Source/ChaosClothAssetDataflowNodes/Private/ChaosClothAsset/SimulationXPBDAnisoBendingConfigNode.cpp
--- a/Engine/Plugins/Experimental/ChaosClothAssetEditor/Source/ChaosClothAssetDataflowNodes/Private/ChaosClothAsset/SimulationXPBDAnisoBendingConfigNode.cpp
+++ b/Engine/Plugins/Experimental/ChaosClothAssetEditor/Source/ChaosClothAssetDataflowNodes/Private/ChaosClothAsset/SimulationXPBDAnisoBendingConfigNode.cpp
@@ -10,9 +10,16 @@ FChaosClothAssetSimulationXPBDAnisoBendingConfigNode::FChaosClothAssetSimulation
 	: FChaosClothAssetSimulationBaseConfigNode(InParam, InGuid)
 {
 	RegisterCollectionConnections();
+	RegisterInputConnection(&XPBDAnisoFlatnessRatio.WeightMap);
+	RegisterInputConnection(&XPBDAnisoRestAngle.WeightMap);
+	RegisterInputConnection(&XPBDAnisoBendingStiffnessWarp.WeightMap);
+	RegisterInputConnection(&XPBDAnisoBendingStiffnessWeft.WeightMap);
+	RegisterInputConnection(&XPBDAnisoBendingStiffnessBias.WeightMap);
+	RegisterInputConnection(&XPBDAnisoBendingDamping.WeightMap);
+	RegisterInputConnection(&XPBDAnisoBucklingStiffnessWarp.WeightMap);
 }
 
-void FChaosClothAssetSimulationXPBDAnisoBendingConfigNode::AddProperties(::Chaos::Softs::FCollectionPropertyMutableFacade& Properties) const
+void FChaosClothAssetSimulationXPBDAnisoBendingConfigNode::AddProperties(Dataflow::FContext& Context, ::Chaos::Softs::FCollectionPropertyMutableFacade& Properties) const
 {
 	UE_CHAOS_CLOTHASSET_SIMULATIONCONFIG_SETPROPERTYENUMCHECKED2(
 		XPBDAnisoRestAngleType,
diff --git a/Engine/Plugins/Experimental/ChaosClothAssetEditor/Source/ChaosClothAssetDataflowNodes/Private/ChaosClothAsset/SimulationXPBDAnisoStretchConfigNode.cpp b/Engine/Plugins/Experimental/ChaosClothAssetEditor/Source/ChaosClothAssetDataflowNodes/Private/ChaosClothAsset/SimulationXPBDAnisoStretchConfigNode.cpp
--- a/Engine/Plugins/Experimental/ChaosClothAssetEditor/Source/ChaosClothAssetDataflowNodes/Private/ChaosClothAsset/SimulationXPBDAnisoStretchConfigNode.cpp
+++ b/Engine/Plugins/Experimental/ChaosClothAssetEditor/Source/ChaosClothAssetDataflowNodes/Private/ChaosClothAsset/SimulationXPBDAnisoStretchConfigNode.cpp
@@ -10,9 +10,15 @@ FChaosClothAssetSimulationXPBDAnisoStretchConfigNode::FChaosClothAssetSimulation
 	: FChaosClothAssetSimulationBaseConfigNode(InParam, InGuid)
 {
 	RegisterCollectionConnections();
+	RegisterInputConnection(&XPBDAnisoStretchStiffnessWarp.WeightMap);
+	RegisterInputConnection(&XPBDAnisoStretchStiffnessWeft.WeightMap);
+	RegisterInputConnection(&XPBDAnisoStretchStiffnessBias.WeightMap);
+	RegisterInputConnection(&XPBDAnisoStretchDamping.WeightMap);
+	RegisterInputConnection(&XPBDAnisoStretchWarpScale.WeightMap);
+	RegisterInputConnection(&XPBDAnisoStretchWeftScale.WeightMap);
 }
 
-void FChaosClothAssetSimulationXPBDAnisoStretchConfigNode::AddProperties(::Chaos::Softs::FCollectionPropertyMutableFacade& Properties) const
+void FChaosClothAssetSimulationXPBDAnisoStretchConfigNode::AddProperties(Dataflow::FContext& Context, ::Chaos::Softs::FCollectionPropertyMutableFacade& Properties) const
 {
 	UE_CHAOS_CLOTHASSET_SIMULATIONCONFIG_SETPROPERTYBOOL(XPBDAnisoStretchUse3dRestLengths);
 	UE_CHAOS_CLOTHASSET_SIMULATIONCONFIG_SETPROPERTYWEIGHTEDCHECKED2(
diff --git a/Engine/Plugins/Experimental/ChaosClothAssetEditor/Source/ChaosClothAssetDataflowNodes/Private/ChaosClothAsset/SimulationXPBDAreaSpringConfigNode.cpp b/Engine/Plugins/Experimental/ChaosClothAssetEditor/Source/ChaosClothAssetDataflowNodes/Private/ChaosClothAsset/SimulationXPBDAreaSpringConfigNode.cpp
--- a/Engine/Plugins/Experimental/ChaosClothAssetEditor/Source/ChaosClothAssetDataflowNodes/Private/ChaosClothAsset/SimulationXPBDAreaSpringConfigNode.cpp
+++ b/Engine/Plugins/Experimental/ChaosClothAssetEditor/Source/ChaosClothAssetDataflowNodes/Private/ChaosClothAsset/SimulationXPBDAreaSpringConfigNode.cpp
@@ -10,9 +10,10 @@ FChaosClothAssetSimulationXPBDAreaSpringConfigNode::FChaosClothAssetSimulationXP
 	: FChaosClothAssetSimulationBaseConfigNode(InParam, InGuid)
 {
 	RegisterCollectionConnections();
+	RegisterInputConnection(&XPBDAreaSpringStiffness.WeightMap);
 }
 
-void FChaosClothAssetSimulationXPBDAreaSpringConfigNode::AddProperties(::Chaos::Softs::FCollectionPropertyMutableFacade& Properties) const
+void FChaosClothAssetSimulationXPBDAreaSpringConfigNode::AddProperties(Dataflow::FContext& Context, ::Chaos::Softs::FCollectionPropertyMutableFacade& Properties) const
 {
 	UE_CHAOS_CLOTHASSET_SIMULATIONCONFIG_SETPROPERTYWEIGHTEDCHECKED1(
 		XPBDAreaSpringStiffness,
diff --git a/Engine/Plugins/Experimental/ChaosClothAssetEditor/Source/ChaosClothAssetDataflowNodes/Private/ChaosClothAsset/SimulationXPBDBendingElementConfigNode.cpp b/Engine/Plugins/Experimental/ChaosClothAssetEditor/Source/ChaosClothAssetDataflowNodes/Private/ChaosClothAsset/SimulationXPBDBendingElementConfigNode.cpp
--- a/Engine/Plugins/Experimental/ChaosClothAssetEditor/Source/ChaosClothAssetDataflowNodes/Private/ChaosClothAsset/SimulationXPBDBendingElementConfigNode.cpp
+++ b/Engine/Plugins/Experimental/ChaosClothAssetEditor/Source/ChaosClothAssetDataflowNodes/Private/ChaosClothAsset/SimulationXPBDBendingElementConfigNode.cpp
@@ -10,9 +10,14 @@ FChaosClothAssetSimulationXPBDBendingElementConfigNode::FChaosClothAssetSimulati
 	: FChaosClothAssetSimulationBaseConfigNode(InParam, InGuid)
 {
 	RegisterCollectionConnections();
+	RegisterInputConnection(&XPBDFlatnessRatio.WeightMap);
+	RegisterInputConnection(&XPBDRestAngle.WeightMap);
+	RegisterInputConnection(&XPBDBendingElementStiffness.WeightMap);
+	RegisterInputConnection(&XPBDBendingElementDamping.WeightMap);
+	RegisterInputConnection(&XPBDBucklingStiffness.WeightMap);
 }
 
-void FChaosClothAssetSimulationXPBDBendingElementConfigNode::AddProperties(::Chaos::Softs::FCollectionPropertyMutableFacade& Properties) const
+void FChaosClothAssetSimulationXPBDBendingElementConfigNode::AddProperties(Dataflow::FContext& Context, ::Chaos::Softs::FCollectionPropertyMutableFacade& Properties) const
 {
 	UE_CHAOS_CLOTHASSET_SIMULATIONCONFIG_SETPROPERTYENUMCHECKED2(
 		XPBDRestAngleType,
diff --git a/Engine/Plugins/Experimental/ChaosClothAssetEditor/Source/ChaosClothAssetDataflowNodes/Private/ChaosClothAsset/SimulationXPBDBendingSpringConfigNode.cpp b/Engine/Plugins/Experimental/ChaosClothAssetEditor/Source/ChaosClothAssetDataflowNodes/Private/ChaosClothAsset/SimulationXPBDBendingSpringConfigNode.cpp
--- a/Engine/Plugins/Experimental/ChaosClothAssetEditor/Source/ChaosClothAssetDataflowNodes/Private/ChaosClothAsset/SimulationXPBDBendingSpringConfigNode.cpp
+++ b/Engine/Plugins/Experimental/ChaosClothAssetEditor/Source/ChaosClothAssetDataflowNodes/Private/ChaosClothAsset/SimulationXPBDBendingSpringConfigNode.cpp
@@ -10,9 +10,11 @@ FChaosClothAssetSimulationXPBDBendingSpringConfigNode::FChaosClothAssetSimulatio
 	: FChaosClothAssetSimulationBaseConfigNode(InParam, InGuid)
 {
 	RegisterCollectionConnections();
+	RegisterInputConnection(&XPBDBendingSpringStiffness.WeightMap);
+	RegisterInputConnection(&XPBDBendingSpringDamping.WeightMap);
 }
 
-void FChaosClothAssetSimulationXPBDBendingSpringConfigNode::AddProperties(::Chaos::Softs::FCollectionPropertyMutableFacade& Properties) const
+void FChaosClothAssetSimulationXPBDBendingSpringConfigNode::AddProperties(Dataflow::FContext& Context, ::Chaos::Softs::FCollectionPropertyMutableFacade& Properties) const
 {
 	UE_CHAOS_CLOTHASSET_SIMULATIONCONFIG_SETPROPERTYWEIGHTEDCHECKED4(
 		XPBDBendingSpringStiffness,
diff --git a/Engine/Plugins/Experimental/ChaosClothAssetEditor/Source/ChaosClothAssetDataflowNodes/Private/ChaosClothAsset/SimulationXPBDEdgeSpringConfigNode.cpp b/Engine/Plugins/Experimental/ChaosClothAssetEditor/Source/ChaosClothAssetDataflowNodes/Private/ChaosClothAsset/SimulationXPBDEdgeSpringConfigNode.cpp
--- a/Engine/Plugins/Experimental/ChaosClothAssetEditor/Source/ChaosClothAssetDataflowNodes/Private/ChaosClothAsset/SimulationXPBDEdgeSpringConfigNode.cpp
+++ b/Engine/Plugins/Experimental/ChaosClothAssetEditor/Source/ChaosClothAssetDataflowNodes/Private/ChaosClothAsset/SimulationXPBDEdgeSpringConfigNode.cpp
@@ -10,9 +10,11 @@ FChaosClothAssetSimulationXPBDEdgeSpringConfigNode::FChaosClothAssetSimulationXP
 	: FChaosClothAssetSimulationBaseConfigNode(InParam, InGuid)
 {
 	RegisterCollectionConnections();
+	RegisterInputConnection(&XPBDEdgeSpringStiffness.WeightMap);
+	RegisterInputConnection(&XPBDEdgeSpringDamping.WeightMap);
 }
 
-void FChaosClothAssetSimulationXPBDEdgeSpringConfigNode::AddProperties(::Chaos::Softs::FCollectionPropertyMutableFacade& Properties) const
+void FChaosClothAssetSimulationXPBDEdgeSpringConfigNode::AddProperties(Dataflow::FContext& Context, ::Chaos::Softs::FCollectionPropertyMutableFacade& Properties) const
 {
 	UE_CHAOS_CLOTHASSET_SIMULATIONCONFIG_SETPROPERTYWEIGHTEDCHECKED2(
 		XPBDEdgeSpringStiffness,
diff --git a/Engine/Plugins/Experimental/ChaosClothAssetEditor/Source/ChaosClothAssetDataflowNodes/Private/ChaosClothAsset/WeightedValueCustomization.cpp b/Engine/Plugins/Experimental/ChaosClothAssetEditor/Source/ChaosClothAssetDataflowNodes/Private/ChaosClothAsset/WeightedValueCustomization.cpp
--- a/Engine/Plugins/Experimental/ChaosClothAssetEditor/Source/ChaosClothAssetDataflowNodes/Private/ChaosClothAsset/WeightedValueCustomization.cpp
+++ b/Engine/Plugins/Experimental/ChaosClothAssetEditor/Source/ChaosClothAssetDataflowNodes/Private/ChaosClothAsset/WeightedValueCustomization.cpp
@@ -12,6 +12,24 @@
 
 namespace UE::Chaos::ClothAsset
 {
+	namespace Private
+	{
+		static const FString OverridePrefix = TEXT("_Override");
+
+		static bool IsOverrideProperty(const TSharedPtr<IPropertyHandle>& Property)
+		{
+			const FStringView PropertyPath = Property ? Property->GetPropertyPath() : FStringView();
+			return PropertyPath.EndsWith(OverridePrefix, ESearchCase::CaseSensitive);
+		}
+
+		static bool IsOverridePropertyOf(const TSharedPtr<IPropertyHandle>& OverrideProperty, const TSharedPtr<IPropertyHandle>& Property)
+		{
+			const FStringView OverridePropertyPath = OverrideProperty ? OverrideProperty->GetPropertyPath() : FStringView();
+			const FStringView PropertyPath = Property ? Property->GetPropertyPath() : FStringView();
+			return OverridePropertyPath == FString(PropertyPath) + OverridePrefix;
+		}
+	}
+
 	TSharedRef<IPropertyTypeCustomization> FWeightedValueCustomization::MakeInstance()
 	{
 		return MakeShareable(new FWeightedValueCustomization);
@@ -44,6 +62,11 @@ namespace UE::Chaos::ClothAsset
 		{
 			TSharedRef<IPropertyHandle> ChildHandle = SortedChildHandles[ChildIndex];
 
+			if (Private::IsOverrideProperty(ChildHandle))
+			{
+				continue;  // Skip overrides
+			}
+
 			// Propagate metadata to child properties so that it's reflected in the nested, individual spin boxes
 			ChildHandle->SetInstanceMetaData(TEXT("UIMin"), StructPropertyHandle->GetMetaData(TEXT("UIMin")));
 			ChildHandle->SetInstanceMetaData(TEXT("UIMax"), StructPropertyHandle->GetMetaData(TEXT("UIMax")));
@@ -119,20 +142,57 @@ namespace UE::Chaos::ClothAsset
 		}
 		if (PropertyClass == FStrProperty::StaticClass())
 		{
-			TWeakPtr<IPropertyHandle> WeakHandlePtr = PropertyHandle;
+			// Manage override property values (properties ending with _Override)
+			TWeakPtr<IPropertyHandle> OverrideHandleWeakPtr;
+
+			for (int32 ChildIndex = 0; ChildIndex < SortedChildHandles.Num(); ++ChildIndex)
+			{
+				const bool bLastChild = SortedChildHandles.Num() - 1 == ChildIndex;
+				const TSharedRef<IPropertyHandle>& ChildHandle = SortedChildHandles[ChildIndex];
+
+				if (Private::IsOverridePropertyOf(ChildHandle, PropertyHandle))
+				{
+					OverrideHandleWeakPtr = ChildHandle;
+					break;
+				}
+			}
+
+			TWeakPtr<IPropertyHandle> HandleWeakPtr = PropertyHandle;
 			return
 				SNew(SEditableTextBox)
 				.ToolTipText(LOCTEXT("WeightMap", "The name of the weight map for this property."))
-				.Text_Lambda([WeakHandlePtr]() -> FText
+				.Text_Lambda([HandleWeakPtr, OverrideHandleWeakPtr]() -> FText
 					{
 						FString Text;
-						WeakHandlePtr.Pin()->GetValue(Text);
+						if (const TSharedPtr<IPropertyHandle> OverrideHandlePtr = OverrideHandleWeakPtr.Pin())
+						{
+							OverrideHandlePtr->GetValue(Text);
+						}
+						if (Text.IsEmpty())
+						{
+							if (const TSharedPtr<IPropertyHandle> HandlePtr = HandleWeakPtr.Pin())
+							{
+								HandlePtr->GetValue(Text);
+							}
+						}
 						return FText::FromString(Text);
 
 					})
-				.OnTextCommitted_Lambda([WeakHandlePtr](const FText& Text, ETextCommit::Type)
+				.OnTextCommitted_Lambda([HandleWeakPtr](const FText& Text, ETextCommit::Type)
 					{
-						WeakHandlePtr.Pin()->SetValue(Text.ToString(), EPropertyValueSetFlags::DefaultFlags);
+						if (const TSharedPtr<IPropertyHandle> HandlePtr = HandleWeakPtr.Pin())
+						{
+							HandlePtr->SetValue(Text.ToString(), EPropertyValueSetFlags::DefaultFlags);
+						}
+					})
+				.IsEnabled_Lambda([OverrideHandleWeakPtr]() -> bool
+					{
+						FString Text;
+						if (const TSharedPtr<IPropertyHandle> OverrideHandlePtr = OverrideHandleWeakPtr.Pin())
+						{
+							OverrideHandlePtr->GetValue(Text);
+						}
+						return Text.IsEmpty();
 					})
 				.Font(IPropertyTypeCustomizationUtils::GetRegularFont());
 		}
diff --git a/Engine/Plugins/Experimental/ChaosClothAssetEditor/Source/ChaosClothAssetDataflowNodes/Public/ChaosClothAsset/SimulationAerodynamicsConfigNode.h b/Engine/Plugins/Experimental/ChaosClothAssetEditor/Source/ChaosClothAssetDataflowNodes/Public/ChaosClothAsset/SimulationAerodynamicsConfigNode.h
--- a/Engine/Plugins/Experimental/ChaosClothAssetEditor/Source/ChaosClothAssetDataflowNodes/Public/ChaosClothAsset/SimulationAerodynamicsConfigNode.h
+++ b/Engine/Plugins/Experimental/ChaosClothAssetEditor/Source/ChaosClothAssetDataflowNodes/Public/ChaosClothAsset/SimulationAerodynamicsConfigNode.h
@@ -43,5 +43,5 @@ public:
 	FChaosClothAssetSimulationAerodynamicsConfigNode(const Dataflow::FNodeParameters& InParam, FGuid InGuid = FGuid::NewGuid());
 
 private:
-	virtual void AddProperties(::Chaos::Softs::FCollectionPropertyMutableFacade& Properties) const override;
+	virtual void AddProperties(Dataflow::FContext& Context, ::Chaos::Softs::FCollectionPropertyMutableFacade& Properties) const override;
 };
diff --git a/Engine/Plugins/Experimental/ChaosClothAssetEditor/Source/ChaosClothAssetDataflowNodes/Public/ChaosClothAsset/SimulationAnimDriveConfigNode.h b/Engine/Plugins/Experimental/ChaosClothAssetEditor/Source/ChaosClothAssetDataflowNodes/Public/ChaosClothAsset/SimulationAnimDriveConfigNode.h
--- a/Engine/Plugins/Experimental/ChaosClothAssetEditor/Source/ChaosClothAssetDataflowNodes/Public/ChaosClothAsset/SimulationAnimDriveConfigNode.h
+++ b/Engine/Plugins/Experimental/ChaosClothAssetEditor/Source/ChaosClothAssetDataflowNodes/Public/ChaosClothAsset/SimulationAnimDriveConfigNode.h
@@ -35,5 +35,5 @@ public:
 	FChaosClothAssetSimulationAnimDriveConfigNode(const Dataflow::FNodeParameters& InParam, FGuid InGuid = FGuid::NewGuid());
 
 private:
-	virtual void AddProperties(::Chaos::Softs::FCollectionPropertyMutableFacade& Properties) const override;
+	virtual void AddProperties(Dataflow::FContext& Context, ::Chaos::Softs::FCollectionPropertyMutableFacade& Properties) const override;
 };
diff --git a/Engine/Plugins/Experimental/ChaosClothAssetEditor/Source/ChaosClothAssetDataflowNodes/Public/ChaosClothAsset/SimulationBackstopConfigNode.h b/Engine/Plugins/Experimental/ChaosClothAssetEditor/Source/ChaosClothAssetDataflowNodes/Public/ChaosClothAsset/SimulationBackstopConfigNode.h
--- a/Engine/Plugins/Experimental/ChaosClothAssetEditor/Source/ChaosClothAssetDataflowNodes/Public/ChaosClothAsset/SimulationBackstopConfigNode.h
+++ b/Engine/Plugins/Experimental/ChaosClothAssetEditor/Source/ChaosClothAssetDataflowNodes/Public/ChaosClothAsset/SimulationBackstopConfigNode.h
@@ -36,5 +36,5 @@ public:
 	FChaosClothAssetSimulationBackstopConfigNode(const Dataflow::FNodeParameters& InParam, FGuid InGuid = FGuid::NewGuid());
 
 private:
-	virtual void AddProperties(::Chaos::Softs::FCollectionPropertyMutableFacade& Properties) const override;
+	virtual void AddProperties(Dataflow::FContext& Context, ::Chaos::Softs::FCollectionPropertyMutableFacade& Properties) const override;
 };
diff --git a/Engine/Plugins/Experimental/ChaosClothAssetEditor/Source/ChaosClothAssetDataflowNodes/Public/ChaosClothAsset/SimulationBaseConfigNode.h b/Engine/Plugins/Experimental/ChaosClothAssetEditor/Source/ChaosClothAssetDataflowNodes/Public/ChaosClothAsset/SimulationBaseConfigNode.h
--- a/Engine/Plugins/Experimental/ChaosClothAssetEditor/Source/ChaosClothAssetDataflowNodes/Public/ChaosClothAsset/SimulationBaseConfigNode.h
+++ b/Engine/Plugins/Experimental/ChaosClothAssetEditor/Source/ChaosClothAssetDataflowNodes/Public/ChaosClothAsset/SimulationBaseConfigNode.h
@@ -31,7 +31,7 @@ public:
 protected:
 	virtual void Evaluate(Dataflow::FContext& Context, const FDataflowOutput* Out) const override;
 
-	virtual void AddProperties(::Chaos::Softs::FCollectionPropertyMutableFacade& Properties) const
+	virtual void AddProperties(Dataflow::FContext& Context, ::Chaos::Softs::FCollectionPropertyMutableFacade& Properties) const
 	PURE_VIRTUAL(FChaosClothAssetSimulationBaseConfigNode::AddProperties, );
 
 	/* Override this to do additional node-specific evaluate on the cloth collection output. AddProperties has already been called when this is called. */
diff --git a/Engine/Plugins/Experimental/ChaosClothAssetEditor/Source/ChaosClothAssetDataflowNodes/Public/ChaosClothAsset/SimulationCollisionConfigNode.h b/Engine/Plugins/Experimental/ChaosClothAssetEditor/Source/ChaosClothAssetDataflowNodes/Public/ChaosClothAsset/SimulationCollisionConfigNode.h
--- a/Engine/Plugins/Experimental/ChaosClothAssetEditor/Source/ChaosClothAssetDataflowNodes/Public/ChaosClothAsset/SimulationCollisionConfigNode.h
+++ b/Engine/Plugins/Experimental/ChaosClothAssetEditor/Source/ChaosClothAssetDataflowNodes/Public/ChaosClothAsset/SimulationCollisionConfigNode.h
@@ -31,5 +31,5 @@ public:
 	FChaosClothAssetSimulationCollisionConfigNode(const Dataflow::FNodeParameters& InParam, FGuid InGuid = FGuid::NewGuid());
 
 private:
-	virtual void AddProperties(::Chaos::Softs::FCollectionPropertyMutableFacade& Properties) const override;
+	virtual void AddProperties(Dataflow::FContext& Context, ::Chaos::Softs::FCollectionPropertyMutableFacade& Properties) const override;
 };
diff --git a/Engine/Plugins/Experimental/ChaosClothAssetEditor/Source/ChaosClothAssetDataflowNodes/Public/ChaosClothAsset/SimulationDampingConfigNode.h b/Engine/Plugins/Experimental/ChaosClothAssetEditor/Source/ChaosClothAssetDataflowNodes/Public/ChaosClothAsset/SimulationDampingConfigNode.h
--- a/Engine/Plugins/Experimental/ChaosClothAssetEditor/Source/ChaosClothAssetDataflowNodes/Public/ChaosClothAsset/SimulationDampingConfigNode.h
+++ b/Engine/Plugins/Experimental/ChaosClothAssetEditor/Source/ChaosClothAssetDataflowNodes/Public/ChaosClothAsset/SimulationDampingConfigNode.h
@@ -32,5 +32,5 @@ public:
 	FChaosClothAssetSimulationDampingConfigNode(const Dataflow::FNodeParameters& InParam, FGuid InGuid = FGuid::NewGuid());
 
 private:
-	virtual void AddProperties(::Chaos::Softs::FCollectionPropertyMutableFacade& Properties) const override;
+	virtual void AddProperties(Dataflow::FContext& Context, ::Chaos::Softs::FCollectionPropertyMutableFacade& Properties) const override;
 };
diff --git a/Engine/Plugins/Experimental/ChaosClothAssetEditor/Source/ChaosClothAssetDataflowNodes/Public/ChaosClothAsset/SimulationGravityConfigNode.h b/Engine/Plugins/Experimental/ChaosClothAssetEditor/Source/ChaosClothAssetDataflowNodes/Public/ChaosClothAsset/SimulationGravityConfigNode.h
--- a/Engine/Plugins/Experimental/ChaosClothAssetEditor/Source/ChaosClothAssetDataflowNodes/Public/ChaosClothAsset/SimulationGravityConfigNode.h
+++ b/Engine/Plugins/Experimental/ChaosClothAssetEditor/Source/ChaosClothAssetDataflowNodes/Public/ChaosClothAsset/SimulationGravityConfigNode.h
@@ -28,5 +28,5 @@ public:
 	FChaosClothAssetSimulationGravityConfigNode(const Dataflow::FNodeParameters& InParam, FGuid InGuid = FGuid::NewGuid());
 
 private:
-	virtual void AddProperties(::Chaos::Softs::FCollectionPropertyMutableFacade& Properties) const override;
+	virtual void AddProperties(Dataflow::FContext& Context, ::Chaos::Softs::FCollectionPropertyMutableFacade& Properties) const override;
 };
diff --git a/Engine/Plugins/Experimental/ChaosClothAssetEditor/Source/ChaosClothAssetDataflowNodes/Public/ChaosClothAsset/SimulationLongRangeAttachmentConfigNode.h b/Engine/Plugins/Experimental/ChaosClothAssetEditor/Source/ChaosClothAssetDataflowNodes/Public/ChaosClothAsset/SimulationLongRangeAttachmentConfigNode.h
--- a/Engine/Plugins/Experimental/ChaosClothAssetEditor/Source/ChaosClothAssetDataflowNodes/Public/ChaosClothAsset/SimulationLongRangeAttachmentConfigNode.h
+++ b/Engine/Plugins/Experimental/ChaosClothAssetEditor/Source/ChaosClothAssetDataflowNodes/Public/ChaosClothAsset/SimulationLongRangeAttachmentConfigNode.h
@@ -51,7 +51,7 @@ public:
 	FChaosClothAssetSimulationLongRangeAttachmentConfigNode(const Dataflow::FNodeParameters& InParam, FGuid InGuid = FGuid::NewGuid());
 
 private:
-	virtual void AddProperties(::Chaos::Softs::FCollectionPropertyMutableFacade& Properties) const override;
+	virtual void AddProperties(Dataflow::FContext& Context, ::Chaos::Softs::FCollectionPropertyMutableFacade& Properties) const override;
 
 	virtual void EvaluateClothCollection(Dataflow::FContext& Context, const TSharedRef<FManagedArrayCollection>& ClothCollection) const override;
 };
diff --git a/Engine/Plugins/Experimental/ChaosClothAssetEditor/Source/ChaosClothAssetDataflowNodes/Public/ChaosClothAsset/SimulationMassConfigNode.h b/Engine/Plugins/Experimental/ChaosClothAssetEditor/Source/ChaosClothAssetDataflowNodes/Public/ChaosClothAsset/SimulationMassConfigNode.h
--- a/Engine/Plugins/Experimental/ChaosClothAssetEditor/Source/ChaosClothAssetDataflowNodes/Public/ChaosClothAsset/SimulationMassConfigNode.h
+++ b/Engine/Plugins/Experimental/ChaosClothAssetEditor/Source/ChaosClothAssetDataflowNodes/Public/ChaosClothAsset/SimulationMassConfigNode.h
@@ -47,5 +47,5 @@ public:
 	FChaosClothAssetSimulationMassConfigNode(const Dataflow::FNodeParameters& InParam, FGuid InGuid = FGuid::NewGuid());
 
 private:
-	virtual void AddProperties(::Chaos::Softs::FCollectionPropertyMutableFacade& Properties) const override;
+	virtual void AddProperties(Dataflow::FContext& Context, ::Chaos::Softs::FCollectionPropertyMutableFacade& Properties) const override;
 };
diff --git a/Engine/Plugins/Experimental/ChaosClothAssetEditor/Source/ChaosClothAssetDataflowNodes/Public/ChaosClothAsset/SimulationMaxDistanceConfigNode.h b/Engine/Plugins/Experimental/ChaosClothAssetEditor/Source/ChaosClothAssetDataflowNodes/Public/ChaosClothAsset/SimulationMaxDistanceConfigNode.h
--- a/Engine/Plugins/Experimental/ChaosClothAssetEditor/Source/ChaosClothAssetDataflowNodes/Public/ChaosClothAsset/SimulationMaxDistanceConfigNode.h
+++ b/Engine/Plugins/Experimental/ChaosClothAssetEditor/Source/ChaosClothAssetDataflowNodes/Public/ChaosClothAsset/SimulationMaxDistanceConfigNode.h
@@ -27,5 +27,5 @@ public:
 	FChaosClothAssetSimulationMaxDistanceConfigNode(const Dataflow::FNodeParameters& InParam, FGuid InGuid = FGuid::NewGuid());
 
 private:
-	virtual void AddProperties(::Chaos::Softs::FCollectionPropertyMutableFacade& Properties) const override;
+	virtual void AddProperties(Dataflow::FContext& Context, ::Chaos::Softs::FCollectionPropertyMutableFacade& Properties) const override;
 };
diff --git a/Engine/Plugins/Experimental/ChaosClothAssetEditor/Source/ChaosClothAssetDataflowNodes/Public/ChaosClothAsset/SimulationPBDAreaSpringConfigNode.h b/Engine/Plugins/Experimental/ChaosClothAssetEditor/Source/ChaosClothAssetDataflowNodes/Public/ChaosClothAsset/SimulationPBDAreaSpringConfigNode.h
--- a/Engine/Plugins/Experimental/ChaosClothAssetEditor/Source/ChaosClothAssetDataflowNodes/Public/ChaosClothAsset/SimulationPBDAreaSpringConfigNode.h
+++ b/Engine/Plugins/Experimental/ChaosClothAssetEditor/Source/ChaosClothAssetDataflowNodes/Public/ChaosClothAsset/SimulationPBDAreaSpringConfigNode.h
@@ -27,5 +27,5 @@ public:
 	FChaosClothAssetSimulationPBDAreaSpringConfigNode(const Dataflow::FNodeParameters& InParam, FGuid InGuid = FGuid::NewGuid());
 
 private:
-	virtual void AddProperties(::Chaos::Softs::FCollectionPropertyMutableFacade& Properties) const override;
+	virtual void AddProperties(Dataflow::FContext& Context, ::Chaos::Softs::FCollectionPropertyMutableFacade& Properties) const override;
 };
diff --git a/Engine/Plugins/Experimental/ChaosClothAssetEditor/Source/ChaosClothAssetDataflowNodes/Public/ChaosClothAsset/SimulationPBDBendingElementConfigNode.h b/Engine/Plugins/Experimental/ChaosClothAssetEditor/Source/ChaosClothAssetDataflowNodes/Public/ChaosClothAsset/SimulationPBDBendingElementConfigNode.h
--- a/Engine/Plugins/Experimental/ChaosClothAssetEditor/Source/ChaosClothAssetDataflowNodes/Public/ChaosClothAsset/SimulationPBDBendingElementConfigNode.h
+++ b/Engine/Plugins/Experimental/ChaosClothAssetEditor/Source/ChaosClothAssetDataflowNodes/Public/ChaosClothAsset/SimulationPBDBendingElementConfigNode.h
@@ -85,5 +85,5 @@ public:
 	FChaosClothAssetSimulationPBDBendingElementConfigNode(const Dataflow::FNodeParameters& InParam, FGuid InGuid = FGuid::NewGuid());
 
 private:
-	virtual void AddProperties(::Chaos::Softs::FCollectionPropertyMutableFacade& Properties) const override;
+	virtual void AddProperties(Dataflow::FContext& Context, ::Chaos::Softs::FCollectionPropertyMutableFacade& Properties) const override;
 };
diff --git a/Engine/Plugins/Experimental/ChaosClothAssetEditor/Source/ChaosClothAssetDataflowNodes/Public/ChaosClothAsset/SimulationPBDBendingSpringConfigNode.h b/Engine/Plugins/Experimental/ChaosClothAssetEditor/Source/ChaosClothAssetDataflowNodes/Public/ChaosClothAsset/SimulationPBDBendingSpringConfigNode.h
--- a/Engine/Plugins/Experimental/ChaosClothAssetEditor/Source/ChaosClothAssetDataflowNodes/Public/ChaosClothAsset/SimulationPBDBendingSpringConfigNode.h
+++ b/Engine/Plugins/Experimental/ChaosClothAssetEditor/Source/ChaosClothAssetDataflowNodes/Public/ChaosClothAsset/SimulationPBDBendingSpringConfigNode.h
@@ -27,5 +27,5 @@ public:
 	FChaosClothAssetSimulationPBDBendingSpringConfigNode(const Dataflow::FNodeParameters& InParam, FGuid InGuid = FGuid::NewGuid());
 
 private:
-	virtual void AddProperties(::Chaos::Softs::FCollectionPropertyMutableFacade& Properties) const override;
+	virtual void AddProperties(Dataflow::FContext& Context, ::Chaos::Softs::FCollectionPropertyMutableFacade& Properties) const override;
 };
diff --git a/Engine/Plugins/Experimental/ChaosClothAssetEditor/Source/ChaosClothAssetDataflowNodes/Public/ChaosClothAsset/SimulationPBDEdgeSpringConfigNode.h b/Engine/Plugins/Experimental/ChaosClothAssetEditor/Source/ChaosClothAssetDataflowNodes/Public/ChaosClothAsset/SimulationPBDEdgeSpringConfigNode.h
--- a/Engine/Plugins/Experimental/ChaosClothAssetEditor/Source/ChaosClothAssetDataflowNodes/Public/ChaosClothAsset/SimulationPBDEdgeSpringConfigNode.h
+++ b/Engine/Plugins/Experimental/ChaosClothAssetEditor/Source/ChaosClothAssetDataflowNodes/Public/ChaosClothAsset/SimulationPBDEdgeSpringConfigNode.h
@@ -27,5 +27,5 @@ public:
 	FChaosClothAssetSimulationPBDEdgeSpringConfigNode(const Dataflow::FNodeParameters& InParam, FGuid InGuid = FGuid::NewGuid());
 
 private:
-	virtual void AddProperties(::Chaos::Softs::FCollectionPropertyMutableFacade& Properties) const override;
+	virtual void AddProperties(Dataflow::FContext& Context, ::Chaos::Softs::FCollectionPropertyMutableFacade& Properties) const override;
 };
diff --git a/Engine/Plugins/Experimental/ChaosClothAssetEditor/Source/ChaosClothAssetDataflowNodes/Public/ChaosClothAsset/SimulationPressureConfigNode.h b/Engine/Plugins/Experimental/ChaosClothAssetEditor/Source/ChaosClothAssetDataflowNodes/Public/ChaosClothAsset/SimulationPressureConfigNode.h
--- a/Engine/Plugins/Experimental/ChaosClothAssetEditor/Source/ChaosClothAssetDataflowNodes/Public/ChaosClothAsset/SimulationPressureConfigNode.h
+++ b/Engine/Plugins/Experimental/ChaosClothAssetEditor/Source/ChaosClothAssetDataflowNodes/Public/ChaosClothAsset/SimulationPressureConfigNode.h
@@ -26,5 +26,5 @@ public:
 	FChaosClothAssetSimulationPressureConfigNode(const Dataflow::FNodeParameters& InParam, FGuid InGuid = FGuid::NewGuid());
 
 private:
-	virtual void AddProperties(::Chaos::Softs::FCollectionPropertyMutableFacade& Properties) const override;
+	virtual void AddProperties(Dataflow::FContext& Context, ::Chaos::Softs::FCollectionPropertyMutableFacade& Properties) const override;
 };
diff --git a/Engine/Plugins/Experimental/ChaosClothAssetEditor/Source/ChaosClothAssetDataflowNodes/Public/ChaosClothAsset/SimulationSelfCollisionConfigNode.h b/Engine/Plugins/Experimental/ChaosClothAssetEditor/Source/ChaosClothAssetDataflowNodes/Public/ChaosClothAsset/SimulationSelfCollisionConfigNode.h
--- a/Engine/Plugins/Experimental/ChaosClothAssetEditor/Source/ChaosClothAssetDataflowNodes/Public/ChaosClothAsset/SimulationSelfCollisionConfigNode.h
+++ b/Engine/Plugins/Experimental/ChaosClothAssetEditor/Source/ChaosClothAssetDataflowNodes/Public/ChaosClothAsset/SimulationSelfCollisionConfigNode.h
@@ -13,9 +13,6 @@ struct FChaosClothAssetSimulationSelfCollisionConfigNode : public FChaosClothAss
 	DATAFLOW_NODE_DEFINE_INTERNAL(FChaosClothAssetSimulationSelfCollisionConfigNode, "SimulationSelfCollisionConfig", "Cloth", "Cloth Simulation Self Collision Config")
 
 public:
-	UPROPERTY(EditAnywhere, Category = "Self-Collision Properties")
-	bool bUseSelfCollisions = true;
-
 	/** The radius of the spheres used in self collision (i.e., offset per side. total thickness of cloth is 2x this value). */
 	UPROPERTY(EditAnywhere, Category = "Self-Collision Properties", meta = (UIMin = "0", UIMax = "100", ClampMin = "0", ClampMax = "1000", EditCondition = "bUseSelfCollisions"))
 	float SelfCollisionThickness = 0.5f;
@@ -35,5 +32,5 @@ public:
 	FChaosClothAssetSimulationSelfCollisionConfigNode(const Dataflow::FNodeParameters& InParam, FGuid InGuid = FGuid::NewGuid());
 
 private:
-	virtual void AddProperties(::Chaos::Softs::FCollectionPropertyMutableFacade& Properties) const override;
+	virtual void AddProperties(Dataflow::FContext& Context, ::Chaos::Softs::FCollectionPropertyMutableFacade& Properties) const override;
 };
diff --git a/Engine/Plugins/Experimental/ChaosClothAssetEditor/Source/ChaosClothAssetDataflowNodes/Public/ChaosClothAsset/SimulationSolverConfigNode.h b/Engine/Plugins/Experimental/ChaosClothAssetEditor/Source/ChaosClothAssetDataflowNodes/Public/ChaosClothAsset/SimulationSolverConfigNode.h
--- a/Engine/Plugins/Experimental/ChaosClothAssetEditor/Source/ChaosClothAssetDataflowNodes/Public/ChaosClothAsset/SimulationSolverConfigNode.h
+++ b/Engine/Plugins/Experimental/ChaosClothAssetEditor/Source/ChaosClothAssetDataflowNodes/Public/ChaosClothAsset/SimulationSolverConfigNode.h
@@ -39,5 +39,5 @@ public:
 	FChaosClothAssetSimulationSolverConfigNode(const Dataflow::FNodeParameters& InParam, FGuid InGuid = FGuid::NewGuid());
 
 private:
-	virtual void AddProperties(::Chaos::Softs::FCollectionPropertyMutableFacade& Properties) const override;
+	virtual void AddProperties(Dataflow::FContext& Context, ::Chaos::Softs::FCollectionPropertyMutableFacade& Properties) const override;
 };
diff --git a/Engine/Plugins/Experimental/ChaosClothAssetEditor/Source/ChaosClothAssetDataflowNodes/Public/ChaosClothAsset/SimulationVelocityScaleConfigNode.h b/Engine/Plugins/Experimental/ChaosClothAssetEditor/Source/ChaosClothAssetDataflowNodes/Public/ChaosClothAsset/SimulationVelocityScaleConfigNode.h
--- a/Engine/Plugins/Experimental/ChaosClothAssetEditor/Source/ChaosClothAssetDataflowNodes/Public/ChaosClothAsset/SimulationVelocityScaleConfigNode.h
+++ b/Engine/Plugins/Experimental/ChaosClothAssetEditor/Source/ChaosClothAssetDataflowNodes/Public/ChaosClothAsset/SimulationVelocityScaleConfigNode.h
@@ -41,5 +41,5 @@ public:
 	FChaosClothAssetSimulationVelocityScaleConfigNode(const Dataflow::FNodeParameters& InParam, FGuid InGuid = FGuid::NewGuid());
 
 private:
-	virtual void AddProperties(::Chaos::Softs::FCollectionPropertyMutableFacade& Properties) const override;
+	virtual void AddProperties(Dataflow::FContext& Context, ::Chaos::Softs::FCollectionPropertyMutableFacade& Properties) const override;
 };
diff --git a/Engine/Plugins/Experimental/ChaosClothAssetEditor/Source/ChaosClothAssetDataflowNodes/Public/ChaosClothAsset/SimulationXPBDAnisoBendingConfigNode.h b/Engine/Plugins/Experimental/ChaosClothAssetEditor/Source/ChaosClothAssetDataflowNodes/Public/ChaosClothAsset/SimulationXPBDAnisoBendingConfigNode.h
--- a/Engine/Plugins/Experimental/ChaosClothAssetEditor/Source/ChaosClothAssetDataflowNodes/Public/ChaosClothAsset/SimulationXPBDAnisoBendingConfigNode.h
+++ b/Engine/Plugins/Experimental/ChaosClothAssetEditor/Source/ChaosClothAssetDataflowNodes/Public/ChaosClothAsset/SimulationXPBDAnisoBendingConfigNode.h
@@ -125,5 +125,5 @@ public:
 	FChaosClothAssetSimulationXPBDAnisoBendingConfigNode(const Dataflow::FNodeParameters& InParam, FGuid InGuid = FGuid::NewGuid());
 
 private:
-	virtual void AddProperties(::Chaos::Softs::FCollectionPropertyMutableFacade& Properties) const override;
+	virtual void AddProperties(Dataflow::FContext& Context, ::Chaos::Softs::FCollectionPropertyMutableFacade& Properties) const override;
 };
diff --git a/Engine/Plugins/Experimental/ChaosClothAssetEditor/Source/ChaosClothAssetDataflowNodes/Public/ChaosClothAsset/SimulationXPBDAnisoStretchConfigNode.h b/Engine/Plugins/Experimental/ChaosClothAssetEditor/Source/ChaosClothAssetDataflowNodes/Public/ChaosClothAsset/SimulationXPBDAnisoStretchConfigNode.h
--- a/Engine/Plugins/Experimental/ChaosClothAssetEditor/Source/ChaosClothAssetDataflowNodes/Public/ChaosClothAsset/SimulationXPBDAnisoStretchConfigNode.h
+++ b/Engine/Plugins/Experimental/ChaosClothAssetEditor/Source/ChaosClothAssetDataflowNodes/Public/ChaosClothAsset/SimulationXPBDAnisoStretchConfigNode.h
@@ -75,5 +75,5 @@ public:
 	FChaosClothAssetSimulationXPBDAnisoStretchConfigNode(const Dataflow::FNodeParameters& InParam, FGuid InGuid = FGuid::NewGuid());
 
 private:
-	virtual void AddProperties(::Chaos::Softs::FCollectionPropertyMutableFacade& Properties) const override;
+	virtual void AddProperties(Dataflow::FContext& Context, ::Chaos::Softs::FCollectionPropertyMutableFacade& Properties) const override;
 };
diff --git a/Engine/Plugins/Experimental/ChaosClothAssetEditor/Source/ChaosClothAssetDataflowNodes/Public/ChaosClothAsset/SimulationXPBDAreaSpringConfigNode.h b/Engine/Plugins/Experimental/ChaosClothAssetEditor/Source/ChaosClothAssetDataflowNodes/Public/ChaosClothAsset/SimulationXPBDAreaSpringConfigNode.h
--- a/Engine/Plugins/Experimental/ChaosClothAssetEditor/Source/ChaosClothAssetDataflowNodes/Public/ChaosClothAsset/SimulationXPBDAreaSpringConfigNode.h
+++ b/Engine/Plugins/Experimental/ChaosClothAssetEditor/Source/ChaosClothAssetDataflowNodes/Public/ChaosClothAsset/SimulationXPBDAreaSpringConfigNode.h
@@ -27,5 +27,5 @@ public:
 	FChaosClothAssetSimulationXPBDAreaSpringConfigNode(const Dataflow::FNodeParameters& InParam, FGuid InGuid = FGuid::NewGuid());
 
 private:
-	virtual void AddProperties(::Chaos::Softs::FCollectionPropertyMutableFacade& Properties) const override;
+	virtual void AddProperties(Dataflow::FContext& Context, ::Chaos::Softs::FCollectionPropertyMutableFacade& Properties) const override;
 };
diff --git a/Engine/Plugins/Experimental/ChaosClothAssetEditor/Source/ChaosClothAssetDataflowNodes/Public/ChaosClothAsset/SimulationXPBDBendingElementConfigNode.h b/Engine/Plugins/Experimental/ChaosClothAssetEditor/Source/ChaosClothAssetDataflowNodes/Public/ChaosClothAsset/SimulationXPBDBendingElementConfigNode.h
--- a/Engine/Plugins/Experimental/ChaosClothAssetEditor/Source/ChaosClothAssetDataflowNodes/Public/ChaosClothAsset/SimulationXPBDBendingElementConfigNode.h
+++ b/Engine/Plugins/Experimental/ChaosClothAssetEditor/Source/ChaosClothAssetDataflowNodes/Public/ChaosClothAsset/SimulationXPBDBendingElementConfigNode.h
@@ -83,5 +83,5 @@ public:
 	FChaosClothAssetSimulationXPBDBendingElementConfigNode(const Dataflow::FNodeParameters& InParam, FGuid InGuid = FGuid::NewGuid());
 
 private:
-	virtual void AddProperties(::Chaos::Softs::FCollectionPropertyMutableFacade& Properties) const override;
+	virtual void AddProperties(Dataflow::FContext& Context, ::Chaos::Softs::FCollectionPropertyMutableFacade& Properties) const override;
 };
diff --git a/Engine/Plugins/Experimental/ChaosClothAssetEditor/Source/ChaosClothAssetDataflowNodes/Public/ChaosClothAsset/SimulationXPBDBendingSpringConfigNode.h b/Engine/Plugins/Experimental/ChaosClothAssetEditor/Source/ChaosClothAssetDataflowNodes/Public/ChaosClothAsset/SimulationXPBDBendingSpringConfigNode.h
--- a/Engine/Plugins/Experimental/ChaosClothAssetEditor/Source/ChaosClothAssetDataflowNodes/Public/ChaosClothAsset/SimulationXPBDBendingSpringConfigNode.h
+++ b/Engine/Plugins/Experimental/ChaosClothAssetEditor/Source/ChaosClothAssetDataflowNodes/Public/ChaosClothAsset/SimulationXPBDBendingSpringConfigNode.h
@@ -35,5 +35,5 @@ public:
 	FChaosClothAssetSimulationXPBDBendingSpringConfigNode(const Dataflow::FNodeParameters& InParam, FGuid InGuid = FGuid::NewGuid());
 
 private:
-	virtual void AddProperties(::Chaos::Softs::FCollectionPropertyMutableFacade& Properties) const override;
+	virtual void AddProperties(Dataflow::FContext& Context, ::Chaos::Softs::FCollectionPropertyMutableFacade& Properties) const override;
 };
diff --git a/Engine/Plugins/Experimental/ChaosClothAssetEditor/Source/ChaosClothAssetDataflowNodes/Public/ChaosClothAsset/SimulationXPBDEdgeSpringConfigNode.h b/Engine/Plugins/Experimental/ChaosClothAssetEditor/Source/ChaosClothAssetDataflowNodes/Public/ChaosClothAsset/SimulationXPBDEdgeSpringConfigNode.h
--- a/Engine/Plugins/Experimental/ChaosClothAssetEditor/Source/ChaosClothAssetDataflowNodes/Public/ChaosClothAsset/SimulationXPBDEdgeSpringConfigNode.h
+++ b/Engine/Plugins/Experimental/ChaosClothAssetEditor/Source/ChaosClothAssetDataflowNodes/Public/ChaosClothAsset/SimulationXPBDEdgeSpringConfigNode.h
@@ -35,5 +35,5 @@ public:
 	FChaosClothAssetSimulationXPBDEdgeSpringConfigNode(const Dataflow::FNodeParameters& InParam, FGuid InGuid = FGuid::NewGuid());
 
 private:
-	virtual void AddProperties(::Chaos::Softs::FCollectionPropertyMutableFacade& Properties) const override;
+	virtual void AddProperties(Dataflow::FContext& Context, ::Chaos::Softs::FCollectionPropertyMutableFacade& Properties) const override;
 };
diff --git a/Engine/Plugins/Experimental/ChaosClothAssetEditor/Source/ChaosClothAssetDataflowNodes/Public/ChaosClothAsset/WeightedValue.h b/Engine/Plugins/Experimental/ChaosClothAssetEditor/Source/ChaosClothAssetDataflowNodes/Public/ChaosClothAsset/WeightedValue.h
--- a/Engine/Plugins/Experimental/ChaosClothAssetEditor/Source/ChaosClothAssetDataflowNodes/Public/ChaosClothAsset/WeightedValue.h
+++ b/Engine/Plugins/Experimental/ChaosClothAssetEditor/Source/ChaosClothAssetDataflowNodes/Public/ChaosClothAsset/WeightedValue.h
@@ -38,8 +38,13 @@ struct FChaosClothAssetWeightedValue
 	float High = 1.f;
 
 	/** The name of the weight map for this property. */
-	UPROPERTY(EditAnywhere, Category = "Weighted Value")
+	UPROPERTY(EditAnywhere, Category = "Weighted Value", Meta = (DataflowInput))
 	FString WeightMap = TEXT("WeightMap");
+
+	/** The weight map override value for when the WeightMap has a connection that replaces the provided weight map value. */
+	UE_DEPRECATED(5.3, "This property is experimental and might be removed without notice.")
+	UPROPERTY(VisibleAnywhere, Transient, Category = "Weighted Value")
+	mutable FString WeightMap_Override;  // _Override has a special meaning to the property customization, mutable because this property is set while getting the original value
 };
 
 USTRUCT()
@@ -70,6 +75,11 @@ struct FChaosClothAssetWeightedValueNonAnimatable
 	float High = 1.f;
 
 	/** The name of the weight map for this property. */
-	UPROPERTY(EditAnywhere, Category = "Weighted Value")
+	UPROPERTY(EditAnywhere, Category = "Weighted Value", Meta = (DataflowInput))
 	FString WeightMap = TEXT("WeightMap");
+
+	/** The weight map override value for when the WeightMap has a connection that replaces the provided weight map value. */
+	UE_DEPRECATED(5.3, "This property is experimental and might be removed without notice.")
+	UPROPERTY(VisibleAnywhere, Transient, Category = "Weighted Value")
+	mutable FString WeightMap_Override;  // _Override has a special meaning to the property customization, mutable because this property is set while getting the original value
 };
