Chaos Cloth - Fixed the skinning/simulation vertex contributions on the render mesh by passing the correct MaxDistance weight map.

#rnx

diff --git a/Engine/Plugins/Experimental/ChaosClothAssetEditor/Source/ChaosClothAssetTools/Private/ChaosClothAsset/ClothAssetBuilderEditor.cpp b/Engine/Plugins/Experimental/ChaosClothAssetEditor/Source/ChaosClothAssetTools/Private/ChaosClothAsset/ClothAssetBuilderEditor.cpp
--- a/Engine/Plugins/Experimental/ChaosClothAssetEditor/Source/ChaosClothAssetTools/Private/ChaosClothAsset/ClothAssetBuilderEditor.cpp
+++ b/Engine/Plugins/Experimental/ChaosClothAssetEditor/Source/ChaosClothAssetTools/Private/ChaosClothAsset/ClothAssetBuilderEditor.cpp
@@ -3,6 +3,7 @@
 #include "ChaosClothAsset/ClothAssetBuilderEditor.h"
 #include "ChaosClothAsset/ClothAsset.h"
 #include "ChaosClothAsset/CollectionClothFacade.h"
+#include "Chaos/CollectionPropertyFacade.h"
 #include "Engine/SkeletalMesh.h"
 #include "Engine/SkinnedAssetCommon.h"
 #include "Modules/ModuleManager.h"
@@ -17,6 +18,7 @@
 void UClothAssetBuilderEditor::BuildLod(FSkeletalMeshLODModel& LODModel, const UChaosClothAsset& ClothAsset, int32 LodIndex) const
 {
 	using namespace UE::Chaos::ClothAsset;
+	using namespace ::Chaos::Softs;
 
 	// Start from an empty LODModel
 	LODModel.Empty();
@@ -71,13 +73,19 @@ void UClothAssetBuilderEditor::BuildLod(FSkeletalMeshLODModel& LODModel, const U
 
 	const int32 NumLodSimVertices = GetNumVertices(ClothAsset, LodIndex);
 
-	// Retrieve the MaxDistance map
-	FPointWeightMap MaxDistances;
-	MaxDistances.Initialize(NumLodSimVertices);
-	for (int32 Index = 0; Index < NumLodSimVertices; ++Index)
-	{
-		MaxDistances[Index] = 200.f;
-	}
+	// Retrieve MaxDistance information (weight map and Low/High values)
+	FCollectionPropertyConstFacade Properties(ClothCollection);
+
+	int32 MaxDistancePropertyKeyIndex;
+	FString MaxDistanceString = TEXT("MaxDistance");
+	MaxDistanceString = Properties.GetStringValue(MaxDistanceString, MaxDistanceString, &MaxDistancePropertyKeyIndex);
+	const float MaxDistanceBase = (MaxDistancePropertyKeyIndex != INDEX_NONE) ? Properties.GetLowValue<float>(MaxDistancePropertyKeyIndex) : 0.f;
+	const float MaxDistanceRange = (MaxDistancePropertyKeyIndex != INDEX_NONE) ? Properties.GetHighValue<float>(MaxDistancePropertyKeyIndex) - MaxDistanceBase : 1.f;
+	const TConstArrayView<float> MaxDistanceWeightMap = ClothFacade.GetWeightMap(FName(MaxDistanceString));
+
+	const FPointWeightMap MaxDistances = (MaxDistanceWeightMap.Num() == NumLodSimVertices) ?
+		FPointWeightMap(MaxDistanceWeightMap) : 
+		FPointWeightMap(TNumericLimits<float>::Max(), NumLodSimVertices);
 
 	const int32 NumRenderVertices = ClothFacade.GetNumRenderVertices();
 	LODModel.MeshToImportVertexMap.Reserve(NumRenderVertices);
diff --git a/Engine/Source/Runtime/ClothingSystemRuntimeCommon/Public/PointWeightMap.h b/Engine/Source/Runtime/ClothingSystemRuntimeCommon/Public/PointWeightMap.h
--- a/Engine/Source/Runtime/ClothingSystemRuntimeCommon/Public/PointWeightMap.h
+++ b/Engine/Source/Runtime/ClothingSystemRuntimeCommon/Public/PointWeightMap.h
@@ -34,6 +34,25 @@ struct FPointWeightMap
 #endif
 	{}
 
+	explicit FPointWeightMap(int32 NumPoints, float Value = 0.f)
+#if WITH_EDITORONLY_DATA
+		: Name(NAME_None)
+		, CurrentTarget((uint8)EWeightMapTargetCommon::None)
+		, bEnabled(false)
+#endif
+	{
+		Values.Init(Value, NumPoints);
+	}
+
+	explicit FPointWeightMap(const TConstArrayView<float>& InValues)
+		: Values(InValues)
+#if WITH_EDITORONLY_DATA
+		, Name(NAME_None)
+		, CurrentTarget((uint8)EWeightMapTargetCommon::None)
+		, bEnabled(false)
+#endif
+	{}
+
 	~FPointWeightMap()
 	{}
 
