Chaos Cloth - Fixed crash on empty MaxDistance maps.

#rnx

diff --git a/Engine/Plugins/Experimental/ChaosClothAssetEditor/Source/ChaosClothAssetTools/Private/ChaosClothAsset/ClothAssetBuilderEditor.cpp b/Engine/Plugins/Experimental/ChaosClothAssetEditor/Source/ChaosClothAssetTools/Private/ChaosClothAsset/ClothAssetBuilderEditor.cpp
--- a/Engine/Plugins/Experimental/ChaosClothAssetEditor/Source/ChaosClothAssetTools/Private/ChaosClothAsset/ClothAssetBuilderEditor.cpp
+++ b/Engine/Plugins/Experimental/ChaosClothAssetEditor/Source/ChaosClothAssetTools/Private/ChaosClothAsset/ClothAssetBuilderEditor.cpp
@@ -85,7 +85,7 @@ void UClothAssetBuilderEditor::BuildLod(FSkeletalMeshLODModel& LODModel, const U
 
 	const FPointWeightMap MaxDistances = (MaxDistanceWeightMap.Num() == NumLodSimVertices) ?
 		FPointWeightMap(MaxDistanceWeightMap) : 
-		FPointWeightMap(TNumericLimits<float>::Max(), NumLodSimVertices);
+		FPointWeightMap(NumLodSimVertices, TNumericLimits<float>::Max());
 
 	const int32 NumRenderVertices = ClothFacade.GetNumRenderVertices();
 	LODModel.MeshToImportVertexMap.Reserve(NumRenderVertices);
