CIS - Fixed declaration shadows a field of 'UE::Chaos::ClothAsset::FChaosClothAssetEditorToolkit'

#rnx

diff --git a/Engine/Plugins/Experimental/ChaosClothAssetEditor/Source/ChaosClothAssetEditor/Private/ChaosClothAsset/ClothEditorToolkit.cpp b/Engine/Plugins/Experimental/ChaosClothAssetEditor/Source/ChaosClothAssetEditor/Private/ChaosClothAsset/ClothEditorToolkit.cpp
--- a/Engine/Plugins/Experimental/ChaosClothAssetEditor/Source/ChaosClothAssetEditor/Private/ChaosClothAsset/ClothEditorToolkit.cpp
+++ b/Engine/Plugins/Experimental/ChaosClothAssetEditor/Source/ChaosClothAssetEditor/Private/ChaosClothAsset/ClothEditorToolkit.cpp
@@ -895,11 +895,11 @@ void FChaosClothAssetEditorToolkit::OnNodeTitleCommitted(const FText& InNewText,
 
 void FChaosClothAssetEditorToolkit::OnNodeSelectionChanged(const TSet<UObject*>& NewSelection)
 {
-	auto GetClothCollectionIfPossible = [](const TSharedPtr<const FDataflowNode> DataflowNode, const TSharedPtr<Dataflow::FEngineContext> Context) -> TSharedPtr<FManagedArrayCollection>
+	auto GetClothCollectionIfPossible = [](const TSharedPtr<const FDataflowNode> InDataflowNode, const TSharedPtr<Dataflow::FEngineContext> Context) -> TSharedPtr<FManagedArrayCollection>
 	{
 		if (Context.IsValid())
 		{
-			for (const FDataflowOutput* const Output : DataflowNode->GetOutputs())
+			for (const FDataflowOutput* const Output : InDataflowNode->GetOutputs())
 			{
 				if (Output->GetType() == FName("FManagedArrayCollection"))
 				{
