CIS - Fixed warning: 'WeightMap_Override' is deprecated.

#rnx

diff --git a/Engine/Plugins/Experimental/ChaosClothAssetEditor/Source/ChaosClothAssetDataflowNodes/Private/ChaosClothAsset/SimulationBaseConfigNodePrivate.h b/Engine/Plugins/Experimental/ChaosClothAssetEditor/Source/ChaosClothAssetDataflowNodes/Private/ChaosClothAsset/SimulationBaseConfigNodePrivate.h
--- a/Engine/Plugins/Experimental/ChaosClothAssetEditor/Source/ChaosClothAssetDataflowNodes/Private/ChaosClothAsset/SimulationBaseConfigNodePrivate.h
+++ b/Engine/Plugins/Experimental/ChaosClothAssetEditor/Source/ChaosClothAssetDataflowNodes/Private/ChaosClothAsset/SimulationBaseConfigNodePrivate.h
@@ -24,9 +24,7 @@
 	const int32 Property##KeyIndex = AddPropertyHelper(Properties, FName(TEXT(#Property)), Property.bIsAnimatable); \
 	Properties.SetWeightedValue(Property##KeyIndex, Property.Low, Property.High); \
 	Properties.SetStringValue(Property##KeyIndex, GetValue<FString>(Context, &Property.WeightMap)); \
-PRAGMA_DISABLE_DEPRECATION_WARNINGS \
-	Property.WeightMap_Override = GetValue<FString>(Context, &Property.WeightMap, FString()); \
-PRAGMA_ENABLE_DEPRECATION_WARNINGS
+	Property.WeightMap_Override = GetValue<FString>(Context, &Property.WeightMap, FString());
 
 /** Macro for adding and setting weighted value properties from inside AddProperties() and check for a similar property being overriden. */
 #define UE_CHAOS_CLOTHASSET_SIMULATIONCONFIG_SETPROPERTYCHECKED1(Property, SimilarProperty1) \
@@ -52,9 +50,7 @@ PRAGMA_ENABLE_DEPRECATION_WARNINGS
 		{ FName(TEXT(#SimilarProperty1)) }); \
 	Properties.SetWeightedValue(Property##KeyIndex, Property.Low, Property.High); \
 	Properties.SetStringValue(Property##KeyIndex, GetValue<FString>(Context, &Property.WeightMap)); \
-PRAGMA_DISABLE_DEPRECATION_WARNINGS \
-	Property.WeightMap_Override = GetValue<FString>(Context, &Property.WeightMap, FString()); \
-PRAGMA_ENABLE_DEPRECATION_WARNINGS
+	Property.WeightMap_Override = GetValue<FString>(Context, &Property.WeightMap, FString());
 
 /** Macro for adding and setting weighted value properties from inside AddProperties() and check for two similar properties being overriden. */
 #define UE_CHAOS_CLOTHASSET_SIMULATIONCONFIG_SETPROPERTYWEIGHTEDCHECKED2(Property, SimilarProperty1, SimilarProperty2) \
@@ -62,9 +58,7 @@ PRAGMA_ENABLE_DEPRECATION_WARNINGS
 		{ FName(TEXT(#SimilarProperty1)), FName(TEXT(#SimilarProperty2)) }); \
 	Properties.SetWeightedValue(Property##KeyIndex, Property.Low, Property.High); \
 	Properties.SetStringValue(Property##KeyIndex, GetValue<FString>(Context, &Property.WeightMap)); \
-PRAGMA_DISABLE_DEPRECATION_WARNINGS \
-	Property.WeightMap_Override = GetValue<FString>(Context, &Property.WeightMap, FString()); \
-PRAGMA_ENABLE_DEPRECATION_WARNINGS
+	Property.WeightMap_Override = GetValue<FString>(Context, &Property.WeightMap, FString());
 
 /** Macro for adding and setting weighted value properties from inside AddProperties() and check for three similar properties being overriden. */
 #define UE_CHAOS_CLOTHASSET_SIMULATIONCONFIG_SETPROPERTYWEIGHTEDCHECKED3(Property, SimilarProperty1, SimilarProperty2, SimilarProperty3) \
@@ -72,9 +66,7 @@ PRAGMA_ENABLE_DEPRECATION_WARNINGS
 		{ FName(TEXT(#SimilarProperty1)), FName(TEXT(#SimilarProperty2)), FName(TEXT(#SimilarProperty3)) }); \
 	Properties.SetWeightedValue(Property##KeyIndex, Property.Low, Property.High); \
 	Properties.SetStringValue(Property##KeyIndex, GetValue<FString>(Context, &Property.WeightMap)); \
-PRAGMA_DISABLE_DEPRECATION_WARNINGS \
-	Property.WeightMap_Override = GetValue<FString>(Context, &Property.WeightMap, FString()); \
-PRAGMA_ENABLE_DEPRECATION_WARNINGS
+	Property.WeightMap_Override = GetValue<FString>(Context, &Property.WeightMap, FString());
 
 /** Macro for adding and setting weighted value properties from inside AddProperties() and check for four similar properties being overriden. */
 #define UE_CHAOS_CLOTHASSET_SIMULATIONCONFIG_SETPROPERTYWEIGHTEDCHECKED4(Property, SimilarProperty1, SimilarProperty2, SimilarProperty3, SimilarProperty4) \
@@ -82,6 +74,4 @@ PRAGMA_ENABLE_DEPRECATION_WARNINGS
 		{ FName(TEXT(#SimilarProperty1)), FName(TEXT(#SimilarProperty2)), FName(TEXT(#SimilarProperty3)), FName(TEXT(#SimilarProperty4)) }); \
 	Properties.SetWeightedValue(Property##KeyIndex, Property.Low, Property.High); \
 	Properties.SetStringValue(Property##KeyIndex, GetValue<FString>(Context, &Property.WeightMap)); \
-PRAGMA_DISABLE_DEPRECATION_WARNINGS \
-	Property.WeightMap_Override = GetValue<FString>(Context, &Property.WeightMap, FString()); \
-PRAGMA_ENABLE_DEPRECATION_WARNINGS
+	Property.WeightMap_Override = GetValue<FString>(Context, &Property.WeightMap, FString());
diff --git a/Engine/Plugins/Experimental/ChaosClothAssetEditor/Source/ChaosClothAssetDataflowNodes/Public/ChaosClothAsset/WeightedValue.h b/Engine/Plugins/Experimental/ChaosClothAssetEditor/Source/ChaosClothAssetDataflowNodes/Public/ChaosClothAsset/WeightedValue.h
--- a/Engine/Plugins/Experimental/ChaosClothAssetEditor/Source/ChaosClothAssetDataflowNodes/Public/ChaosClothAsset/WeightedValue.h
+++ b/Engine/Plugins/Experimental/ChaosClothAssetEditor/Source/ChaosClothAssetDataflowNodes/Public/ChaosClothAsset/WeightedValue.h
@@ -42,7 +42,6 @@ struct FChaosClothAssetWeightedValue
 	FString WeightMap = TEXT("WeightMap");
 
 	/** The weight map override value for when the WeightMap has a connection that replaces the provided weight map value. */
-	UE_DEPRECATED(5.3, "This property is experimental and might be removed without notice.")
 	UPROPERTY(VisibleAnywhere, Transient, Category = "Weighted Value")
 	mutable FString WeightMap_Override;  // _Override has a special meaning to the property customization, mutable because this property is set while getting the original value
 };
@@ -79,7 +78,6 @@ struct FChaosClothAssetWeightedValueNonAnimatable
 	FString WeightMap = TEXT("WeightMap");
 
 	/** The weight map override value for when the WeightMap has a connection that replaces the provided weight map value. */
-	UE_DEPRECATED(5.3, "This property is experimental and might be removed without notice.")
 	UPROPERTY(VisibleAnywhere, Transient, Category = "Weighted Value")
 	mutable FString WeightMap_Override;  // _Override has a special meaning to the property customization, mutable because this property is set while getting the original value
 };
