Chaos Cloth Editor: pass which nodes are currently selected to FDataflowEditorCommands::OnPropertyValueChanged from the GraphEditor. This will correctly dirty the graph so we only re-evaluate the nodes AFTER the dirtied ones, instead of all of them.

diff --git a/Engine/Plugins/Experimental/ChaosClothAssetEditor/Source/ChaosClothAssetEditor/Private/ChaosClothAsset/ClothEditorToolkit.cpp b/Engine/Plugins/Experimental/ChaosClothAssetEditor/Source/ChaosClothAssetEditor/Private/ChaosClothAsset/ClothEditorToolkit.cpp
--- a/Engine/Plugins/Experimental/ChaosClothAssetEditor/Source/ChaosClothAssetEditor/Private/ChaosClothAsset/ClothEditorToolkit.cpp
+++ b/Engine/Plugins/Experimental/ChaosClothAssetEditor/Source/ChaosClothAssetEditor/Private/ChaosClothAsset/ClothEditorToolkit.cpp
@@ -880,7 +880,7 @@ TSharedPtr<IStructureDetailsView> FChaosClothAssetEditorToolkit::CreateNodeDetai
 
 void FChaosClothAssetEditorToolkit::OnPropertyValueChanged(const FPropertyChangedEvent& PropertyChangedEvent)
 {
-	FDataflowEditorCommands::OnPropertyValueChanged(Dataflow, DataflowContext, LastDataflowNodeTimestamp, PropertyChangedEvent);
+	FDataflowEditorCommands::OnPropertyValueChanged(Dataflow, DataflowContext, LastDataflowNodeTimestamp, PropertyChangedEvent, GraphEditor ? GraphEditor->GetSelectedNodes() : FGraphPanelSelectionSet());
 }
 
 bool FChaosClothAssetEditorToolkit::OnNodeVerifyTitleCommit(const FText& NewText, UEdGraphNode* GraphNode, FText& OutErrorMessage) const
