Add an optional post-processing smoothing to the transfer skin weights node.
Add an internal weight map indicating which vertices had their weights directly copied from source and which were automatically computed.

#rnx

diff --git a/Engine/Plugins/Experimental/ChaosClothAssetEditor/Source/ChaosClothAssetDataflowNodes/Private/ChaosClothAsset/TransferSkinWeightsNode.cpp b/Engine/Plugins/Experimental/ChaosClothAssetEditor/Source/ChaosClothAssetDataflowNodes/Private/ChaosClothAsset/TransferSkinWeightsNode.cpp
--- a/Engine/Plugins/Experimental/ChaosClothAssetEditor/Source/ChaosClothAssetDataflowNodes/Private/ChaosClothAsset/TransferSkinWeightsNode.cpp
+++ b/Engine/Plugins/Experimental/ChaosClothAssetEditor/Source/ChaosClothAssetDataflowNodes/Private/ChaosClothAsset/TransferSkinWeightsNode.cpp
@@ -66,10 +66,12 @@ namespace UE::Chaos::ClothAsset::Private
 		return true;
 	}
 
-	static void TransferInpaintWeights(
+	static bool TransferInpaintWeights(
 		const FReferenceSkeleton& TargetRefSkeleton,
 		const double NormalThreshold,
 		const double RadiusPercentage,
+		const int32 NumSmoothingIterations,
+		const double SmoothingStrength,
 		bool bUseParallel,
 		const TSharedRef<FManagedArrayCollection>& ClothCollection,
 		UE::Geometry::FTransferBoneWeights& TransferBoneWeights)
@@ -87,7 +89,7 @@ namespace UE::Chaos::ClothAsset::Private
 		if (!ensure(SimClothToDynamicMesh(ClothCollection, TargetRefSkeleton, WeldedSimMesh)))
 		{
 			UE_LOG(LogChaosClothAssetDataflowNodes, Warning, TEXT("TransferSkinWeightsNode: Failed to weld the simulation mesh for LOD."));
-			return;
+			return false;
 		}
 
 		//
@@ -95,18 +97,25 @@ namespace UE::Chaos::ClothAsset::Private
 		//
 		TransferBoneWeights.NormalThreshold = FMathd::DegToRad * NormalThreshold;
 		TransferBoneWeights.SearchRadius = RadiusPercentage * WeldedSimMesh.GetBounds().DiagonalLength();
+		TransferBoneWeights.NumSmoothingIterations = NumSmoothingIterations;
+		TransferBoneWeights.SmoothingStrength = SmoothingStrength;
+		TransferBoneWeights.LayeredMeshSupport = true; // multilayerd clothing
+
 		if (!ensure(TransferBoneWeights.Validate() == EOperationValidationResult::Ok))
 		{
 			UE_LOG(LogChaosClothAssetDataflowNodes, Warning, TEXT("TransferSkinWeightsNode: Transfer method parameters are invalid."));
-			return;
+			return false;
 		}
 		if (!ensure(TransferBoneWeights.TransferWeightsToMesh(WeldedSimMesh, FSkeletalMeshAttributes::DefaultSkinWeightProfileName)))
 		{
 			UE_LOG(LogChaosClothAssetDataflowNodes, Warning, TEXT("TransferSkinWeightsNode: Transferring skin weights failed."));
-			return;
+			return false;
 		}
 
 		UE::Chaos::ClothAsset::FCollectionClothFacade ClothFacade(ClothCollection);
+		const FName InpaintWeightMaskName = TEXT("_InpaintWeightMask");
+		ClothFacade.AddWeightMap(InpaintWeightMaskName);
+		TArrayView<float> InpaintWeightMask = ClothFacade.GetWeightMap(InpaintWeightMaskName);
 
 		//
 		// Copy the new bone weight data from the welded sim mesh back to the cloth patterns.
@@ -122,7 +131,7 @@ namespace UE::Chaos::ClothAsset::Private
 				SimMeshToDynamicMesh[NonManifoldMapping.GetOriginalNonManifoldVertexID(DynamicMeshVert)].Add(DynamicMeshVert);
 			}
 
-			ParallelFor(ClothFacade.GetNumSimVertices3D(), [&ClothFacade, &WeldedSimMesh, &SimMeshToDynamicMesh](int32 SimVertexID)
+			ParallelFor(ClothFacade.GetNumSimVertices3D(), [&ClothFacade, &WeldedSimMesh, &SimMeshToDynamicMesh, &TransferBoneWeights, &InpaintWeightMask](int32 SimVertexID)
 			{
 				const FDynamicMeshVertexSkinWeightsAttribute* const OutAttribute = WeldedSimMesh.Attributes()->GetSkinWeightsAttribute(FSkeletalMeshAttributes::DefaultSkinWeightProfileName);
 				checkSlow(OutAttribute);
@@ -139,9 +148,13 @@ namespace UE::Chaos::ClothAsset::Private
 					checkSlow(WeldedSimMesh.IsVertex(WeldedID));
 					OutAttribute->GetValue(WeldedID, ClothFacade.GetSimBoneIndices()[SimVertexID],
 						ClothFacade.GetSimBoneWeights()[SimVertexID]);
+					InpaintWeightMask[SimVertexID] = TransferBoneWeights.MatchedVertices[WeldedID] ? 1.0f : 0.0f;
 				}
 				else
 				{
+					// TODO: The logic below should be moved to BoneWeights.h and implemented as Merge method which takes
+					// an array of bone weights and computes their average.
+					
 					// Need to merge data because dynamic mesh split the original vertex
 					TMap<int32, TPair<float, int32>> CombinedData;
 					for (const int32 WeldedID : SimMeshToDynamicMesh[SimVertexID])
@@ -199,12 +212,16 @@ namespace UE::Chaos::ClothAsset::Private
 					{
 						Weight *= WeightsSumRecip;
 					}
+
+					const int32 WeldedID = SimMeshToDynamicMesh[SimVertexID][0]; // Any welded id can be used here
+					InpaintWeightMask[SimVertexID] = TransferBoneWeights.MatchedVertices[WeldedID] ? 1.0f : 0.0f;
 				}
 			}, bUseParallel ? EParallelForFlags::None : EParallelForFlags::ForceSingleThread);
 		}
 		else
 		{
-			ParallelFor(WeldedSimMesh.MaxVertexID(), [&ClothFacade, &WeldedSimMesh, &NonManifoldMapping](int32 WeldedID)
+
+			ParallelFor(WeldedSimMesh.MaxVertexID(), [&ClothFacade, &WeldedSimMesh, &NonManifoldMapping, &InpaintWeightMask, &TransferBoneWeights](int32 WeldedID)
 			{
 				const FDynamicMeshVertexSkinWeightsAttribute* OutAttribute = WeldedSimMesh.Attributes()->GetSkinWeightsAttribute(FSkeletalMeshAttributes::DefaultSkinWeightProfileName);
 
@@ -213,6 +230,7 @@ namespace UE::Chaos::ClothAsset::Private
 				checkSlow(WeldedID < ClothFacade.GetNumSimVertices3D());
 				OutAttribute->GetValue(WeldedID, ClothFacade.GetSimBoneIndices()[WeldedID],
 					ClothFacade.GetSimBoneWeights()[WeldedID]);
+				InpaintWeightMask[WeldedID] = TransferBoneWeights.MatchedVertices[WeldedID] ? 1.0f : 0.0f;
 			}, bUseParallel ? EParallelForFlags::None : EParallelForFlags::ForceSingleThread);
 		}
 
@@ -231,9 +249,11 @@ namespace UE::Chaos::ClothAsset::Private
 				ClothFacade.GetRenderPosition()[VertexID]);
 
 		}, bUseParallel ? EParallelForFlags::None : EParallelForFlags::ForceSingleThread);
+
+		return true;
 	}
 
-	static void TransferClosestPointOnSurface(
+	static bool TransferClosestPointOnSurface(
 		const FReferenceSkeleton& TargetRefSkeleton,
 		const bool bUseParallel,
 		UE::Chaos::ClothAsset::FCollectionClothFacade& ClothFacade,
@@ -254,7 +274,7 @@ namespace UE::Chaos::ClothAsset::Private
 		if (!ensure(TransferBoneWeights.Validate() == EOperationValidationResult::Ok))
 		{
 			UE_LOG(LogChaosClothAssetDataflowNodes, Warning, TEXT("TransferSkinWeightsNode: Transfer method parameters are invalid."));
-			return;
+			return false;
 		}
 
 		//
@@ -280,6 +300,8 @@ namespace UE::Chaos::ClothAsset::Private
 				&TargetBoneToIndex);
 
 		}, bUseParallel ? EParallelForFlags::None : EParallelForFlags::ForceSingleThread);
+
+		return true;
 	}
 }
 
@@ -340,18 +362,29 @@ void FChaosClothAssetTransferSkinWeightsNode::Evaluate(Dataflow::FContext& Conte
 			FCollectionClothFacade ClothFacade(ClothCollection);
 			ClothFacade.SetSkeletonAssetPathName(SkeletalMesh->GetSkeleton()->GetPathName());
 
+			bool bTransferResult = false;
 			if (TransferMethod == EChaosClothAssetTransferSkinWeightsMethod::InpaintWeights)
 			{
-				TransferInpaintWeights(TargetRefSkeleton, NormalThreshold, RadiusPercentage, bUseParallel, ClothCollection, TransferBoneWeights);
+				bTransferResult = TransferInpaintWeights(TargetRefSkeleton, NormalThreshold, RadiusPercentage, NumSmoothingIterations, SmoothingStrength, bUseParallel, ClothCollection, TransferBoneWeights);
 			}
 			else if (TransferMethod == EChaosClothAssetTransferSkinWeightsMethod::ClosestPointOnSurface)
 			{
-				TransferClosestPointOnSurface(TargetRefSkeleton, bUseParallel, ClothFacade, TransferBoneWeights);
+				bTransferResult = TransferClosestPointOnSurface(TargetRefSkeleton, bUseParallel, ClothFacade, TransferBoneWeights);
 			}
 			else
 			{
 				checkNoEntry();
 			}
+
+			if (!bTransferResult)
+			{
+				FClothDataflowTools::LogAndToastWarning(*this,
+					LOCTEXT("TransferWeightsFailedHeadline", "Transfer Weights Failed."),
+					LOCTEXT("TransferWeightsDetails", "Failed to transfer skinning weights from the source."));
+				SetValue(Context, MoveTemp(*ClothCollection), &Collection);
+				return;
+			}
+
 		}
 
 		SetValue(Context, MoveTemp(*ClothCollection), &Collection);
diff --git a/Engine/Plugins/Experimental/ChaosClothAssetEditor/Source/ChaosClothAssetDataflowNodes/Public/ChaosClothAsset/TransferSkinWeightsNode.h b/Engine/Plugins/Experimental/ChaosClothAssetEditor/Source/ChaosClothAssetDataflowNodes/Public/ChaosClothAsset/TransferSkinWeightsNode.h
--- a/Engine/Plugins/Experimental/ChaosClothAssetEditor/Source/ChaosClothAssetDataflowNodes/Public/ChaosClothAsset/TransferSkinWeightsNode.h
+++ b/Engine/Plugins/Experimental/ChaosClothAssetEditor/Source/ChaosClothAssetDataflowNodes/Public/ChaosClothAsset/TransferSkinWeightsNode.h
@@ -15,9 +15,9 @@ enum class EChaosClothAssetTransferSkinWeightsMethod : uint8
 	ClosestPointOnSurface,
 	
 	/**
-	 * For every vertex on the target mesh, find the closest vertex on the surface of the source mesh.
-	 * If that vertex position is within the search radius, and their normals differ by less than the specified normal threshold,
-	 * then the vertex weights are directly copied from the source vertex to the target mesh vertex.
+	 * For every vertex on the target mesh, find the closest point on the surface of the source mesh.
+	 * If that point position is within the search radius, and their normals differ by less than the specified normal threshold,
+	 * then the vertex weights are directly copied from the source point to the target mesh vertex.
 	 * For all other vertices whose weights didn't get transferred, smoothed weight values are automatically computed.
 	 */
 	InpaintWeights
@@ -65,6 +65,14 @@ public:
 	UPROPERTY(EditAnywhere, Category = "Transfer Skin Weights|Transfer Method", Meta = (UIMin = -1, UIMax = 180, ClampMin = -1, ClampMax = 180, EditCondition="TransferMethod==EChaosClothAssetTransferSkinWeightsMethod::InpaintWeights"))
 	double NormalThreshold = 30;
 
+	/** The number of smoothing iterations applied to the vertices whose weights were automatically computed. */
+	UPROPERTY(EditAnywhere, Category = "Transfer Skin Weights|Transfer Method", Meta = (UIMin = 0, UIMax = 100, ClampMin = 0, ClampMax = 100, DisplayName = "Smoothing Iterations", EditCondition="TransferMethod==EChaosClothAssetTransferSkinWeightsMethod::InpaintWeights"))
+	int32 NumSmoothingIterations = 10;
+
+	/** The smoothing strength of each smoothing iteration. */
+	UPROPERTY(EditAnywhere, Category = "Transfer Skin Weights|Transfer Method", Meta = (UIMin = 0, UIMax = 1, ClampMin = 0, ClampMax = 1, DisplayName = "Smoothing Strength", EditCondition="TransferMethod==EChaosClothAssetTransferSkinWeightsMethod::InpaintWeights"))
+	float SmoothingStrength = 0.1;
+
 	FChaosClothAssetTransferSkinWeightsNode(const Dataflow::FNodeParameters& InParam, FGuid InGuid = FGuid::NewGuid());
 
 private:
