Chaos Cloth Asset: Fix XPBDAnisoBendingConfig inputs.

diff --git a/Engine/Plugins/Experimental/ChaosCloth/Source/Chaos/Private/ChaosCloth/ChaosClothConstraints.cpp b/Engine/Plugins/Experimental/ChaosCloth/Source/Chaos/Private/ChaosCloth/ChaosClothConstraints.cpp
--- a/Engine/Plugins/Experimental/ChaosCloth/Source/Chaos/Private/ChaosCloth/ChaosClothConstraints.cpp
+++ b/Engine/Plugins/Experimental/ChaosCloth/Source/Chaos/Private/ChaosCloth/ChaosClothConstraints.cpp
@@ -1193,37 +1193,45 @@ void FClothConstraints::Update(
 	Softs::FSolverReal MeshScale,
 	Softs::FSolverReal MaxDistancesScale)
 {
+	if (EdgeConstraints)
+	{
+		EdgeConstraints->SetProperties(ConfigProperties, WeightMaps);
+	}
 	if (XEdgeConstraints)
 	{
 		XEdgeConstraints->SetProperties(ConfigProperties, WeightMaps);
 	}
-	else if (EdgeConstraints)
+	if (XStretchBiasConstraints)
 	{
-		EdgeConstraints->SetProperties(ConfigProperties, WeightMaps);
+		XStretchBiasConstraints->SetProperties(ConfigProperties, WeightMaps);
+	}
+	if (BendingConstraints)
+	{
+		BendingConstraints->SetProperties(ConfigProperties, WeightMaps);
 	}
 	if (XBendingConstraints)
 	{
 		XBendingConstraints->SetProperties(ConfigProperties, WeightMaps);
 	}
-	else if (BendingConstraints)
+	if (BendingElementConstraints)
 	{
-		BendingConstraints->SetProperties(ConfigProperties, WeightMaps);
+		BendingElementConstraints->SetProperties(ConfigProperties, WeightMaps);
 	}
 	if (XBendingElementConstraints)
 	{
 		XBendingElementConstraints->SetProperties(ConfigProperties, WeightMaps);
 	}
-	else if (BendingElementConstraints)
+	if (XAnisoBendingElementConstraints)
 	{
-		BendingElementConstraints->SetProperties(ConfigProperties, WeightMaps);
+		XAnisoBendingElementConstraints->SetProperties(ConfigProperties, WeightMaps);
 	}
-	if (XAreaConstraints)
+	if (AreaConstraints)
 	{
-		XAreaConstraints->SetProperties(ConfigProperties, WeightMaps);
+		AreaConstraints->SetProperties(ConfigProperties, WeightMaps);
 	}
-	else if (AreaConstraints)
+	if (XAreaConstraints)
 	{
-		AreaConstraints->SetProperties(ConfigProperties, WeightMaps);
+		XAreaConstraints->SetProperties(ConfigProperties, WeightMaps);
 	}
 	if (LongRangeConstraints)
 	{
@@ -1233,6 +1241,10 @@ void FClothConstraints::Update(
 	{
 		MaximumDistanceConstraints->SetProperties(ConfigProperties, WeightMaps, MeshScale * MaxDistancesScale);
 	}
+	if (BackstopConstraints)
+	{
+		BackstopConstraints->SetProperties(ConfigProperties, WeightMaps, MeshScale);
+	}
 	if (AnimDriveConstraints)
 	{
 		AnimDriveConstraints->SetProperties(ConfigProperties, WeightMaps);
@@ -1241,10 +1253,6 @@ void FClothConstraints::Update(
 	{
 		SelfCollisionConstraints->SetProperties(ConfigProperties);
 	}
-	if (BackstopConstraints)
-	{
-		BackstopConstraints->SetProperties(ConfigProperties, WeightMaps, MeshScale);
-	}
 }
 
 // Deprecated
diff --git a/Engine/Plugins/Experimental/ChaosClothAssetEditor/Source/ChaosClothAssetDataflowNodes/Private/ChaosClothAsset/SimulationXPBDAnisoBendingConfigNode.cpp b/Engine/Plugins/Experimental/ChaosClothAssetEditor/Source/ChaosClothAssetDataflowNodes/Private/ChaosClothAsset/SimulationXPBDAnisoBendingConfigNode.cpp
--- a/Engine/Plugins/Experimental/ChaosClothAssetEditor/Source/ChaosClothAssetDataflowNodes/Private/ChaosClothAsset/SimulationXPBDAnisoBendingConfigNode.cpp
+++ b/Engine/Plugins/Experimental/ChaosClothAssetEditor/Source/ChaosClothAssetDataflowNodes/Private/ChaosClothAsset/SimulationXPBDAnisoBendingConfigNode.cpp
@@ -17,6 +17,8 @@ FChaosClothAssetSimulationXPBDAnisoBendingConfigNode::FChaosClothAssetSimulation
 	RegisterInputConnection(&XPBDAnisoBendingStiffnessBias.WeightMap);
 	RegisterInputConnection(&XPBDAnisoBendingDamping.WeightMap);
 	RegisterInputConnection(&XPBDAnisoBucklingStiffnessWarp.WeightMap);
+	RegisterInputConnection(&XPBDAnisoBucklingStiffnessWeft.WeightMap);
+	RegisterInputConnection(&XPBDAnisoBucklingStiffnessBias.WeightMap);
 }
 
 void FChaosClothAssetSimulationXPBDAnisoBendingConfigNode::AddProperties(Dataflow::FContext& Context, ::Chaos::Softs::FCollectionPropertyMutableFacade& Properties) const
@@ -53,4 +55,6 @@ void FChaosClothAssetSimulationXPBDAnisoBendingConfigNode::AddProperties(Dataflo
 		XPBDAnisoBucklingStiffnessWarp,
 		BucklingStiffness,       // Existing properties to warn against
 		XPBDBucklingStiffness);  //
+	UE_CHAOS_CLOTHASSET_SIMULATIONCONFIG_SETPROPERTYWEIGHTED(XPBDAnisoBucklingStiffnessWeft);
+	UE_CHAOS_CLOTHASSET_SIMULATIONCONFIG_SETPROPERTYWEIGHTED(XPBDAnisoBucklingStiffnessBias);
 }
diff --git a/Engine/Plugins/Experimental/ChaosClothAssetEditor/Source/ChaosClothAssetDataflowNodes/Public/ChaosClothAsset/SimulationXPBDAnisoBendingConfigNode.h b/Engine/Plugins/Experimental/ChaosClothAssetEditor/Source/ChaosClothAssetDataflowNodes/Public/ChaosClothAsset/SimulationXPBDAnisoBendingConfigNode.h
--- a/Engine/Plugins/Experimental/ChaosClothAssetEditor/Source/ChaosClothAssetDataflowNodes/Public/ChaosClothAsset/SimulationXPBDAnisoBendingConfigNode.h
+++ b/Engine/Plugins/Experimental/ChaosClothAssetEditor/Source/ChaosClothAssetDataflowNodes/Public/ChaosClothAsset/SimulationXPBDAnisoBendingConfigNode.h
@@ -95,7 +95,7 @@ public:
 	 * are interpolated with the per particle weight to make the final value used for the simulation.
 	 * Otherwise all particles are considered to have a zero weight, and only the Low value is meaningful.
 	 */
-	UPROPERTY(EditAnywhere, Category = "XPBDAnisoBending Properties", Meta = (UIMin = "0", UIMax = "10000", ClampMin = "0", ClampMax = "10000000", EditCondition = "XPBDBucklingRatio != 0"))
+	UPROPERTY(EditAnywhere, Category = "XPBDAnisoBending Properties", Meta = (UIMin = "0", UIMax = "10000", ClampMin = "0", ClampMax = "10000000", EditCondition = "XPBDAnisoBucklingRatio != 0"))
 	FChaosClothAssetWeightedValue XPBDAnisoBucklingStiffnessWarp = { true, 50.f, 50.f, TEXT("XPBDAnisoBucklingStiffnessWarp") };
 
 	/**
@@ -107,7 +107,7 @@ public:
 	 * are interpolated with the per particle weight to make the final value used for the simulation.
 	 * Otherwise all particles are considered to have a zero weight, and only the Low value is meaningful.
 	 */
-	UPROPERTY(EditAnywhere, Category = "XPBDAnisoBending Properties", Meta = (UIMin = "0", UIMax = "10000", ClampMin = "0", ClampMax = "10000000", EditCondition = "XPBDBucklingRatio != 0"))
+	UPROPERTY(EditAnywhere, Category = "XPBDAnisoBending Properties", Meta = (UIMin = "0", UIMax = "10000", ClampMin = "0", ClampMax = "10000000", EditCondition = "XPBDAnisoBucklingRatio != 0"))
 	FChaosClothAssetWeightedValue XPBDAnisoBucklingStiffnessWeft = { true, 50.f, 50.f, TEXT("XPBDAnisoBucklingStiffnessWeft") };
 
 	/**
@@ -119,7 +119,7 @@ public:
 	 * are interpolated with the per particle weight to make the final value used for the simulation.
 	 * Otherwise all particles are considered to have a zero weight, and only the Low value is meaningful.
 	 */
-	UPROPERTY(EditAnywhere, Category = "XPBDAnisoBending Properties", Meta = (UIMin = "0", UIMax = "10000", ClampMin = "0", ClampMax = "10000000", EditCondition = "XPBDBucklingRatio != 0"))
+	UPROPERTY(EditAnywhere, Category = "XPBDAnisoBending Properties", Meta = (UIMin = "0", UIMax = "10000", ClampMin = "0", ClampMax = "10000000", EditCondition = "XPBDAnisoBucklingRatio != 0"))
 	FChaosClothAssetWeightedValue XPBDAnisoBucklingStiffnessBias = { true, 50.f, 50.f, TEXT("XPBDAnisoBucklingStiffnessBias") };
 
 	FChaosClothAssetSimulationXPBDAnisoBendingConfigNode(const Dataflow::FNodeParameters& InParam, FGuid InGuid = FGuid::NewGuid());
