Cloth Editor: adding support for Selections. This change adds the following classes:
- Selection Facade, which can represent a set of integer indices, as well as a name and type name for the selection, all stored on an FManagedArrayCollection
- Selection Node, a new Cloth Editor Datasmith node that will use the Selection Facade to add a Selection to an given FManagedArrayCollection
- Selection to Weight Map Node, another Cloth Editor Datasmith node that turns a selection into a 0/1 weight map on the given ClothCollection
- Cloth Mesh Selection Tool, a new Cloth Editor interactive tool that can be used to select vertices or faces, and add the selection to a SelectionNode in the Dataflow graph

diff --git a/Engine/Plugins/Experimental/ChaosClothAsset/Source/ChaosClothAsset/Private/ChaosClothAsset/CollectionClothSelectionFacade.cpp b/Engine/Plugins/Experimental/ChaosClothAsset/Source/ChaosClothAsset/Private/ChaosClothAsset/CollectionClothSelectionFacade.cpp
new file mode 100644
--- /dev/null
+++ b/Engine/Plugins/Experimental/ChaosClothAsset/Source/ChaosClothAsset/Private/ChaosClothAsset/CollectionClothSelectionFacade.cpp
@@ -0,0 +1,170 @@
+// Copyright Epic Games, Inc. All Rights Reserved.
+
+#include "ChaosClothAsset/CollectionClothSelectionFacade.h"
+#include "GeometryCollection/ManagedArrayCollection.h"
+
+#define LOCTEXT_NAMESPACE "FCollectionClothSelectionFacade"
+
+namespace UE::Chaos::ClothAsset
+{
+
+	namespace Private
+	{
+		static const FName SelectionGroup(TEXT("Selection"));
+		static const FName NameAttribute(TEXT("Name"));
+		static const FName TypeAttribute(TEXT("Type"));
+		static const FName IndicesAttribute(TEXT("Indices"));
+	}
+
+	// --------------- FCollectionClothSelectionConstFacade -------------------------------
+
+	FCollectionClothSelectionConstFacade::FCollectionClothSelectionConstFacade(const TSharedRef<const FManagedArrayCollection>& InManagedArrayCollection) :
+		ManagedArrayCollection(ConstCastSharedRef<FManagedArrayCollection>(InManagedArrayCollection))
+	{
+		Name = ManagedArrayCollection->FindAttribute<FString>(Private::NameAttribute, Private::SelectionGroup);
+		Type = ManagedArrayCollection->FindAttribute<FString>(Private::TypeAttribute, Private::SelectionGroup);
+		Indices = ManagedArrayCollection->FindAttribute<TSet<int32>>(Private::IndicesAttribute, Private::SelectionGroup);
+	}
+
+	bool FCollectionClothSelectionConstFacade::IsValid() const
+	{
+		return Name && Type && Indices;
+	}
+
+	int32 FCollectionClothSelectionConstFacade::GetNumSelections() const
+	{
+		return ManagedArrayCollection->NumElements(Private::SelectionGroup);
+	}
+
+	const TArrayView<const FString> FCollectionClothSelectionConstFacade::GetName() const
+	{
+		if (Name)
+		{
+			return TArrayView<const FString>(Name->GetData(), Name->Num());
+		}
+		return TArrayView<const FString>();
+	}
+
+	const TArrayView<const FString> FCollectionClothSelectionConstFacade::GetType() const
+	{
+		if (Type)
+		{
+			return TArrayView<const FString>(Type->GetData(), Type->Num());
+		}
+		return TArrayView<const FString>();
+	}
+
+	const TArrayView<const TSet<int32>> FCollectionClothSelectionConstFacade::GetIndices() const
+	{
+		if (Indices)
+		{
+			return TArrayView<const TSet<int32>>(Indices->GetData(), Indices->Num());
+		}
+		return TArrayView<const TSet<int32>>();
+	}
+
+	int32 FCollectionClothSelectionConstFacade::FindSelection(const FString& SearchName) const
+	{
+		const TArrayView<const FString> Names = GetName();
+		for (int32 SelectionIndex = 0; SelectionIndex < Names.Num(); ++SelectionIndex)
+		{
+			if (Names[SelectionIndex] == SearchName)
+			{
+				return SelectionIndex;
+			}
+		}
+		return INDEX_NONE;
+	}
+
+
+	// --------------- FCollectionClothSelectionFacade -------------------------------
+
+	FCollectionClothSelectionFacade::FCollectionClothSelectionFacade(const TSharedRef<const FManagedArrayCollection>& ManagedArrayCollection) :
+		FCollectionClothSelectionConstFacade(ManagedArrayCollection)
+	{}
+
+	void FCollectionClothSelectionFacade::DefineSchema()
+	{
+		if (!ManagedArrayCollection->HasGroup(Private::SelectionGroup))
+		{
+			ManagedArrayCollection->AddGroup(Private::SelectionGroup);
+		}
+
+		// AddAttribute will only add the attribute if it doesn't already exist, otherwise it will return the existing one
+		Name = &ManagedArrayCollection->AddAttribute<FString>(Private::NameAttribute, Private::SelectionGroup);
+		Type = &ManagedArrayCollection->AddAttribute<FString>(Private::TypeAttribute, Private::SelectionGroup);
+		Indices = &ManagedArrayCollection->AddAttribute<TSet<int32>>(Private::IndicesAttribute, Private::SelectionGroup);
+	}
+
+	TArrayView<FString> FCollectionClothSelectionFacade::GetName()
+	{
+		if (Name)
+		{
+			return TArrayView<FString>(Name->GetData(), Name->Num());
+		}
+		return TArrayView<FString>();
+	}
+
+	TArrayView<FString> FCollectionClothSelectionFacade::GetType()
+	{
+		if (Type)
+		{
+			return TArrayView<FString>(Type->GetData(), Type->Num());
+		}
+		return TArrayView<FString>();
+	}
+
+	TArrayView<TSet<int32>> FCollectionClothSelectionFacade::GetIndices()
+	{
+		if (Indices)
+		{
+			return TArrayView<TSet<int32>>(Indices->GetData(), Indices->Num());
+		}
+		return TArrayView<TSet<int32>>();
+	}
+
+	int32 FCollectionClothSelectionFacade::FindOrAddSelection(const FString& InName)
+	{
+		if (!IsValid())
+		{
+			return INDEX_NONE;
+		}
+
+		// Check for existing selection with matching Name
+		const int32 FoundIndex = FindSelection(InName);
+		if (FoundIndex != INDEX_NONE)
+		{
+			return FoundIndex;
+		}
+
+		// No selection with the given name exists
+		checkf(ManagedArrayCollection->HasGroup(Private::SelectionGroup), TEXT("Expected SelectionGroup to exist in the managed array at this point"));
+
+		const int32 NewIndex = ManagedArrayCollection->AddElements(1, Private::SelectionGroup);
+		
+		GetName()[NewIndex] = InName;
+
+		return NewIndex;
+	}
+
+	bool FCollectionClothSelectionFacade::SetSelection(int32 SelectionIndex, const FString& InType, const TSet<int32>& InSelectedIndices)
+	{
+		if (!IsValid())
+		{
+			return false;
+		}
+
+		if (SelectionIndex < 0 || SelectionIndex >= GetNumSelections())
+		{
+			return false;
+		}
+
+		GetType()[SelectionIndex] = InType;
+		GetIndices()[SelectionIndex] = InSelectedIndices;
+		return true;
+	}
+
+
+}	// namespace  UE::Chaos::ClothAsset
+
+#undef LOCTEXT_NAMESPACE
diff --git a/Engine/Plugins/Experimental/ChaosClothAsset/Source/ChaosClothAsset/Public/ChaosClothAsset/CollectionClothSelectionFacade.h b/Engine/Plugins/Experimental/ChaosClothAsset/Source/ChaosClothAsset/Public/ChaosClothAsset/CollectionClothSelectionFacade.h
new file mode 100644
--- /dev/null
+++ b/Engine/Plugins/Experimental/ChaosClothAsset/Source/ChaosClothAsset/Public/ChaosClothAsset/CollectionClothSelectionFacade.h
@@ -0,0 +1,88 @@
+// Copyright Epic Games, Inc. All Rights Reserved.
+
+#pragma once 
+
+#include "Templates/SharedPointer.h"
+#include "GeometryCollection/ManagedArray.h"
+
+struct FManagedArrayCollection;
+
+namespace UE::Chaos::ClothAsset
+{
+	/**
+	 * Cloth Asset collection Selection facade class. Each Selection consists of a set of integer indices, a name, and a string representing the type of element being indexed.
+	 * Const access (read only) version.
+	 */
+	class CHAOSCLOTHASSET_API FCollectionClothSelectionConstFacade
+	{
+	public:
+		explicit FCollectionClothSelectionConstFacade(const TSharedRef<const FManagedArrayCollection>& ManagedArrayCollection);
+
+		FCollectionClothSelectionConstFacade() = delete;
+		FCollectionClothSelectionConstFacade(const FCollectionClothSelectionConstFacade&) = delete;
+		FCollectionClothSelectionConstFacade& operator=(const FCollectionClothSelectionConstFacade&) = delete;
+		FCollectionClothSelectionConstFacade(FCollectionClothSelectionConstFacade&&) = default;
+		FCollectionClothSelectionConstFacade& operator=(FCollectionClothSelectionConstFacade&&) = default;
+		virtual ~FCollectionClothSelectionConstFacade() = default;
+
+		/** Return whether the facade is defined on the collection. */
+		bool IsValid() const;
+
+		/** Return the total number of selections in the collection */
+		int32 GetNumSelections() const;
+
+		const TArrayView<const FString> GetName() const;
+		const TArrayView<const FString> GetType() const;
+		const TArrayView<const TSet<int32>> GetIndices() const;
+
+		/** Find a selection with the given name. Return INDEX_NONE if no such selection exists. */
+		int32 FindSelection(const FString& Name) const;
+
+	protected:
+
+		TSharedRef<FManagedArrayCollection> ManagedArrayCollection;
+
+		// Schema (non-const since these will be used by the subclass as well)
+		TManagedArray<FString>* Name;
+		TManagedArray<FString>* Type;
+		TManagedArray<TSet<int32>>* Indices;
+	};
+
+
+	/**
+	 * Cloth Asset collection Selection facade class. Each Selection consists of a set of integer indices, a name, and a string representing the type of element being indexed.
+	 * Non-const access (read/write) version.
+	 */
+	class CHAOSCLOTHASSET_API FCollectionClothSelectionFacade final : public FCollectionClothSelectionConstFacade
+	{
+	public:
+
+		explicit FCollectionClothSelectionFacade(const TSharedRef<const FManagedArrayCollection>& ManagedArrayCollection);
+
+		FCollectionClothSelectionFacade() = delete;
+		FCollectionClothSelectionFacade(const FCollectionClothSelectionFacade&) = delete;
+		FCollectionClothSelectionFacade& operator=(const FCollectionClothSelectionFacade&) = delete;
+		FCollectionClothSelectionFacade(FCollectionClothSelectionFacade&&) = default;
+		FCollectionClothSelectionFacade& operator=(FCollectionClothSelectionFacade&&) = default;
+		virtual ~FCollectionClothSelectionFacade() override = default;
+
+		/** Add the Selection attributes to the underlying collection */
+		void DefineSchema();
+
+		TArrayView<FString> GetName();
+		TArrayView<FString> GetType();
+		TArrayView<TSet<int32>> GetIndices();
+
+		/** 
+		* Create a new Selection in the collection with the given name and return its index. If a Selection with the name already exists, return the index of the existing selection. 
+		* 
+		* @param Name The name of the new selection to create if it doesn't already exist
+		* @return The previous number of Selections if we added a new Selection, or the index of the Selection with the given Name if it already existed, or INDEX_NONE if the schema for this facade object is not valid.
+		*/
+		int32 FindOrAddSelection(const FString& Name);
+
+		/** Set the Type and Indices for the given Selection. Returns false if the Facade is invalid or if the SelectionIndex is out of range */
+		bool SetSelection(int32 SelectionIndex, const FString& InType, const TSet<int32>& InSelectedIndices);
+	};
+
+}
diff --git a/Engine/Plugins/Experimental/ChaosClothAssetEditor/Source/ChaosClothAssetDataflowNodes/Private/ChaosClothAsset/ChaosClothAssetDataflowNodesModule.cpp b/Engine/Plugins/Experimental/ChaosClothAssetEditor/Source/ChaosClothAssetDataflowNodes/Private/ChaosClothAsset/ChaosClothAssetDataflowNodesModule.cpp
--- a/Engine/Plugins/Experimental/ChaosClothAssetEditor/Source/ChaosClothAssetDataflowNodes/Private/ChaosClothAsset/ChaosClothAssetDataflowNodesModule.cpp
+++ b/Engine/Plugins/Experimental/ChaosClothAssetEditor/Source/ChaosClothAssetDataflowNodes/Private/ChaosClothAsset/ChaosClothAssetDataflowNodesModule.cpp
@@ -9,6 +9,8 @@
 #include "ChaosClothAsset/ImportNode.h"
 #include "ChaosClothAsset/MergeClothCollectionsNode.h"
 #include "ChaosClothAsset/ReverseNormalsNode.h"
+#include "ChaosClothAsset/SelectionNode.h"
+#include "ChaosClothAsset/SelectionToWeightMapNode.h"
 #include "ChaosClothAsset/SetPhysicsAssetNode.h"
 #include "ChaosClothAsset/SimulationAerodynamicsConfigNode.h"
 #include "ChaosClothAsset/SimulationAnimDriveConfigNode.h"
@@ -56,6 +58,8 @@ namespace UE::Chaos::ClothAsset::Private
 		DATAFLOW_NODE_REGISTER_CREATION_FACTORY(FChaosClothAssetImportNode);
 		DATAFLOW_NODE_REGISTER_CREATION_FACTORY(FChaosClothAssetMergeClothCollectionsNode);
 		DATAFLOW_NODE_REGISTER_CREATION_FACTORY(FChaosClothAssetReverseNormalsNode);
+		DATAFLOW_NODE_REGISTER_CREATION_FACTORY(FChaosClothAssetSelectionNode);
+		DATAFLOW_NODE_REGISTER_CREATION_FACTORY(FChaosClothAssetSelectionToWeightMapNode);
 		DATAFLOW_NODE_REGISTER_CREATION_FACTORY(FChaosClothAssetSetPhysicsAssetNode);
 		DATAFLOW_NODE_REGISTER_CREATION_FACTORY(FChaosClothAssetSimulationAerodynamicsConfigNode);
 		DATAFLOW_NODE_REGISTER_CREATION_FACTORY(FChaosClothAssetSimulationAnimDriveConfigNode);
diff --git a/Engine/Plugins/Experimental/ChaosClothAssetEditor/Source/ChaosClothAssetDataflowNodes/Private/ChaosClothAsset/SelectionNode.cpp b/Engine/Plugins/Experimental/ChaosClothAssetEditor/Source/ChaosClothAssetDataflowNodes/Private/ChaosClothAsset/SelectionNode.cpp
new file mode 100644
--- /dev/null
+++ b/Engine/Plugins/Experimental/ChaosClothAssetEditor/Source/ChaosClothAssetDataflowNodes/Private/ChaosClothAsset/SelectionNode.cpp
@@ -0,0 +1,56 @@
+// Copyright Epic Games, Inc. All Rights Reserved.
+
+#include "ChaosClothAsset/SelectionNode.h"
+#include "ChaosClothAsset/CollectionClothSelectionFacade.h"
+#include "ChaosClothAsset/ClothDataflowTools.h"
+#include "Dataflow/DataflowInputOutput.h"
+
+#include UE_INLINE_GENERATED_CPP_BY_NAME(SelectionNode)
+
+#define LOCTEXT_NAMESPACE "FChaosClothAssetSelectionNode"
+
+FChaosClothAssetSelectionNode::FChaosClothAssetSelectionNode(const Dataflow::FNodeParameters& InParam, FGuid InGuid)
+	: FDataflowNode(InParam, InGuid)
+{
+	RegisterInputConnection(&Collection);
+	RegisterOutputConnection(&Collection, &Collection);
+	RegisterOutputConnection(&Name);
+}
+
+void FChaosClothAssetSelectionNode::Evaluate(Dataflow::FContext& Context, const FDataflowOutput* Out) const
+{
+	using namespace UE::Chaos::ClothAsset;
+
+	if (Out->IsA<FManagedArrayCollection>(&Collection))
+	{
+		FManagedArrayCollection InSelectionCollection = GetValue<FManagedArrayCollection>(Context, &Collection);
+		const TSharedRef<FManagedArrayCollection> SelectionCollection = MakeShared<FManagedArrayCollection>(MoveTemp(InSelectionCollection));
+		FCollectionClothSelectionFacade SelectionFacade(SelectionCollection);
+		SelectionFacade.DefineSchema();
+		
+		checkSlow(SelectionFacade.IsValid());
+
+		const int32 PreviousNumSelections = SelectionFacade.GetNumSelections();
+		const int32 NewSelectionIndex = SelectionFacade.FindOrAddSelection(Name);
+		checkf(NewSelectionIndex != INDEX_NONE, TEXT("Failed to find or add a new Selection"));
+
+		if (NewSelectionIndex != PreviousNumSelections)
+		{
+			FClothDataflowTools::LogAndToastWarning(*this,
+				LOCTEXT("SelectionAlreadyExistsHeadline", "Selection already exists"),
+				FText::Format(LOCTEXT("SelectionAlreadyExistsDetails", "A Selection with Name \"{0}\" already exists in the Collection. Its data will be overwritten."),
+					FText::FromString(Name)));
+		}
+
+		const UEnum* const SelectionTypeEnum = StaticEnum<EChaosClothAssetSelectionType>();
+
+		const bool bSuccess = SelectionFacade.SetSelection(NewSelectionIndex, SelectionTypeEnum->GetNameStringByValue(static_cast<int32>(Type)), Indices);
+		checkf(bSuccess, TEXT("Failed to set value on a Selection"));
+
+		SetValue(Context, MoveTemp(*SelectionCollection), &Collection);
+		SetValue(Context, Name, &Name);
+	}
+}
+
+
+#undef LOCTEXT_NAMESPACE
diff --git a/Engine/Plugins/Experimental/ChaosClothAssetEditor/Source/ChaosClothAssetDataflowNodes/Private/ChaosClothAsset/SelectionToWeightMapNode.cpp b/Engine/Plugins/Experimental/ChaosClothAssetEditor/Source/ChaosClothAssetDataflowNodes/Private/ChaosClothAsset/SelectionToWeightMapNode.cpp
new file mode 100644
--- /dev/null
+++ b/Engine/Plugins/Experimental/ChaosClothAssetEditor/Source/ChaosClothAssetDataflowNodes/Private/ChaosClothAsset/SelectionToWeightMapNode.cpp
@@ -0,0 +1,105 @@
+// Copyright Epic Games, Inc. All Rights Reserved.
+
+#include "ChaosClothAsset/SelectionToWeightMapNode.h"
+#include "ChaosClothAsset/CollectionClothFacade.h"
+#include "ChaosClothAsset/CollectionClothSelectionFacade.h"
+#include "ChaosClothAsset/ClothDataflowTools.h"
+#include "Dataflow/DataflowInputOutput.h"
+
+#include UE_INLINE_GENERATED_CPP_BY_NAME(SelectionToWeightMapNode)
+#define LOCTEXT_NAMESPACE "FChaosClothAssetSelectionToWeightMapNode"
+
+FChaosClothAssetSelectionToWeightMapNode::FChaosClothAssetSelectionToWeightMapNode(const Dataflow::FNodeParameters& InParam, FGuid InGuid)
+	: FDataflowNode(InParam, InGuid)
+{
+	RegisterInputConnection(&Collection);
+	RegisterOutputConnection(&Collection, &Collection);
+	RegisterInputConnection(&Name);
+	RegisterOutputConnection(&Name, &Name);
+}
+
+
+void FChaosClothAssetSelectionToWeightMapNode::Evaluate(Dataflow::FContext& Context, const FDataflowOutput* Out) const
+{
+	using namespace UE::Chaos::ClothAsset;
+
+	if (Out->IsA<FManagedArrayCollection>(&Collection))
+	{
+		FManagedArrayCollection InCollection = GetValue<FManagedArrayCollection>(Context, &Collection);
+		const TSharedRef<FManagedArrayCollection> ClothCollection = MakeShared<FManagedArrayCollection>(MoveTemp(InCollection));
+		FCollectionClothFacade ClothFacade(ClothCollection);
+
+		const FCollectionClothSelectionConstFacade SelectionFacade(ClothCollection);
+		const FString InName = GetValue<FString>(Context, &Name);
+
+		if (SelectionFacade.IsValid())
+		{
+			const int32 FoundSelectionIndex = SelectionFacade.FindSelection(InName);
+
+			if (FoundSelectionIndex != INDEX_NONE)
+			{
+				const FString& SelectionType = SelectionFacade.GetType()[FoundSelectionIndex];
+
+				if (SelectionType == "SimVertex3D" || SelectionType == "SimVertex2D")
+				{
+					const FName MapName(InName);
+					ClothFacade.AddWeightMap(MapName);
+					TArrayView<float> OutClothWeights = ClothFacade.GetWeightMap(MapName);
+
+					const TSet<int32>& Selection = SelectionFacade.GetIndices()[FoundSelectionIndex];
+
+					if (SelectionType == "SimVertex3D")
+					{
+						for (int32 VertexIndex = 0; VertexIndex < OutClothWeights.Num(); ++VertexIndex)
+						{
+							OutClothWeights[VertexIndex] = Selection.Contains(VertexIndex) ? 1.0f : 0.0f;
+						}
+					}
+					else
+					{
+						check(SelectionType == "SimVertex2D");
+
+						// We are given a selection over the set of 2D vertices, but weight maps only exist for 3D vertices, so
+						// we need a bit of translation
+
+						const TConstArrayView<TArray<int32>> Vertex3DTo2D = ClothFacade.GetSimVertex2DLookup();
+
+						for (int32 Vertex3DIndex = 0; Vertex3DIndex < OutClothWeights.Num(); ++Vertex3DIndex)
+						{
+							// If any corresponding 2D vertex is selected, set the 3D weight map value to one, otherwise it gets zero
+							OutClothWeights[Vertex3DIndex] = 0.0;
+							for (const int32 Vertex2DIndex : Vertex3DTo2D[Vertex3DIndex])
+							{
+								if (Selection.Contains(Vertex2DIndex))
+								{
+									OutClothWeights[Vertex3DIndex] = 1.0;
+									break;
+								}
+							}
+						}
+					}
+				}
+				else
+				{
+					FClothDataflowTools::LogAndToastWarning(*this,
+						LOCTEXT("SelectionTypeNotCorrectHeadline", "Selection type is incompatible."),
+						FText::Format(LOCTEXT("SelectionTypeNotCorrectDetails", "Selection with Name \"{0}\" does not have Type \"SimVertex3D\" or \"SimVertex2D\"."),
+							FText::FromString(InName)));
+				}
+			}
+			else
+			{
+				FClothDataflowTools::LogAndToastWarning(*this,
+					LOCTEXT("SelectionNameNotFoundHeadline", "Selection Name was not found."),
+					FText::Format(LOCTEXT("SelectionNameNotFoundDetails", "A Selection with Name \"{0}\" was not found in the Collection."),
+						FText::FromString(InName)));
+			}
+		}
+
+		SetValue(Context, MoveTemp(*ClothCollection), &Collection);
+		SetValue(Context, InName, &Name);
+	}
+}
+
+
+#undef LOCTEXT_NAMESPACE
diff --git a/Engine/Plugins/Experimental/ChaosClothAssetEditor/Source/ChaosClothAssetDataflowNodes/Public/ChaosClothAsset/SelectionNode.h b/Engine/Plugins/Experimental/ChaosClothAssetEditor/Source/ChaosClothAssetDataflowNodes/Public/ChaosClothAsset/SelectionNode.h
new file mode 100644
--- /dev/null
+++ b/Engine/Plugins/Experimental/ChaosClothAssetEditor/Source/ChaosClothAssetDataflowNodes/Public/ChaosClothAsset/SelectionNode.h
@@ -0,0 +1,59 @@
+// Copyright Epic Games, Inc. All Rights Reserved.
+
+#pragma once 
+
+#include "Dataflow/DataflowNode.h"
+#include "GeometryCollection/ManagedArrayCollection.h"
+#include "SelectionNode.generated.h"
+
+/** What type of element is selected in the Selection */
+UENUM()
+enum class EChaosClothAssetSelectionType : uint8
+{
+	/** 2D simulation vertices */
+	SimVertex2D,
+
+	/** 3D simulation vertices */
+	SimVertex3D,
+
+	/** Render vertices */
+	RenderVertex,
+
+	/** Simulation faces (2D/3D are the same) */
+	SimFace,
+
+	/** Render faces */
+	RenderFace
+};
+
+/** Integer index set selection node */
+USTRUCT(Meta = (DataflowCloth))
+struct FChaosClothAssetSelectionNode : public FDataflowNode
+{
+	GENERATED_USTRUCT_BODY()
+	DATAFLOW_NODE_DEFINE_INTERNAL(FChaosClothAssetSelectionNode, "Selection", "Cloth", "Cloth Selection")
+
+public:
+
+	UPROPERTY(Meta = (Dataflowinput, DataflowOutput, DataflowPassthrough = "Collection"))
+	FManagedArrayCollection Collection;
+
+	/** The name to give the selection attribute */
+	UPROPERTY(EditAnywhere, Category = "Selection", Meta = (DataflowOutput))
+	FString Name;
+
+	/** The type of element the selection refers to */
+	UPROPERTY(EditAnywhere, Category = "Selection")
+	EChaosClothAssetSelectionType Type = EChaosClothAssetSelectionType::SimVertex2D;
+
+	/** Selected element indices */
+	UPROPERTY(EditAnywhere, Category = "Selection")
+	TSet<int32> Indices;
+
+	FChaosClothAssetSelectionNode(const Dataflow::FNodeParameters& InParam, FGuid InGuid = FGuid::NewGuid());
+
+private:
+
+	virtual void Evaluate(Dataflow::FContext& Context, const FDataflowOutput* Out) const override;
+};
+
diff --git a/Engine/Plugins/Experimental/ChaosClothAssetEditor/Source/ChaosClothAssetDataflowNodes/Public/ChaosClothAsset/SelectionToWeightMapNode.h b/Engine/Plugins/Experimental/ChaosClothAssetEditor/Source/ChaosClothAssetDataflowNodes/Public/ChaosClothAsset/SelectionToWeightMapNode.h
new file mode 100644
--- /dev/null
+++ b/Engine/Plugins/Experimental/ChaosClothAssetEditor/Source/ChaosClothAssetDataflowNodes/Public/ChaosClothAsset/SelectionToWeightMapNode.h
@@ -0,0 +1,30 @@
+// Copyright Epic Games, Inc. All Rights Reserved.
+
+#pragma once 
+
+#include "Dataflow/DataflowNode.h"
+#include "GeometryCollection/ManagedArrayCollection.h"
+#include "SelectionToWeightMapNode.generated.h"
+
+/** Convert an integer index selection to a vertex weight map where the map value is one for vertices in the selection set, and zero otherwise */
+USTRUCT(Meta = (DataflowCloth))
+struct FChaosClothAssetSelectionToWeightMapNode : public FDataflowNode
+{
+	GENERATED_USTRUCT_BODY()
+	DATAFLOW_NODE_DEFINE_INTERNAL(FChaosClothAssetSelectionToWeightMapNode, "SelectionToWeightMap", "Cloth", "Cloth Selection To Weight Map")
+
+public:
+
+	UPROPERTY(Meta = (Dataflowinput, DataflowOutput, DisplayName = "Collection", DataflowPassthrough = "Collection"))
+	FManagedArrayCollection Collection;
+
+	/** The name of the selection to convert and also the name of the weight map attribute that will be added to the collection */
+	UPROPERTY(Meta = (DataflowInput, DataflowOutput, DataflowPassthrough = "Name"))
+	FString Name;
+
+	FChaosClothAssetSelectionToWeightMapNode(const Dataflow::FNodeParameters& InParam, FGuid InGuid = FGuid::NewGuid());
+
+private:
+
+	virtual void Evaluate(Dataflow::FContext& Context, const FDataflowOutput* Out) const override;
+};
diff --git a/Engine/Plugins/Experimental/ChaosClothAssetEditor/Source/ChaosClothAssetEditor/Private/ChaosClothAsset/ClothEditorCommands.cpp b/Engine/Plugins/Experimental/ChaosClothAssetEditor/Source/ChaosClothAssetEditor/Private/ChaosClothAsset/ClothEditorCommands.cpp
--- a/Engine/Plugins/Experimental/ChaosClothAssetEditor/Source/ChaosClothAssetEditor/Private/ChaosClothAsset/ClothEditorCommands.cpp
+++ b/Engine/Plugins/Experimental/ChaosClothAssetEditor/Source/ChaosClothAssetEditor/Private/ChaosClothAsset/ClothEditorCommands.cpp
@@ -3,6 +3,7 @@
 #include "ChaosClothAsset/ClothEditorCommands.h"
 #include "ChaosClothAsset/ClothEditorStyle.h"
 #include "ChaosClothAsset/ClothWeightMapPaintTool.h"
+#include "ChaosClothAsset/ClothMeshSelectionTool.h"
 
 #define LOCTEXT_NAMESPACE "FChaosClothAssetEditorCommands"
 
@@ -15,6 +16,8 @@ const FString FChaosClothAssetEditorCommands::BeginWeightMapPaintToolIdentifier
 const FString FChaosClothAssetEditorCommands::AddWeightMapNodeIdentifier = TEXT("AddWeightMapNode");
 const FString FChaosClothAssetEditorCommands::BeginTransferSkinWeightsToolIdentifier = TEXT("BeginTransferSkinWeightsTool");
 const FString FChaosClothAssetEditorCommands::AddTransferSkinWeightsNodeIdentifier = TEXT("AddTransferSkinWeightsNode");
+const FString FChaosClothAssetEditorCommands::BeginMeshSelectionToolIdentifier = TEXT("BeginMeshSelectionTool");
+const FString FChaosClothAssetEditorCommands::AddMeshSelectionNodeIdentifier = TEXT("AddMeshSelectionNode");
 const FString FChaosClothAssetEditorCommands::ToggleSimulationSuspendedIdentifier = TEXT("ToggleSimulationSuspended");
 const FString FChaosClothAssetEditorCommands::SoftResetSimulationIdentifier = TEXT("SoftResetSimulation");
 const FString FChaosClothAssetEditorCommands::HardResetSimulationIdentifier = TEXT("HardResetSimulation");
@@ -41,6 +44,8 @@ void FChaosClothAssetEditorCommands::RegisterCommands()
 	
 	UI_COMMAND(BeginWeightMapPaintTool, "MapPnt", "Paint weight maps on the mesh", EUserInterfaceActionType::None, FInputChord());
 	UI_COMMAND(AddWeightMapNode, "MapPnt", "Paint weight maps on the mesh", EUserInterfaceActionType::Button, FInputChord());
+	UI_COMMAND(BeginMeshSelectionTool, "Select", "Select mesh elements", EUserInterfaceActionType::None, FInputChord());
+	UI_COMMAND(AddMeshSelectionNode, "Select", "Select mesh elements", EUserInterfaceActionType::Button, FInputChord());
 
 	UI_COMMAND(BeginTransferSkinWeightsTool, "TransferSkinWeights", "Transfer skinning weights from a SkeletalMesh", EUserInterfaceActionType::None, FInputChord());
 	UI_COMMAND(AddTransferSkinWeightsNode, "TransferSkinWeights", "Transfer skinning weights from a SkeletalMesh", EUserInterfaceActionType::Button, FInputChord());
diff --git a/Engine/Plugins/Experimental/ChaosClothAssetEditor/Source/ChaosClothAssetEditor/Private/ChaosClothAsset/ClothEditorCommands.h b/Engine/Plugins/Experimental/ChaosClothAssetEditor/Source/ChaosClothAssetEditor/Private/ChaosClothAsset/ClothEditorCommands.h
--- a/Engine/Plugins/Experimental/ChaosClothAssetEditor/Source/ChaosClothAssetEditor/Private/ChaosClothAsset/ClothEditorCommands.h
+++ b/Engine/Plugins/Experimental/ChaosClothAssetEditor/Source/ChaosClothAssetEditor/Private/ChaosClothAsset/ClothEditorCommands.h
@@ -46,6 +46,11 @@ public:
 	const static FString AddTransferSkinWeightsNodeIdentifier;
 	TSharedPtr<FUICommandInfo> AddTransferSkinWeightsNode;
 
+	const static FString BeginMeshSelectionToolIdentifier;
+	TSharedPtr<FUICommandInfo> BeginMeshSelectionTool;
+	const static FString AddMeshSelectionNodeIdentifier;
+	TSharedPtr<FUICommandInfo> AddMeshSelectionNode;
+
 	// Construction viewport commands
 	const static FString ToggleConstructionViewWireframeIdentifier;
 	TSharedPtr<FUICommandInfo> ToggleConstructionViewWireframe;
diff --git a/Engine/Plugins/Experimental/ChaosClothAssetEditor/Source/ChaosClothAssetEditor/Private/ChaosClothAsset/ClothEditorMode.cpp b/Engine/Plugins/Experimental/ChaosClothAssetEditor/Source/ChaosClothAssetEditor/Private/ChaosClothAsset/ClothEditorMode.cpp
--- a/Engine/Plugins/Experimental/ChaosClothAssetEditor/Source/ChaosClothAssetEditor/Private/ChaosClothAsset/ClothEditorMode.cpp
+++ b/Engine/Plugins/Experimental/ChaosClothAssetEditor/Source/ChaosClothAssetEditor/Private/ChaosClothAsset/ClothEditorMode.cpp
@@ -12,6 +12,7 @@
 #include "ChaosClothAsset/ClothEditorContextObject.h"
 #include "ChaosClothAsset/AddWeightMapNode.h"
 #include "ChaosClothAsset/TransferSkinWeightsNode.h"
+#include "ChaosClothAsset/SelectionNode.h"
 #include "AssetEditorModeManager.h"
 #include "Drawing/MeshElementsVisualizer.h"
 #include "EditorViewportClient.h"
@@ -37,6 +38,7 @@
 #include "MeshAttributePaintTool.h"
 #include "ChaosClothAsset/ClothWeightMapPaintTool.h"
 #include "ChaosClothAsset/ClothTransferSkinWeightsTool.h"
+#include "ChaosClothAsset/ClothMeshSelectionTool.h"
 #include "DynamicMesh/DynamicMeshAttributeSet.h"
 #include "DynamicMesh/DynamicVertexSkinWeightsAttribute.h"
 #include "DynamicMeshEditor.h"
@@ -217,6 +219,9 @@ void UChaosClothAssetEditorMode::RegisterTools()
 
 	RegisterClothTool(CommandInfos.BeginTransferSkinWeightsTool, FChaosClothAssetEditorCommands::BeginTransferSkinWeightsToolIdentifier, NewObject<UClothTransferSkinWeightsToolBuilder>(), ConstructionViewportToolsContext);
 	RegisterAddNodeCommand(CommandInfos.AddTransferSkinWeightsNode, FChaosClothAssetTransferSkinWeightsNode::StaticType(), CommandInfos.BeginTransferSkinWeightsTool);
+
+	RegisterClothTool(CommandInfos.BeginMeshSelectionTool, FChaosClothAssetEditorCommands::BeginMeshSelectionToolIdentifier, NewObject<UClothMeshSelectionToolBuilder>(), ConstructionViewportToolsContext);
+	RegisterAddNodeCommand(CommandInfos.AddMeshSelectionNode, FChaosClothAssetSelectionNode::StaticType(), CommandInfos.BeginMeshSelectionTool);
 }
 
 bool UChaosClothAssetEditorMode::ShouldToolStartBeAllowed(const FString& ToolIdentifier) const
@@ -452,7 +457,6 @@ void UChaosClothAssetEditorMode::ReinitializeDynamicMeshComponents()
 
 	if (DynamicMeshComponent)
 	{
-		// TODO: Check for any outstanding changes on the dynamic mesh component?
 		DynamicMeshComponent->UnregisterComponent();
 		DynamicMeshComponent->SelectionOverrideDelegate.Unbind();
 
@@ -551,11 +555,8 @@ void UChaosClothAssetEditorMode::ReinitializeDynamicMeshComponents()
 	WireframeToTick->Settings->bVisible = bRestSpaceMeshVisible && bConstructionViewWireframe;
 
 	SelectedComponents->DeselectAll();
-	if (ConstructionViewMode != EClothPatternVertexType::Render)
-	{
-		SelectedComponents->Select(DynamicMeshComponent);
-		DynamicMeshComponent->PushSelectionToProxy();
-	}
+	SelectedComponents->Select(DynamicMeshComponent);
+	DynamicMeshComponent->PushSelectionToProxy();
 
 	// Update the context object with the ConstructionViewMode and Collection used to build the DynamicMeshComponents, so 
 	// tools know how to use the components.
diff --git a/Engine/Plugins/Experimental/ChaosClothAssetEditor/Source/ChaosClothAssetEditor/Private/ChaosClothAsset/ClothEditorModeToolkit.cpp b/Engine/Plugins/Experimental/ChaosClothAssetEditor/Source/ChaosClothAssetEditor/Private/ChaosClothAsset/ClothEditorModeToolkit.cpp
--- a/Engine/Plugins/Experimental/ChaosClothAssetEditor/Source/ChaosClothAssetEditor/Private/ChaosClothAsset/ClothEditorModeToolkit.cpp
+++ b/Engine/Plugins/Experimental/ChaosClothAssetEditor/Source/ChaosClothAssetEditor/Private/ChaosClothAsset/ClothEditorModeToolkit.cpp
@@ -148,6 +148,7 @@ void FChaosClothAssetEditorModeToolkit::BuildToolPalette(FName PaletteIndex, cla
 	if (PaletteIndex == FBaseCharacterFXEditorModeToolkit::ToolsTabName)
 	{
 		ToolbarBuilder.AddToolBarButton(Commands.AddWeightMapNode);
+		ToolbarBuilder.AddToolBarButton(Commands.AddMeshSelectionNode);
 		ToolbarBuilder.AddToolBarButton(Commands.AddTransferSkinWeightsNode);
 	}
 }
diff --git a/Engine/Plugins/Experimental/ChaosClothAssetEditor/Source/ChaosClothAssetEditor/Private/ChaosClothAsset/SClothCollectionOutliner.cpp b/Engine/Plugins/Experimental/ChaosClothAssetEditor/Source/ChaosClothAssetEditor/Private/ChaosClothAsset/SClothCollectionOutliner.cpp
--- a/Engine/Plugins/Experimental/ChaosClothAssetEditor/Source/ChaosClothAssetEditor/Private/ChaosClothAsset/SClothCollectionOutliner.cpp
+++ b/Engine/Plugins/Experimental/ChaosClothAssetEditor/Source/ChaosClothAssetEditor/Private/ChaosClothAsset/SClothCollectionOutliner.cpp
@@ -228,6 +228,23 @@ namespace ClothCollectionOutlinerHelpers
 		return Out;
 	}
 
+	template<typename T>
+	FString AttributeValueToString(const TSet<T>& Set)
+	{
+		FString Out;
+		typename TSet<T>::TConstIterator Iter = Set.CreateConstIterator();
+		while(Iter)
+		{
+			Out += AttributeValueToString(*Iter);
+
+			if (++Iter)
+			{
+				Out += "; ";
+			}
+		}
+		return Out;
+	}
+
 	template<typename T>
 	FString AttributeValueToString(const FManagedArrayCollection& ClothCollection, const FName& AttributeName, const FName& GroupName, int32 AttributeArrayIndex)
 	{
@@ -284,6 +301,9 @@ namespace ClothCollectionOutlinerHelpers
 		case FManagedArrayCollection::EArrayType::FIntVector2ArrayType:
 			ValueAsString = AttributeValueToString<TArray<FIntVector2>>(ClothCollection, AttributeName, GroupName, AttributeArrayIndex);
 			break;
+		case FManagedArrayCollection::EArrayType::FIntArrayType:
+			ValueAsString = AttributeValueToString<TSet<int32>>(ClothCollection, AttributeName, GroupName, AttributeArrayIndex);
+			break;
 		default:
 			ensure(false);
 			ValueAsString = "(Unknown Data Type)";
diff --git a/Engine/Plugins/Experimental/ChaosClothAssetEditor/Source/ChaosClothAssetEditorTools/Private/ChaosClothAsset/ClothMeshSelectionTool.cpp b/Engine/Plugins/Experimental/ChaosClothAssetEditor/Source/ChaosClothAssetEditorTools/Private/ChaosClothAsset/ClothMeshSelectionTool.cpp
new file mode 100644
--- /dev/null
+++ b/Engine/Plugins/Experimental/ChaosClothAssetEditor/Source/ChaosClothAssetEditorTools/Private/ChaosClothAsset/ClothMeshSelectionTool.cpp
@@ -0,0 +1,421 @@
+// Copyright Epic Games, Inc. All Rights Reserved.
+
+#include "ChaosClothAsset/ClothMeshSelectionTool.h"
+#include "ChaosClothAsset/ClothEditorContextObject.h"
+#include "ChaosClothAsset/ClothPatternVertexType.h"
+#include "ChaosClothAsset/SelectionNode.h"
+#include "ChaosClothAsset/CollectionClothFacade.h"
+#include "InteractiveToolManager.h"
+#include "PreviewMesh.h"
+#include "TargetInterfaces/PrimitiveComponentBackedTarget.h"
+#include "ToolTargetManager.h"
+#include "Selection/PolygonSelectionMechanic.h"
+#include "ModelingToolTargetUtil.h"
+#include "GroupTopology.h"
+#include "ToolSetupUtil.h"
+#include "Selections/GeometrySelection.h"
+#include "ContextObjectStore.h"
+#include "Dataflow/DataflowGraphEditor.h"
+#include "Dataflow/DataflowEdNode.h"
+#include "DynamicMesh/NonManifoldMappingSupport.h"
+
+#include UE_INLINE_GENERATED_CPP_BY_NAME(ClothMeshSelectionTool)
+
+#define LOCTEXT_NAMESPACE "ClothMeshSelectionTool"
+
+
+// ------------------- Builder -------------------
+
+const FToolTargetTypeRequirements& UClothMeshSelectionToolBuilder::GetTargetRequirements() const
+{
+	static FToolTargetTypeRequirements TypeRequirements({
+		UPrimitiveComponentBackedTarget::StaticClass()});
+	return TypeRequirements;
+}
+
+bool UClothMeshSelectionToolBuilder::CanBuildTool(const FToolBuilderState& SceneState) const
+{
+	return (SceneState.TargetManager->CountSelectedAndTargetable(SceneState, GetTargetRequirements()) == 1);
+}
+
+UInteractiveTool* UClothMeshSelectionToolBuilder::BuildTool(const FToolBuilderState& SceneState) const
+{
+	UClothMeshSelectionTool* const NewTool = NewObject<UClothMeshSelectionTool>(SceneState.ToolManager);
+
+	UToolTarget* const Target = SceneState.TargetManager->BuildFirstSelectedTargetable(SceneState, GetTargetRequirements());
+	NewTool->SetTarget(Target);
+	NewTool->SetWorld(SceneState.World);
+
+	if (UClothEditorContextObject* const ContextObject = SceneState.ToolManager->GetContextObjectStore()->FindContext<UClothEditorContextObject>())
+	{
+		NewTool->SetClothEditorContextObject(ContextObject);
+	}
+
+	return NewTool;
+}
+
+
+// ------------------- Tool -------------------
+	
+void UClothMeshSelectionTool::Setup()
+{
+	using namespace UE::Geometry;
+
+	//
+	// Preview 
+	//
+
+	// Create the preview mesh
+	PreviewMesh = NewObject<UPreviewMesh>(this, TEXT("PreviewMesh"));
+	PreviewMesh->CreateInWorld(GetTargetWorld(), FTransform::Identity);
+
+	ToolSetupUtil::ApplyRenderingConfigurationToPreview(PreviewMesh, Target);
+
+	// We will use the preview mesh's spatial data structure
+	PreviewMesh->bBuildSpatialDataStructure = true;
+
+	// set materials
+	FComponentMaterialSet MaterialSet = UE::ToolTarget::GetMaterialSet(Target);
+	PreviewMesh->SetMaterials(MaterialSet.Materials);
+
+	// configure secondary render material for selected triangles
+	UMaterialInterface* SelectionMaterial = ToolSetupUtil::GetSelectionMaterial(FLinearColor::Yellow, GetToolManager());
+	if (SelectionMaterial != nullptr)
+	{
+		PreviewMesh->SetSecondaryRenderMaterial(SelectionMaterial);
+	}
+
+	PreviewMesh->SetTangentsMode(EDynamicMeshComponentTangentsMode::AutoCalculated);
+	PreviewMesh->UpdatePreview(UE::ToolTarget::GetDynamicMeshCopy(Target));
+	PreviewMesh->SetVisible(true);
+
+	//
+	// SelectionMechanic
+	//
+
+	SelectionMechanic = NewObject<UPolygonSelectionMechanic>(this);
+	SelectionMechanic->bAddSelectionFilterPropertiesToParentTool = false;   // We'll do this ourselves later
+	SelectionMechanic->Setup(this);
+	SelectionMechanic->Properties->RestoreProperties(this);
+	SelectionMechanic->Properties->bDisplayPolygroupReliantControls = false;	// this is for polygroup-specific selections like edge loops
+	
+	SelectionMechanic->Properties->bCanSelectVertices = true;
+	SelectionMechanic->Properties->bCanSelectEdges = false;		// For now do not allow edge selection
+	SelectionMechanic->Properties->bCanSelectFaces = true;
+	
+	SelectionMechanic->SetShowEdges(false);
+	SelectionMechanic->SetShowSelectableCorners(false);
+
+	SelectionMechanic->PolyEdgesRenderer.DepthBias = 0.01f;
+	SelectionMechanic->PolyEdgesRenderer.LineThickness = 1.0f;
+	SelectionMechanic->PolyEdgesRenderer.PointSize = 2.0f;
+
+	SelectionMechanic->SelectionRenderer.DepthBias = 0.01f;
+	SelectionMechanic->SelectionRenderer.LineThickness = 1.0f;
+	SelectionMechanic->SelectionRenderer.PointSize = 2.0f;
+
+	SelectionMechanic->OnSelectionChanged.AddWeakLambda(this, [this]()
+	{
+		bAnyChangeMade = true;
+		PreviewMesh->FastNotifySecondaryTrianglesChanged();
+	});
+
+	SelectionMechanic->OnFaceSelectionPreviewChanged.AddWeakLambda(this, [this]()
+	{
+		PreviewMesh->FastNotifySecondaryTrianglesChanged();
+	});
+
+
+	// Enable only one selection mode at a time (this is different from other mesh modeling tools using the SelectionMechanic)
+
+	SelectionMechanic->Properties->WatchProperty(SelectionMechanic->Properties->bSelectVertices, [this](bool bNewSelectVertices)
+	{
+		SelectionMechanic->Properties->bSelectFaces = !bNewSelectVertices;
+	});
+
+	SelectionMechanic->Properties->WatchProperty(SelectionMechanic->Properties->bSelectFaces, [this](bool bNewSelectFaces)
+	{
+		SelectionMechanic->Properties->bSelectVertices = !bNewSelectFaces;
+	});
+
+	// Set up Topology and SelectionMechanic using Preview's DynamicMesh
+
+	PreviewMesh->ProcessMesh([this](const FDynamicMesh3& Mesh)
+	{
+		Topology = MakeUnique<FTriangleGroupTopology>(&Mesh, true);
+
+		SelectionMechanic->Initialize(&Mesh,
+			FTransform3d(),
+			GetTargetWorld(),
+			Topology.Get(),
+			[this]() { return PreviewMesh->GetSpatial(); });
+	});
+
+	PreviewMesh->EnableSecondaryTriangleBuffers([this](const FDynamicMesh3* Mesh, int32 TriangleID)
+	{
+		return SelectionMechanic->GetActiveSelection().IsSelectedTriangle(Mesh, Topology.Get(), TriangleID);
+	});
+
+	// Hide input target mesh
+	UE::ToolTarget::HideSourceObject(Target);
+
+	// Initialize the Selection from the selected Dataflow node
+	FString ExistingSelectionName;
+	FGroupTopologySelection ExistingNodeSelection;
+	GetSelectedNodeInfo(ExistingSelectionName, ExistingNodeSelection);
+
+	constexpr bool bBroadcastChange = false;
+	SelectionMechanic->SetSelection(ExistingNodeSelection, bBroadcastChange);
+
+	if (ExistingNodeSelection.SelectedCornerIDs.Num() > 0)
+	{
+		check(ExistingNodeSelection.SelectedEdgeIDs.Num() == 0);
+		check(ExistingNodeSelection.SelectedGroupIDs.Num() == 0);
+		SelectionMechanic->Properties->bSelectVertices = true;
+		SelectionMechanic->Properties->bSelectFaces = false;
+	} 
+	else
+	{
+		SelectionMechanic->Properties->bSelectVertices = false;
+		SelectionMechanic->Properties->bSelectFaces = true;
+	}
+
+	// Setup non-manifold mapping if necessary
+	if (ClothEditorContextObject)
+	{
+		if (const TSharedPtr<const FManagedArrayCollection> ClothCollection = ClothEditorContextObject->GetSelectedClothCollection().Pin())
+		{
+			PreviewMesh->ProcessMesh([this, ClothCollection](const FDynamicMesh3& Mesh)
+			{
+				using namespace UE::Chaos::ClothAsset;
+				const FNonManifoldMappingSupport NonManifoldMapping(Mesh);
+
+				bHasNonManifoldMapping = NonManifoldMapping.IsNonManifoldVertexInSource();
+
+				if (bHasNonManifoldMapping)
+				{
+					const FCollectionClothConstFacade Cloth(ClothCollection.ToSharedRef());
+					check(Cloth.IsValid());
+
+					const TConstArrayView<int32> SimVertex3DLookup = Cloth.GetSimVertex3DLookup();
+
+					DynamicMeshToSelection.SetNumUninitialized(Mesh.VertexCount());
+					SelectionToDynamicMesh.Reset();
+					SelectionToDynamicMesh.SetNum(Cloth.GetNumSimVertices3D());
+					for (int32 DynamicMeshVert = 0; DynamicMeshVert < Mesh.VertexCount(); ++DynamicMeshVert)
+					{
+						DynamicMeshToSelection[DynamicMeshVert] = NonManifoldMapping.GetOriginalNonManifoldVertexID(DynamicMeshVert);
+						SelectionToDynamicMesh[DynamicMeshToSelection[DynamicMeshVert]].Add(DynamicMeshVert);
+					}
+				}
+			});
+		}
+	}
+
+
+	//
+	// Properties
+	//
+
+	ToolProperties = NewObject<UClothMeshSelectionToolProperties>();
+	ToolProperties->WatchProperty(ToolProperties->bShowVertices, [this](bool bNewShowVertices)
+	{
+		SelectionMechanic->SetShowSelectableCorners(bNewShowVertices);
+	});
+
+	ToolProperties->WatchProperty(ToolProperties->bShowEdges, [this](bool bNewShowEdges)
+	{
+		SelectionMechanic->SetShowEdges(bNewShowEdges);
+	});
+
+	ToolProperties->Name = ExistingSelectionName;
+
+	ToolProperties->WatchProperty(ToolProperties->Name, [this, OriginalName = ToolProperties->Name](const FString& NewName)
+	{
+		if (NewName != OriginalName)
+		{
+			bAnyChangeMade = true;
+		}
+	});
+
+	AddToolPropertySource(ToolProperties);
+	ToolProperties->RestoreProperties(this);
+
+	AddToolPropertySource(SelectionMechanic->Properties);
+}
+
+void UClothMeshSelectionTool::OnShutdown(EToolShutdownType ShutdownType)
+{
+	if (ShutdownType == EToolShutdownType::Accept)
+	{
+		GetToolManager()->BeginUndoTransaction(LOCTEXT("SelectionToolTransactionName", "Mesh Selection"));
+		UpdateSelectedNode();
+		GetToolManager()->EndUndoTransaction();
+	}
+
+	SelectionMechanic->Properties->SaveProperties(this);
+	ToolProperties->SaveProperties(this);
+
+	if (PreviewMesh != nullptr)
+	{
+		UE::ToolTarget::ShowSourceObject(Target);
+		PreviewMesh->Disconnect();
+		PreviewMesh = nullptr;
+	}
+
+	SelectionMechanic->Shutdown();
+
+	Topology.Reset();
+}
+
+void UClothMeshSelectionTool::Render(IToolsContextRenderAPI* RenderAPI)
+{
+	SelectionMechanic->Render(RenderAPI);
+}
+
+void UClothMeshSelectionTool::DrawHUD(FCanvas* Canvas, IToolsContextRenderAPI* RenderAPI)
+{
+	SelectionMechanic->DrawHUD(Canvas, RenderAPI);
+}
+
+bool UClothMeshSelectionTool::CanAccept() const
+{
+	return bAnyChangeMade;
+}
+
+void UClothMeshSelectionTool::SetClothEditorContextObject(TObjectPtr<UClothEditorContextObject> InClothEditorContextObject)
+{
+	ClothEditorContextObject = InClothEditorContextObject;
+}
+
+bool UClothMeshSelectionTool::GetSelectedNodeInfo(FString& OutSelectionName, UE::Geometry::FGroupTopologySelection& OutSelection)
+{
+	const FChaosClothAssetSelectionNode* const MeshSelectionNode = ClothEditorContextObject->GetSingleSelectedNodeOfType<FChaosClothAssetSelectionNode>();
+	check(MeshSelectionNode);
+
+	// We need to sanitize the incoming indices, as the user can manually set them to anything on the node
+
+	auto AppendVerticesIfValid = [this]<typename T>(TSet<int32>& Dest, const T& Source)
+	{
+		PreviewMesh->ProcessMesh([this, &Dest, &Source](const UE::Geometry::FDynamicMesh3& Mesh)
+		{
+			for (const int32& VertexIndex : Source)
+			{
+				if (Mesh.IsVertex(VertexIndex))
+				{
+					Dest.Add(VertexIndex);
+				}
+			}
+		});
+	};
+
+	auto AppendFacesIfValid = [this](TSet<int32>& Dest, const TSet<int32>& Source)
+	{
+		PreviewMesh->ProcessMesh([this, &Dest, &Source](const UE::Geometry::FDynamicMesh3& Mesh)
+		{
+			for (const int32& FaceIndex : Source)
+			{
+				if (Mesh.IsTriangle(FaceIndex))
+				{
+					Dest.Add(FaceIndex);
+				}
+			}
+		});
+	};
+
+
+	if (MeshSelectionNode->Type == EChaosClothAssetSelectionType::SimVertex2D || MeshSelectionNode->Type == EChaosClothAssetSelectionType::SimVertex3D || MeshSelectionNode->Type == EChaosClothAssetSelectionType::RenderVertex)
+	{
+		if (bHasNonManifoldMapping)
+		{
+			for (const int32 SelectionIndex : MeshSelectionNode->Indices)
+			{
+				AppendVerticesIfValid(OutSelection.SelectedCornerIDs, SelectionToDynamicMesh[SelectionIndex]);
+			}
+		}
+		else
+		{
+			AppendVerticesIfValid(OutSelection.SelectedCornerIDs, MeshSelectionNode->Indices);
+		}
+	}
+	else if (MeshSelectionNode->Type == EChaosClothAssetSelectionType::SimFace || MeshSelectionNode->Type == EChaosClothAssetSelectionType::RenderFace)
+	{
+		AppendFacesIfValid(OutSelection.SelectedGroupIDs, MeshSelectionNode->Indices);
+	}
+
+	OutSelectionName = MeshSelectionNode->Name;
+
+	return true;
+}
+
+
+void UClothMeshSelectionTool::UpdateSelectedNode()
+{
+	const UE::Geometry::FGroupTopologySelection& Selection = SelectionMechanic->GetActiveSelection();
+	const UE::Chaos::ClothAsset::EClothPatternVertexType ViewMode = ClothEditorContextObject->GetConstructionViewMode();
+
+	TSet<int32> Indices;
+	EChaosClothAssetSelectionType Type = EChaosClothAssetSelectionType::SimVertex2D;
+
+	if (SelectionMechanic->Properties->bSelectVertices)
+	{
+		check(!SelectionMechanic->Properties->bSelectEdges);
+		check(!SelectionMechanic->Properties->bSelectFaces);
+
+		Indices = Selection.SelectedCornerIDs;
+
+		switch(ViewMode)
+		{
+		case UE::Chaos::ClothAsset::EClothPatternVertexType::Sim2D:
+			Type = EChaosClothAssetSelectionType::SimVertex2D;
+			break;
+		case UE::Chaos::ClothAsset::EClothPatternVertexType::Sim3D:
+			Type = EChaosClothAssetSelectionType::SimVertex3D;
+			break;
+		case UE::Chaos::ClothAsset::EClothPatternVertexType::Render:
+			Type = EChaosClothAssetSelectionType::RenderVertex;
+			break;
+		}
+
+	}
+	else
+	{
+		Indices = Selection.SelectedGroupIDs;
+
+		switch (ViewMode)
+		{
+		case UE::Chaos::ClothAsset::EClothPatternVertexType::Sim2D:
+		case UE::Chaos::ClothAsset::EClothPatternVertexType::Sim3D:
+			Type = EChaosClothAssetSelectionType::SimFace;
+			break;
+		case UE::Chaos::ClothAsset::EClothPatternVertexType::Render:
+			Type = EChaosClothAssetSelectionType::RenderFace;
+			break;
+		}
+	}
+
+	FChaosClothAssetSelectionNode* const MeshSelectionNode = ClothEditorContextObject->GetSingleSelectedNodeOfType<FChaosClothAssetSelectionNode>();
+	check(MeshSelectionNode);
+
+	MeshSelectionNode->Name = ToolProperties->Name;
+	MeshSelectionNode->Type = Type;
+
+	if (SelectionMechanic->Properties->bSelectVertices && bHasNonManifoldMapping)
+	{
+		MeshSelectionNode->Indices.Reset();
+		for (const int32 DynamicMeshIdx : Indices)
+		{
+			const int32 MappedSelectionIndex = DynamicMeshToSelection[DynamicMeshIdx];
+			MeshSelectionNode->Indices.Add(MappedSelectionIndex);
+		}
+	}
+	else
+	{
+		MeshSelectionNode->Indices = MoveTemp(Indices);
+	}
+
+	MeshSelectionNode->Invalidate();
+}
+
+
+#undef LOCTEXT_NAMESPACE
diff --git a/Engine/Plugins/Experimental/ChaosClothAssetEditor/Source/ChaosClothAssetEditorTools/Public/ChaosClothAsset/ClothMeshSelectionTool.h b/Engine/Plugins/Experimental/ChaosClothAssetEditor/Source/ChaosClothAssetEditorTools/Public/ChaosClothAsset/ClothMeshSelectionTool.h
new file mode 100644
--- /dev/null
+++ b/Engine/Plugins/Experimental/ChaosClothAssetEditor/Source/ChaosClothAssetEditorTools/Public/ChaosClothAsset/ClothMeshSelectionTool.h
@@ -0,0 +1,89 @@
+#pragma once
+
+#include "BaseTools/SingleSelectionMeshEditingTool.h"
+#include "ClothMeshSelectionTool.generated.h"
+
+class UPolygonSelectionMechanic;
+class UClothEditorContextObject;
+class UPreviewMesh;
+
+namespace UE::Geometry
+{
+	class FGroupTopology;
+	struct FGroupTopologySelection;
+}
+
+UCLASS()
+class CHAOSCLOTHASSETEDITORTOOLS_API UClothMeshSelectionToolBuilder : public UInteractiveToolWithToolTargetsBuilder
+{
+	GENERATED_BODY()
+
+private:
+
+	virtual bool CanBuildTool(const FToolBuilderState& SceneState) const override;
+	virtual UInteractiveTool* BuildTool(const FToolBuilderState& SceneState) const override;
+	virtual const FToolTargetTypeRequirements& GetTargetRequirements() const override;
+};
+
+UCLASS()
+class CHAOSCLOTHASSETEDITORTOOLS_API UClothMeshSelectionToolProperties : public UInteractiveToolPropertySet
+{
+	GENERATED_BODY()
+
+public:
+
+	UPROPERTY(EditAnywhere, Transient, Category = Name, meta = (DisplayName = "Name", TransientToolProperty))
+	FString Name;
+
+	UPROPERTY(EditAnywhere, Category = Visualization, meta = (DisplayName = "Show Vertices"))
+	bool bShowVertices = false;
+	
+	UPROPERTY(EditAnywhere, Category = Visualization, meta = (DisplayName = "Show Edges"))
+	bool bShowEdges = false;
+
+};
+
+UCLASS()
+class CHAOSCLOTHASSETEDITORTOOLS_API UClothMeshSelectionTool : public USingleSelectionMeshEditingTool
+{
+	GENERATED_BODY()
+
+private:
+
+	friend class UClothMeshSelectionToolBuilder;
+
+	virtual void Setup() override;
+	virtual void OnShutdown(EToolShutdownType ShutdownType) override;
+	virtual void Render(IToolsContextRenderAPI* RenderAPI) override;
+	virtual void DrawHUD(FCanvas* Canvas, IToolsContextRenderAPI* RenderAPI) override;
+
+	virtual bool HasCancel() const override { return true; }
+	virtual bool HasAccept() const override { return true; }
+	virtual bool CanAccept() const override;
+
+	void SetClothEditorContextObject(TObjectPtr<UClothEditorContextObject> InClothEditorContextObject);
+	bool GetSelectedNodeInfo(FString& OutMapName, UE::Geometry::FGroupTopologySelection& OutSelection);
+	void UpdateSelectedNode();
+
+	UPROPERTY()
+	TObjectPtr<UClothMeshSelectionToolProperties> ToolProperties;
+
+	UPROPERTY()
+	TObjectPtr<UPreviewMesh> PreviewMesh = nullptr;
+
+	UPROPERTY()
+	TObjectPtr<UPolygonSelectionMechanic> SelectionMechanic = nullptr;
+
+	UPROPERTY()
+	TObjectPtr<UClothEditorContextObject> ClothEditorContextObject = nullptr;
+
+	TUniquePtr<UE::Geometry::FGroupTopology> Topology;
+
+	bool bAnyChangeMade = false;
+
+	bool bHasNonManifoldMapping = false;
+	TArray<int32> DynamicMeshToSelection;
+	TArray<TArray<int32>> SelectionToDynamicMesh;
+
+};
+
