Chaos Cloth - Fixed crash/ensure when using the default config node after any other config property node.

#rnx

diff --git a/Engine/Plugins/Experimental/ChaosClothAssetEditor/Source/ChaosClothAssetDataflowNodes/Private/ChaosClothAsset/SimulationDefaultConfigNode.cpp b/Engine/Plugins/Experimental/ChaosClothAssetEditor/Source/ChaosClothAssetDataflowNodes/Private/ChaosClothAsset/SimulationDefaultConfigNode.cpp
--- a/Engine/Plugins/Experimental/ChaosClothAssetEditor/Source/ChaosClothAssetDataflowNodes/Private/ChaosClothAsset/SimulationDefaultConfigNode.cpp
+++ b/Engine/Plugins/Experimental/ChaosClothAssetEditor/Source/ChaosClothAssetDataflowNodes/Private/ChaosClothAsset/SimulationDefaultConfigNode.cpp
@@ -1,6 +1,7 @@
 // Copyright Epic Games, Inc. All Rights Reserved.
 
 #include "ChaosClothAsset/SimulationDefaultConfigNode.h"
+#include "ChaosClothAsset/ClothDataflowTools.h"
 #include "ChaosClothAsset/ClothEngineTools.h"
 #include "ChaosClothAsset/CollectionClothFacade.h"
 #include "ChaosCloth/ChaosClothConfig.h"
@@ -50,19 +51,28 @@ void FChaosClothAssetSimulationDefaultConfigNode::Evaluate(Dataflow::FContext& C
 
 		if (FCollectionClothFacade(ClothCollection).IsValid())  // Can only act on the collection if it is a valid cloth collection
 		{
-			FClothingSimulationConfig ClothingSimulationConfig;
-			ClothingSimulationConfig.Initialize(SimulationConfig.Get(), SharedSimulationConfig.Get());
-
-			ClothingSimulationConfig.GetPropertyCollection()->CopyTo(&ClothCollection.Get());
-
-			// Generate tethers
-			const FCollectionPropertyConstFacade& Properties = ClothingSimulationConfig.GetProperties();
-			constexpr bool bUseGeodesicTethersDefault = true;
-			const bool bUseGeodesicTethers = Properties.GetValue<bool>(TEXT("UseGeodesicTethers"), bUseGeodesicTethersDefault);
-			// Use the "MaxDistance" weight map to generate tethers. This follows legacy behavior.
-			static const FName MaxDistanceName(TEXT("MaxDistance"));
-
-			UE::Chaos::ClothAsset::FClothEngineTools::GenerateTethers(ClothCollection, MaxDistanceName, bUseGeodesicTethers);
+			if (!FCollectionPropertyConstFacade(ClothCollection).IsValid())  // Cannot already have any properties set before this node (it's using CopyTo and wouldn't warn when overriding)
+			{
+				FClothingSimulationConfig ClothingSimulationConfig;
+				ClothingSimulationConfig.Initialize(SimulationConfig.Get(), SharedSimulationConfig.Get());
+
+				ClothingSimulationConfig.GetPropertyCollection()->CopyTo(&ClothCollection.Get());
+
+				// Generate tethers
+				const FCollectionPropertyConstFacade& Properties = ClothingSimulationConfig.GetProperties();
+				constexpr bool bUseGeodesicTethersDefault = true;
+				const bool bUseGeodesicTethers = Properties.GetValue<bool>(TEXT("UseGeodesicTethers"), bUseGeodesicTethersDefault);
+				// Use the "MaxDistance" weight map to generate tethers. This follows legacy behavior.
+				static const FName MaxDistanceName(TEXT("MaxDistance"));
+
+				UE::Chaos::ClothAsset::FClothEngineTools::GenerateTethers(ClothCollection, MaxDistanceName, bUseGeodesicTethers);
+			}
+			else
+			{
+				FClothDataflowTools::LogAndToastWarning(*this,
+					LOCTEXT("CannotOverridePreExistingConfigHeadline", "Cannot override pre-existing config."),
+					LOCTEXT("CannotOverridePreExistingConfigDetails", "The Simulation Default Config node cannot be used to override any existing properties, and must be placed ahead of any other config nodes."));
+			}
 		}
 
 		SetValue(Context, MoveTemp(*ClothCollection), &Collection);
