GraphEditor: add a new delegate that is called when the user single-clicks on a node
Cloth Editor: fix for interactive tool starting when the user drags a node in the dataflow graph editor

diff --git a/Engine/Plugins/Experimental/ChaosClothAssetEditor/Source/ChaosClothAssetEditor/Private/ChaosClothAsset/ClothEditorMode.cpp b/Engine/Plugins/Experimental/ChaosClothAssetEditor/Source/ChaosClothAssetEditor/Private/ChaosClothAsset/ClothEditorMode.cpp
--- a/Engine/Plugins/Experimental/ChaosClothAssetEditor/Source/ChaosClothAssetEditor/Private/ChaosClothAsset/ClothEditorMode.cpp
+++ b/Engine/Plugins/Experimental/ChaosClothAssetEditor/Source/ChaosClothAssetEditor/Private/ChaosClothAsset/ClothEditorMode.cpp
@@ -908,41 +908,24 @@ void UChaosClothAssetEditorMode::SetDataflowGraphEditor(TSharedPtr<SDataflowGrap
 	}
 }
 
-void UChaosClothAssetEditorMode::OnDataflowNodeSelectionChanged(const TSet<UObject*>& NewSelection)
+void UChaosClothAssetEditorMode::StartToolForSelectedNode(const UObject* SelectedNode)
 {
-	using namespace UE::Chaos::ClothAsset;
-
 	// Start the tool associated with the selected node, if any
 
 	const TSharedRef<FUICommandList> CommandList = Toolkit->GetToolkitCommands();
-	const FChaosClothAssetEditorCommands& CommandInfos = FChaosClothAssetEditorCommands::Get();
+	const UE::Chaos::ClothAsset::FChaosClothAssetEditorCommands& CommandInfos = UE::Chaos::ClothAsset::FChaosClothAssetEditorCommands::Get();
 
-	bool bNewToolStarted = false;
-	if (NewSelection.Num() == 1)
+	if (const UDataflowEdNode* const EdNode = Cast<UDataflowEdNode>(SelectedNode))
 	{
-		if (const UDataflowEdNode* const Node = Cast<UDataflowEdNode>(*NewSelection.CreateConstIterator()))
+		if (const TSharedPtr<const FDataflowNode> DataflowNode = EdNode->GetDataflowNode())
 		{
-			if (const TSharedPtr<const FDataflowNode> DataflowNode = Node->GetDataflowNode())
+			const FName DataflowNodeName = DataflowNode->GetType();
+			if (const TSharedPtr<const FUICommandInfo>* const Command = NodeTypeToToolCommandMap.Find(DataflowNodeName))
 			{
-				const FName DataflowNodeName = DataflowNode->GetType();
-				if (NodeTypeToToolCommandMap.Contains(DataflowNodeName))
-				{
-					CommandList->TryExecuteAction(NodeTypeToToolCommandMap[DataflowNodeName].ToSharedRef());
-					bNewToolStarted = true;
-				}
+				CommandList->TryExecuteAction(Command->ToSharedRef());
 			}
 		}
 	}
-
-	UEditorInteractiveToolsContext* const ToolsContext = GetInteractiveToolsContext();
-	checkf(ToolsContext, TEXT("No valid ToolsContext found for UChaosClothAssetEditorMode"));
-	if (!bNewToolStarted && ToolsContext->HasActiveTool())
-	{
-		// The user has clicked away from the selected node, end the tool
-		// TODO: If we wanted we could probably auto-accept the pending changes in the active tool
-		ToolsContext->EndTool(EToolShutdownType::Completed);
-	}
-
 }
 
 
diff --git a/Engine/Plugins/Experimental/ChaosClothAssetEditor/Source/ChaosClothAssetEditor/Private/ChaosClothAsset/ClothEditorMode.h b/Engine/Plugins/Experimental/ChaosClothAssetEditor/Source/ChaosClothAssetEditor/Private/ChaosClothAsset/ClothEditorMode.h
--- a/Engine/Plugins/Experimental/ChaosClothAssetEditor/Source/ChaosClothAssetEditor/Private/ChaosClothAsset/ClothEditorMode.h
+++ b/Engine/Plugins/Experimental/ChaosClothAssetEditor/Source/ChaosClothAssetEditor/Private/ChaosClothAsset/ClothEditorMode.h
@@ -143,7 +143,7 @@ private:
 
 	void SetDataflowGraphEditor(TSharedPtr<SDataflowGraphEditor> InGraphEditor);
 	
-	void OnDataflowNodeSelectionChanged(const TSet<UObject*>& NewSelection);
+	void StartToolForSelectedNode(const UObject* SelectedNode);
 	void OnDataflowNodeDeleted(const TSet<UObject*>& DeletedNodes);
 
 	/**
diff --git a/Engine/Plugins/Experimental/ChaosClothAssetEditor/Source/ChaosClothAssetEditor/Private/ChaosClothAsset/ClothEditorToolkit.cpp b/Engine/Plugins/Experimental/ChaosClothAssetEditor/Source/ChaosClothAssetEditor/Private/ChaosClothAsset/ClothEditorToolkit.cpp
--- a/Engine/Plugins/Experimental/ChaosClothAssetEditor/Source/ChaosClothAssetEditor/Private/ChaosClothAsset/ClothEditorToolkit.cpp
+++ b/Engine/Plugins/Experimental/ChaosClothAssetEditor/Source/ChaosClothAssetEditor/Private/ChaosClothAsset/ClothEditorToolkit.cpp
@@ -786,6 +786,7 @@ TSharedRef<SDataflowGraphEditor> FChaosClothAssetEditorToolkit::CreateGraphEdito
 	SGraphEditor::FGraphEditorEvents InEvents;
 	InEvents.OnVerifyTextCommit = FOnNodeVerifyTextCommit::CreateSP(this, &FChaosClothAssetEditorToolkit::OnNodeVerifyTitleCommit);
 	InEvents.OnTextCommitted = FOnNodeTextCommitted::CreateSP(this, &FChaosClothAssetEditorToolkit::OnNodeTitleCommitted);
+	InEvents.OnNodeSingleClicked = SGraphEditor::FOnNodeSingleClicked::CreateSP(this, &FChaosClothAssetEditorToolkit::OnNodeSingleClicked);
 
 	TSharedRef<SDataflowGraphEditor> NewGraphEditor = SNew(SDataflowGraphEditor, Dataflow)
 		.GraphToEdit(Dataflow)
@@ -966,11 +967,20 @@ void FChaosClothAssetEditorToolkit::OnNodeSelectionChanged(const TSet<UObject*>&
 		Dataflow->LastModifiedRenderTarget = Dataflow::FTimestamp::Current();
 	}
 
+
 	UChaosClothAssetEditorMode* const ClothMode = CastChecked<UChaosClothAssetEditorMode>(EditorModeManager->GetActiveScriptableMode(UChaosClothAssetEditorMode::EM_ChaosClothAssetEditorModeId));
 	if (ClothMode)
 	{
+		// Close any running tool. OnNodeSingleClicked() will start a new tool if a new node was clicked.
+		UEditorInteractiveToolsContext* const ToolsContext = ClothMode->GetInteractiveToolsContext();
+		checkf(ToolsContext, TEXT("No valid ToolsContext found for FChaosClothAssetEditorToolkit"));
+		if (ToolsContext->HasActiveTool())
+		{
+			ToolsContext->EndTool(EToolShutdownType::Completed);
+		}
+
+		// Update the Construction viewport with the newly selected node's Collection
 		ClothMode->SetSelectedClothCollection(Collection);
-		ClothMode->OnDataflowNodeSelectionChanged(NewSelection);
 	}
 
 	if (Outliner)
@@ -979,6 +989,18 @@ void FChaosClothAssetEditorToolkit::OnNodeSelectionChanged(const TSet<UObject*>&
 	}
 }
 
+void FChaosClothAssetEditorToolkit::OnNodeSingleClicked(UObject* ClickedNode) const
+{
+	UChaosClothAssetEditorMode* const ClothMode = CastChecked<UChaosClothAssetEditorMode>(EditorModeManager->GetActiveScriptableMode(UChaosClothAssetEditorMode::EM_ChaosClothAssetEditorModeId));
+	if (ClothMode)
+	{
+		// Start the corresponding tool
+		UEditorInteractiveToolsContext* const ToolsContext = ClothMode->GetInteractiveToolsContext();
+		checkf(ToolsContext, TEXT("No valid ToolsContext found for FChaosClothAssetEditorToolkit"));
+		ClothMode->StartToolForSelectedNode(ClickedNode);
+	}
+}
+
 
 void FChaosClothAssetEditorToolkit::OnNodeDeleted(const TSet<UObject*>& DeletedNodes) const
 {
diff --git a/Engine/Plugins/Experimental/ChaosClothAssetEditor/Source/ChaosClothAssetEditor/Private/ChaosClothAsset/ClothEditorToolkit.h b/Engine/Plugins/Experimental/ChaosClothAssetEditor/Source/ChaosClothAssetEditor/Private/ChaosClothAsset/ClothEditorToolkit.h
--- a/Engine/Plugins/Experimental/ChaosClothAssetEditor/Source/ChaosClothAssetEditor/Private/ChaosClothAsset/ClothEditorToolkit.h
+++ b/Engine/Plugins/Experimental/ChaosClothAssetEditor/Source/ChaosClothAssetEditor/Private/ChaosClothAsset/ClothEditorToolkit.h
@@ -124,6 +124,7 @@ private:
 	void OnNodeTitleCommitted(const FText& InNewText, ETextCommit::Type InCommitType, UEdGraphNode* GraphNode) const;
 	void OnNodeSelectionChanged(const TSet<UObject*>& NewSelection);
 	void OnNodeDeleted(const TSet<UObject*>& DeletedNodes) const;
+	void OnNodeSingleClicked(UObject* ClickedNode) const;
 
 	/** Scene in which the 3D sim space preview meshes live. Ownership shared with AdvancedPreviewSettingsWidget*/
 	TSharedPtr<FChaosClothPreviewScene> ClothPreviewScene;
diff --git a/Engine/Source/Editor/GraphEditor/Private/SGraphEditorImpl.cpp b/Engine/Source/Editor/GraphEditor/Private/SGraphEditorImpl.cpp
--- a/Engine/Source/Editor/GraphEditor/Private/SGraphEditorImpl.cpp
+++ b/Engine/Source/Editor/GraphEditor/Private/SGraphEditorImpl.cpp
@@ -770,6 +770,7 @@ void SGraphEditorImpl::Construct( const FArguments& InArgs )
 			.ShowGraphStateOverlay(InArgs._ShowGraphStateOverlay)
 			.OnDoubleClicked(InArgs._GraphEvents.OnDoubleClicked)
 			.OnMouseButtonDown(InArgs._GraphEvents.OnMouseButtonDown)
+		    .OnNodeSingleClicked(InArgs._GraphEvents.OnNodeSingleClicked)
 		]
 
 		// Indicator of current zoom level
diff --git a/Engine/Source/Editor/GraphEditor/Private/SGraphPanel.cpp b/Engine/Source/Editor/GraphEditor/Private/SGraphPanel.cpp
--- a/Engine/Source/Editor/GraphEditor/Private/SGraphPanel.cpp
+++ b/Engine/Source/Editor/GraphEditor/Private/SGraphPanel.cpp
@@ -108,7 +108,8 @@ void SGraphPanel::Construct( const SGraphPanel::FArguments& InArgs )
 	this->OnDisallowedPinConnection = InArgs._OnDisallowedPinConnection;
 	this->OnDoubleClicked = InArgs._OnDoubleClicked;
 	this->OnClicked = InArgs._OnMouseButtonDown;
-	
+	this->OnNodeSingleClicked = InArgs._OnNodeSingleClicked;
+
 	this->bPreservePinPreviewConnection = false;
 	this->PinVisibility = SGraphEditor::Pin_Show;
 
diff --git a/Engine/Source/Editor/GraphEditor/Private/SNodePanel.cpp b/Engine/Source/Editor/GraphEditor/Private/SNodePanel.cpp
--- a/Engine/Source/Editor/GraphEditor/Private/SNodePanel.cpp
+++ b/Engine/Source/Editor/GraphEditor/Private/SNodePanel.cpp
@@ -1009,6 +1009,12 @@ FReply SNodePanel::OnMouseButtonUp( const FGeometry& MyGeometry, const FPointerE
 
 				// We're done interacting with this node.
 				NodeUnderMousePtr.Reset();
+
+				if (OnNodeSingleClicked.IsBound())
+				{
+					OnNodeSingleClicked.Execute(NodeWidgetUnderMouse->GetObjectBeingDisplayed());
+				}
+
 			}
 			else if (this->HasMouseCapture())
 			{
diff --git a/Engine/Source/Editor/GraphEditor/Public/SGraphPanel.h b/Engine/Source/Editor/GraphEditor/Public/SGraphPanel.h
--- a/Engine/Source/Editor/GraphEditor/Public/SGraphPanel.h
+++ b/Engine/Source/Editor/GraphEditor/Public/SGraphPanel.h
@@ -111,6 +111,7 @@ public:
 		SLATE_EVENT( SGraphEditor::FOnDisallowedPinConnection, OnDisallowedPinConnection )
 		SLATE_EVENT( SGraphEditor::FOnDoubleClicked, OnDoubleClicked )
 		SLATE_EVENT( SGraphEditor::FOnMouseButtonDown, OnMouseButtonDown )
+		SLATE_EVENT( SGraphEditor::FOnNodeSingleClicked, OnNodeSingleClicked )
 		//SLATE_ATTRIBUTE( FGraphAppearanceInfo, Appearance )
 	SLATE_END_ARGS()
 
diff --git a/Engine/Source/Editor/GraphEditor/Public/SNodePanel.h b/Engine/Source/Editor/GraphEditor/Public/SNodePanel.h
--- a/Engine/Source/Editor/GraphEditor/Public/SNodePanel.h
+++ b/Engine/Source/Editor/GraphEditor/Public/SNodePanel.h
@@ -1096,6 +1096,9 @@ protected:
 	/** Node positions pre-drag, used to limit transaction creation on drag */
 	TMap<TWeakPtr<SNode>, FVector2D> OriginalNodePositions;
 
+	/** Called when the user left clicks on a node without dragging */
+	SGraphEditor::FOnNodeSingleClicked OnNodeSingleClicked;
+
 private:
 	/** Active timer that handles deferred zooming until the target zoom is reached */
 	EActiveTimerReturnType HandleZoomToFit(double InCurrentTime, float InDeltaTime);
diff --git a/Engine/Source/Editor/UnrealEd/Public/GraphEditor.h b/Engine/Source/Editor/UnrealEd/Public/GraphEditor.h
--- a/Engine/Source/Editor/UnrealEd/Public/GraphEditor.h
+++ b/Engine/Source/Editor/UnrealEd/Public/GraphEditor.h
@@ -122,6 +122,8 @@ public:
 
 	DECLARE_DELEGATE( FOnDoubleClicked );
 
+	DECLARE_DELEGATE_OneParam( FOnNodeSingleClicked, UObject* );
+
 	DECLARE_DELEGATE_RetVal_TwoParams( FReply, FOnMouseButtonDown, const FGeometry&, const FPointerEvent& );
 
 	/** Info about events occurring in/on the graph */
@@ -155,6 +157,8 @@ public:
 		FOnDoubleClicked OnDoubleClicked;
 		/** Called when the graph is clicked */
 		FOnMouseButtonDown OnMouseButtonDown;
+		/** Called when a node is single-clicked without drag */
+		FOnNodeSingleClicked OnNodeSingleClicked;
 	};
 
 
