Improved the perfomance of launching skin weight transfer node:
-removed skinning weight visualization
-removed conversion of skeletal mesh to dynamic mesh and instead create skeletalmeshcomponent directly
-removed LOD field since we are not converting to dynamic mesh anymore
-removed unused imports

#rnx

diff --git a/Engine/Plugins/Experimental/ChaosClothAssetEditor/Source/ChaosClothAssetEditorTools/Private/ChaosClothAsset/ClothTransferSkinWeightsTool.cpp b/Engine/Plugins/Experimental/ChaosClothAssetEditor/Source/ChaosClothAssetEditorTools/Private/ChaosClothAsset/ClothTransferSkinWeightsTool.cpp
--- a/Engine/Plugins/Experimental/ChaosClothAssetEditor/Source/ChaosClothAssetEditorTools/Private/ChaosClothAsset/ClothTransferSkinWeightsTool.cpp
+++ b/Engine/Plugins/Experimental/ChaosClothAssetEditor/Source/ChaosClothAssetEditorTools/Private/ChaosClothAsset/ClothTransferSkinWeightsTool.cpp
@@ -3,12 +3,7 @@
 #include "ChaosClothAsset/ClothTransferSkinWeightsTool.h"
 
 #include "SkeletalMeshAttributes.h"
-#include "MeshDescriptionToDynamicMesh.h"
-#include "Rendering/SkeletalMeshLODImporterData.h"
-#include "Rendering/SkeletalMeshModel.h"
-#include "DynamicMesh/DynamicVertexSkinWeightsAttribute.h"
-#include "DynamicMesh/DynamicBoneAttribute.h"
-#include "Operations/TransferBoneWeights.h"
+#include "Engine/World.h"
 #include "ToolSetupUtil.h"
 #include "ModelingToolTargetUtil.h"
 #include "Engine/SkeletalMesh.h"
@@ -19,73 +14,33 @@
 #include "BaseGizmos/TransformGizmoUtil.h"
 #include "InteractiveToolObjects.h"
 #include "MeshOpPreviewHelpers.h"
+#include "Components/SkeletalMeshComponent.h"
 
 #define LOCTEXT_NAMESPACE "ClothTransferSkinWeightsTool"
 
-
 namespace UE::Chaos::ClothAsset::Private
 {
-	void SkeletalMeshToDynamicMesh(USkeletalMesh* FromSkeletalMeshAsset, int32 SourceLODIdx, FDynamicMesh3& ToDynamicMesh)
-	{
-		FMeshDescription SourceMesh;
-		FromSkeletalMeshAsset->GetMeshDescription(SourceLODIdx, SourceMesh);
-		FMeshDescriptionToDynamicMesh Converter;
-		Converter.Convert(&SourceMesh, ToDynamicMesh);
-	};
-
-
+	/** An empty operator in case we need to add a background compute operation in the future. */
 	class FClothTransferSkinWeightsOp : public UE::Geometry::FDynamicMeshOperator
 	{
 	public:
-
-		FClothTransferSkinWeightsOp(FDynamicMesh3&& OriginalTargetMesh, const TSharedPtr<const FDynamicMesh3> SourceMesh, const FTransform& SourceMeshTransform) :
-			SourceMesh(SourceMesh),
-			SourceMeshTransform(SourceMeshTransform)
+		FClothTransferSkinWeightsOp()
 		{
-			*ResultMesh = MoveTemp(OriginalTargetMesh);
 		}
 
 	private:
-
 		// FDynamicMeshOperator interface
 		virtual void CalculateResult(FProgressCancel* Progress) override
 		{
 			using namespace UE::Geometry;
-			using namespace UE::AnimationCore;
-
-			// Copy over bone attributes from the source mesh to the target/preview
-			ResultMesh->Attributes()->CopyBoneAttributes(*SourceMesh->Attributes());
-
-			// Do the transfer
-			FTransferBoneWeights TransferBoneWeights(SourceMesh.Get(), FSkeletalMeshAttributes::DefaultSkinWeightProfileName);
-			TransferBoneWeights.TargetToWorld = FTransformSRT3d(SourceMeshTransform.Inverse());
-			bool bComputeOK = false;
-			if (TransferBoneWeights.Validate() == EOperationValidationResult::Ok)
-			{
-				bComputeOK = TransferBoneWeights.TransferWeightsToMesh(*ResultMesh, FSkeletalMeshAttributes::DefaultSkinWeightProfileName);
-			}
 
 			FGeometryResult OpResult;
-			if (!bComputeOK)
-			{
-				OpResult.Result = EGeometryResultType::Failure;
-			}
-			else
-			{
-				OpResult.Result = EGeometryResultType::Success;
-			}
+			OpResult.Result = EGeometryResultType::Success;
 			SetResultInfo(OpResult);
 		}
-
-		// Inputs
-		const TSharedPtr<const FDynamicMesh3> SourceMesh;
-		const FTransform SourceMeshTransform;		
 	};
-
-
 }
 
-
 // ------------------- Builder -------------------
 
 USingleSelectionMeshEditingTool* UClothTransferSkinWeightsToolBuilder::CreateNewTool(const FToolBuilderState& SceneState) const
@@ -115,8 +70,6 @@ void UClothTransferSkinWeightsTool::Setup()
 
 	SetSRTPropertiesFromTransform(TransferSkinWeightsNode->Transform);
 	ToolProperties->SourceMesh = TransferSkinWeightsNode->SkeletalMesh;
-	ToolProperties->SourceMeshLOD = TransferSkinWeightsNode->LodIndex;
-	
 
 	AddToolPropertySource(ToolProperties);
 
@@ -128,7 +81,8 @@ void UClothTransferSkinWeightsTool::Setup()
 	TargetClothPreview = NewObject<UMeshOpPreviewWithBackgroundCompute>(this);
 	TargetClothPreview->Setup(GetTargetWorld(), this);
 	ToolSetupUtil::ApplyRenderingConfigurationToPreview(TargetClothPreview->PreviewMesh, Target);
-	TargetClothPreview->ConfigureMaterials(ToolSetupUtil::GetVertexColorMaterial(GetToolManager()), ToolSetupUtil::GetDefaultWorkingMaterial(GetToolManager()));
+	UMaterialInterface* const TargetMaterial = ToolSetupUtil::GetDefaultSculptMaterial(GetToolManager());
+	TargetClothPreview->ConfigureMaterials(TargetMaterial, TargetMaterial);
 
 	// Mesh topology is not being changed 
 	TargetClothPreview->SetIsMeshTopologyConstant(true, EMeshRenderAttributeFlags::VertexColors);
@@ -147,8 +101,8 @@ void UClothTransferSkinWeightsTool::Setup()
 	//
 
 	SourceMeshParentActor = GetTargetWorld()->SpawnActor<AInternalToolFrameworkActor>();
-	SourceMeshComponent = NewObject<UDynamicMeshComponent>(SourceMeshParentActor);
-	ToolSetupUtil::ApplyRenderingConfigurationToPreview(SourceMeshComponent, nullptr);
+	SourceMeshComponent = NewObject<USkeletalMeshComponent>(SourceMeshParentActor);
+	SourceMeshComponent->SetDisablePostProcessBlueprint(true); 
 	SourceMeshParentActor->SetRootComponent(SourceMeshComponent);
 	SourceMeshComponent->RegisterComponent();
 
@@ -156,40 +110,12 @@ void UClothTransferSkinWeightsTool::Setup()
 
 	ToolProperties->WatchProperty(ToolProperties->SourceMesh, [this](TObjectPtr<USkeletalMesh>) { UpdateSourceMesh(); });
 
-	ToolProperties->WatchProperty(ToolProperties->BoneName, [this](const FName&) 
-	{
-		check(TargetClothPreview);
-		TargetClothPreview->PreviewMesh->NotifyDeferredEditCompleted(UPreviewMesh::ERenderUpdateMode::FastUpdate, EMeshRenderAttributeFlags::VertexColors, false);
-	});
-	
 	ToolProperties->WatchProperty(ToolProperties->bHideSourceMesh, [this](bool bNewHideSourceMesh) 
 	{ 
 		check(SourceMeshComponent);
 		SourceMeshComponent->SetVisibility(!bNewHideSourceMesh);
 	});
 
-	ToolProperties->WatchProperty(ToolProperties->SourceMeshLOD, [this](int32 NewLOD)
-	{
-		const bool bLODIsValid = !ToolProperties->SourceMesh || ToolProperties->SourceMesh->IsValidLODIndex(NewLOD);
-
-		if (bLODIsValid)
-		{
-			if (bHasInvalidLODWarning)
-			{
-				GetToolManager()->DisplayMessage(FText(), EToolMessageLevel::UserWarning);    // clear old warning
-				bHasInvalidLODWarning = false;
-			}
-		}
-		else
-		{
-			GetToolManager()->DisplayMessage(LOCTEXT("InvalidLODIndex", "Specified source mesh LOD is not valid"), EToolMessageLevel::UserWarning);
-			bHasInvalidLODWarning = true;
-		}
-
-		UpdateSourceMesh();
-	});
-
-
 	//
 	// Transform/Gizmo/Proxy stuff
 	//
@@ -231,12 +157,6 @@ void UClothTransferSkinWeightsTool::Setup()
 		{
 			SourceMeshParentActor->SetActorTransform(NewTransform);
 		}
-
-		// Recompute result after moving the transform gizmo
-		if (ToolProperties->SourceMesh && ToolProperties->SourceMesh->IsValidLODIndex(ToolProperties->SourceMeshLOD))
-		{
-			TargetClothPreview->InvalidateResult();
-		}
 	});
 
 	SourceMeshTransformGizmo = UE::TransformGizmoUtil::CreateCustomTransformGizmo(GizmoManager, ETransformGizmoSubElements::StandardTranslateRotate, this);
@@ -254,7 +174,6 @@ void UClothTransferSkinWeightsTool::Setup()
 
 
 	UpdateSourceMesh();
-	SetPreviewMeshColorFunction();
 
 	UE::ToolTarget::HideSourceObject(Target);
 }
@@ -266,7 +185,6 @@ void UClothTransferSkinWeightsTool::Shutdown(EToolShutdownType ShutdownType)
 	if (ShutdownType == EToolShutdownType::Accept)
 	{
 		TransferSkinWeightsNode->SkeletalMesh = ToolProperties->SourceMesh;
-		TransferSkinWeightsNode->LodIndex = ToolProperties->SourceMeshLOD;
 		TransferSkinWeightsNode->Transform = TransformFromProperties();
 	}
 
@@ -307,7 +225,6 @@ bool UClothTransferSkinWeightsTool::CanAccept() const
 	const FTransform& TransformOnNode = TransferSkinWeightsNode->Transform;;
 
 	return (ToolProperties->SourceMesh != TransferSkinWeightsNode->SkeletalMesh) ||
-		(ToolProperties->SourceMeshLOD != TransferSkinWeightsNode->LodIndex) ||
 		(ToolProperties->SourceMeshRotation != TransformOnNode.Rotator().Euler()) ||
 		(ToolProperties->SourceMeshTranslation != TransformOnNode.GetTranslation()) ||
 		(ToolProperties->SourceMeshScale != TransformOnNode.GetScale3D());
@@ -324,16 +241,8 @@ void UClothTransferSkinWeightsTool::OnTick(float DeltaTime)
 
 TUniquePtr<UE::Geometry::FDynamicMeshOperator> UClothTransferSkinWeightsTool::MakeNewOperator()
 {
-	checkf(ToolProperties->SourceMesh, TEXT("Expected non-null SourceMesh before starting Operator"));
-	checkf(ToolProperties->SourceMesh->IsValidLODIndex(ToolProperties->SourceMeshLOD), TEXT("Source Mesh LOD level should be validated before starting Operator"));
-
-	UE::Geometry::FDynamicMesh3 TargetDynamicMesh = UE::ToolTarget::GetDynamicMeshCopy(Target, true);
-	const TSharedPtr<const UE::Geometry::FDynamicMesh3> SourceDynamicMesh = MakeShared<UE::Geometry::FDynamicMesh3>(*SourceMeshComponent->GetMesh());
-
-	const FTransform SourceMeshTransform = SourceMeshTransformGizmo->GetGizmoTransform();
 	TUniquePtr<UE::Chaos::ClothAsset::Private::FClothTransferSkinWeightsOp> TransferOp =
-		MakeUnique<UE::Chaos::ClothAsset::Private::FClothTransferSkinWeightsOp>(MoveTemp(TargetDynamicMesh), SourceDynamicMesh, SourceMeshTransform);
-
+		MakeUnique<UE::Chaos::ClothAsset::Private::FClothTransferSkinWeightsOp>();
 	return TransferOp;
 }
 
@@ -342,54 +251,6 @@ void UClothTransferSkinWeightsTool::SetClothEditorContextObject(TObjectPtr<UClot
 	ClothEditorContextObject = InClothEditorContextObject;
 }
 
-
-void UClothTransferSkinWeightsTool::SetPreviewMeshColorFunction()
-{
-	using namespace UE::AnimationCore;
-	using namespace UE::Geometry;
-
-	TargetClothPreview->PreviewMesh->SetTriangleColorFunction([this](const FDynamicMesh3* Mesh, int TriangleID) -> FColor
-	{
-		const FName& CurrentBoneName = ToolProperties->BoneName;
-		if (!TargetMeshBoneNameToIndex.Contains(CurrentBoneName))
-		{
-			return FColor::Black;
-		}
-
-		const FBoneIndexType CurrentBoneIndex = TargetMeshBoneNameToIndex[CurrentBoneName];
-		const FIndex3i Tri = Mesh->GetTriangle(TriangleID);
-
-		const FName ProfileName = FSkeletalMeshAttributes::DefaultSkinWeightProfileName; // always use default profile for now, later this will be set by the user
-		FDynamicMeshVertexSkinWeightsAttribute* Attribute = Mesh->Attributes()->GetSkinWeightsAttribute(ProfileName);
-		if (Attribute == nullptr)
-		{
-			FLinearColor Lin(1.0f, 0.3f, 0.3f, 1.0f);
-			return Lin.ToFColor(/*bSRGB = */ true);
-		}
-
-		float AvgWeight = 0.0f;
-		for (int32 VID = 0; VID < 3; ++VID)
-		{
-			const int32 VertexID = Tri[VID];
-			FBoneWeights Data;
-			Attribute->GetValue(VertexID, Data);
-			for (FBoneWeight Wt : Data)
-			{
-				if (Wt.GetBoneIndex() == CurrentBoneIndex)
-				{
-					AvgWeight += Wt.GetWeight();
-				}
-			}
-		}
-
-		AvgWeight /= 3.0f;
-		const FLinearColor Lin(AvgWeight, AvgWeight, AvgWeight, 1.0f);
-		return Lin.ToFColor(/*bSRGB = */ true);
-	},
-	UPreviewMesh::ERenderUpdateMode::FullUpdate);
-}
-
-
 FTransform UClothTransferSkinWeightsTool::TransformFromProperties() const
 {
 	const FRotator Rotation = FRotator::MakeFromEuler(ToolProperties->SourceMeshRotation);
@@ -408,25 +269,17 @@ void UClothTransferSkinWeightsTool::UpdateSourceMesh()
 {
 	checkf(ToolProperties, TEXT("ToolProperties is expected to be non-null. Be sure to run Setup() on this tool when it is created."));
 
-	if (ToolProperties->SourceMesh && ToolProperties->SourceMesh->IsValidLODIndex(ToolProperties->SourceMeshLOD))
+	if (ToolProperties->SourceMesh)
 	{
 		// Set up source mesh (from the SkeletalMesh)
-		FDynamicMesh3 SourceDynamicMesh;
-		UE::Chaos::ClothAsset::Private::SkeletalMeshToDynamicMesh(ToolProperties->SourceMesh, ToolProperties->SourceMeshLOD, SourceDynamicMesh);
-		SourceMeshComponent->SetMesh(MoveTemp(SourceDynamicMesh));
-
-		checkf(SourceMeshComponent, TEXT("Source mesh specified in the Tool Properties, but no SourceMesh exists"));
-		
+		SourceMeshComponent->SetSkeletalMeshAsset(ToolProperties->SourceMesh);	
 		SourceMeshParentActor->SetActorTransform(TransformFromProperties());
-		SourceMeshComponent->SetMaterial(0, ToolSetupUtil::GetTransparentSculptMaterial(GetToolManager(), FLinearColor::Red, 0.4, true));
 		SourceMeshComponent->SetVisibility(!ToolProperties->bHideSourceMesh);
 
 		// Use ReinitializeGizmoTransform rather than SetNewGizmoTransform to avoid having this on the undo stack
 		SourceMeshTransformGizmo->ReinitializeGizmoTransform(SourceMeshParentActor->GetActorTransform());
 		SourceMeshTransformGizmo->SetVisibility(!ToolProperties->bHideSourceMesh);
 		SourceMeshTransformGizmo->ActiveGizmoMode = EToolContextTransformGizmoMode::Combined;
-
-		TargetClothPreview->InvalidateResult();
 	}
 	else
 	{
@@ -440,7 +293,7 @@ void UClothTransferSkinWeightsTool::OpFinishedCallback(const UE::Geometry::FDyna
 {
 	if (Op->GetResultInfo().Result == UE::Geometry::EGeometryResultType::Failure)
 	{
-		GetToolManager()->DisplayMessage(LOCTEXT("TransferOpFailedWarning", "Weight transfer failed"), EToolMessageLevel::UserWarning);
+		GetToolManager()->DisplayMessage(LOCTEXT("TransferOpFailedWarning", "Operation failed"), EToolMessageLevel::UserWarning);
 		bHasOpFailedWarning = true;
 	}
 	else
@@ -456,58 +309,6 @@ void UClothTransferSkinWeightsTool::OpFinishedCallback(const UE::Geometry::FDyna
 
 void UClothTransferSkinWeightsTool::PreviewMeshUpdatedCallback(UMeshOpPreviewWithBackgroundCompute* Preview)
 {
-	constexpr bool bProcessOnlyIfValid = true;
-
-	Preview->ProcessCurrentMesh([&](const FDynamicMesh3& ResultMesh)
-	{
-		using namespace UE::Geometry;
-		using namespace UE::AnimationCore;
-
-		ToolProperties->BoneNameList.Reset();
-		TargetMeshBoneNameToIndex.Reset();
-
-		// Rebuild the set of selectable bone names, and the name -> index map
-		// TODO: Do we maybe want to do this in the background op and then copy the final results back to the member variables?
-
-		if (ResultMesh.Attributes() && ResultMesh.Attributes()->GetBoneNames() && ResultMesh.Attributes()->GetBoneParentIndices())
-		{
-			// Get set of bone indices used in the target mesh
-			const TArray<FName>& TargetBoneNames = ResultMesh.Attributes()->GetBoneNames()->GetAttribValues();
-			const TArray<int32>& TargetBoneIndices = ResultMesh.Attributes()->GetBoneParentIndices()->GetAttribValues();
-
-			const TMap<FName, TUniquePtr<FDynamicMeshVertexSkinWeightsAttribute>>& WeightLayers = ResultMesh.Attributes()->GetSkinWeightsAttributes();
-			TMap<FName, FBoneIndexType> UsedBoneNames;
-			for (const TPair<FName, TUniquePtr<FDynamicMeshVertexSkinWeightsAttribute>>& LayerPair : WeightLayers)
-			{
-				const FDynamicMeshVertexSkinWeightsAttribute* Layer = LayerPair.Value.Get();
-
-				for (const int VertexID : ResultMesh.VertexIndicesItr())
-				{
-					FBoneWeights Data;
-					Layer->GetValue(VertexID, Data);
-					for (FBoneWeight Wt : Data)
-					{
-						const FBoneIndexType BoneIndex = Wt.GetBoneIndex();
-						const FName& BoneName = TargetBoneNames[BoneIndex];
-						UsedBoneNames.Add(BoneName, BoneIndex);
-					}
-				}
-			}
-
-			// Build the bone name -> index map
-			for (int32 BoneID = 0; BoneID < TargetBoneNames.Num(); ++BoneID)
-			{
-				const FName& BoneName = TargetBoneNames[BoneID];
-				const int32 BoneIndex = TargetBoneIndices[BoneID];
-				TargetMeshBoneNameToIndex.Add(BoneName, BoneID);
-			}
-
-			// Update list of bone names in the Properties panel
-			UsedBoneNames.ValueSort([](int16 A, int16 B) { return A < B; });
-			UsedBoneNames.GetKeys(ToolProperties->BoneNameList);
-		}
-	}, 
-	bProcessOnlyIfValid);
 }
 
 #undef LOCTEXT_NAMESPACE
diff --git a/Engine/Plugins/Experimental/ChaosClothAssetEditor/Source/ChaosClothAssetEditorTools/Public/ChaosClothAsset/ClothTransferSkinWeightsTool.h b/Engine/Plugins/Experimental/ChaosClothAssetEditor/Source/ChaosClothAssetEditorTools/Public/ChaosClothAsset/ClothTransferSkinWeightsTool.h
--- a/Engine/Plugins/Experimental/ChaosClothAssetEditor/Source/ChaosClothAssetEditorTools/Public/ChaosClothAsset/ClothTransferSkinWeightsTool.h
+++ b/Engine/Plugins/Experimental/ChaosClothAssetEditor/Source/ChaosClothAssetEditorTools/Public/ChaosClothAsset/ClothTransferSkinWeightsTool.h
@@ -3,13 +3,11 @@
 #pragma once
 
 #include "BaseTools/SingleSelectionMeshEditingTool.h"
-#include "BoneIndices.h"
 #include "ModelingOperators.h"
 #include "Transforms/TransformGizmoDataBinder.h"
-
-
 #include "ClothTransferSkinWeightsTool.generated.h"
 
+
 class UClothTransferSkinWeightsTool;
 class USkeletalMesh;
 class UClothEditorContextObject;
@@ -17,7 +15,7 @@ class UTransformProxy;
 class UCombinedTransformGizmo;
 class UMeshOpPreviewWithBackgroundCompute;
 class AInternalToolFrameworkActor;
-class UDynamicMeshComponent;
+class USkeletalMeshComponent;
 class FTransformGizmoDataBinder;
 struct FChaosClothAssetTransferSkinWeightsNode;
 
@@ -39,26 +37,8 @@ public:
 	UPROPERTY(EditAnywhere, Category = Source, meta = (EditCondition = "SourceMesh != nullptr"))
 	FVector3d SourceMeshScale;
 
-
-	UPROPERTY(EditAnywhere, Category = Source)
-	int32 SourceMeshLOD = 0;
-
 	UPROPERTY(EditAnywhere, Category = Source)
 	bool bHideSourceMesh = false;
-
-	UPROPERTY(EditAnywhere, Category = Visualization, meta = (GetOptions = GetBoneNameList))
-	FName BoneName;
-
-	// Get the list of valid bone names
-	UFUNCTION()
-	TArray<FName> GetBoneNameList()
-	{
-		return BoneNameList;
-	}
-
-	UPROPERTY(meta = (TransientToolProperty))
-	TArray<FName> BoneNameList;
-
 };
 
 
@@ -95,8 +75,6 @@ private:
 	
 	void SetClothEditorContextObject(TObjectPtr<UClothEditorContextObject> InClothEditorContextObject);
 
-	void SetPreviewMeshColorFunction();
-
 	FTransform TransformFromProperties() const;
 	void SetSRTPropertiesFromTransform(const FTransform& Transform) const;
 
@@ -120,7 +98,7 @@ private:
 	TObjectPtr<AInternalToolFrameworkActor> SourceMeshParentActor;
 
 	UPROPERTY(Transient)
-	TObjectPtr<UDynamicMeshComponent> SourceMeshComponent;
+	TObjectPtr<USkeletalMeshComponent> SourceMeshComponent;
 
 	// Source mesh transform gizmo support
 	UPROPERTY(Transient)
@@ -131,12 +109,8 @@ private:
 
 	TSharedPtr<FTransformGizmoDataBinder> DataBinder;
 
-	// Used to lookup the index of the currently selected-by-name bone
-	TMap<FName, FBoneIndexType> TargetMeshBoneNameToIndex;
-
 	FChaosClothAssetTransferSkinWeightsNode* TransferSkinWeightsNode = nullptr;
 
-	bool bHasInvalidLODWarning = false;
 	bool bHasOpFailedWarning = false;
 };
 
