Cloth Editor: fix for "Save As" not updating the ClothAsset with the new name of the Dataflow asset property

diff --git a/Engine/Plugins/Experimental/ChaosClothAssetEditor/Source/ChaosClothAssetEditor/Private/ChaosClothAsset/ClothEditorToolkit.cpp b/Engine/Plugins/Experimental/ChaosClothAssetEditor/Source/ChaosClothAssetEditor/Private/ChaosClothAsset/ClothEditorToolkit.cpp
--- a/Engine/Plugins/Experimental/ChaosClothAssetEditor/Source/ChaosClothAssetEditor/Private/ChaosClothAsset/ClothEditorToolkit.cpp
+++ b/Engine/Plugins/Experimental/ChaosClothAssetEditor/Source/ChaosClothAssetEditor/Private/ChaosClothAsset/ClothEditorToolkit.cpp
@@ -41,6 +41,7 @@
 #include "Widgets/Layout/SSpacer.h"
 #include "DynamicMesh/DynamicMesh3.h"
 #include "Toolkits/AssetEditorToolkitMenuContext.h"
+#include "FileHelpers.h"
 
 #define LOCTEXT_NAMESPACE "ChaosClothAssetEditorToolkit"
 
@@ -506,6 +507,41 @@ void FChaosClothAssetEditorToolkit::GetSaveableObjects(TArray<UObject*>& OutObje
 	}
 }
 
+bool FChaosClothAssetEditorToolkit::ShouldReopenEditorForSavedAsset(const UObject* Asset) const
+{
+	// "Save As" will potentially save the Dataflow asset with a new name, along with the cloth asset. 
+	// We don't really want to open a new Dataflow editor in that case, just the cloth editor
+	return Asset->IsA<UChaosClothAsset>();
+}
+
+void FChaosClothAssetEditorToolkit::OnAssetsSavedAs(const TArray<UObject*>& SavedObjects)
+{
+	// Set the Dataflow property on the Cloth object to point to the new DataflowAsset
+	UDataflow* NewDataflowAsset = nullptr;
+	UChaosClothAsset* NewClothAsset = nullptr;
+	for (UObject* const SavedObj : SavedObjects)
+	{
+		if (SavedObj->IsA<UDataflow>())
+		{
+			NewDataflowAsset = Cast<UDataflow>(SavedObj);
+		}
+		else if (SavedObj->IsA<UChaosClothAsset>())
+		{
+			NewClothAsset = Cast<UChaosClothAsset>(SavedObj);
+		}
+	}
+
+	if (NewClothAsset && NewDataflowAsset)
+	{
+		NewClothAsset->DataflowAsset = NewDataflowAsset;
+
+		// Now save the new Cloth asset again since we've updated its Property
+		const TArray<UPackage*> PackagesToSave{NewClothAsset->GetOutermost()};
+		constexpr bool bPromptToSave = false;
+		FEditorFileUtils::PromptForCheckoutAndSave(PackagesToSave, bCheckDirtyOnAssetSave, bPromptToSave);
+	}
+}
+
 //~ End FAssetEditorToolkit overrides
 
 
diff --git a/Engine/Plugins/Experimental/ChaosClothAssetEditor/Source/ChaosClothAssetEditor/Private/ChaosClothAsset/ClothEditorToolkit.h b/Engine/Plugins/Experimental/ChaosClothAssetEditor/Source/ChaosClothAssetEditor/Private/ChaosClothAsset/ClothEditorToolkit.h
--- a/Engine/Plugins/Experimental/ChaosClothAssetEditor/Source/ChaosClothAssetEditor/Private/ChaosClothAsset/ClothEditorToolkit.h
+++ b/Engine/Plugins/Experimental/ChaosClothAssetEditor/Source/ChaosClothAssetEditor/Private/ChaosClothAsset/ClothEditorToolkit.h
@@ -85,6 +85,8 @@ private:
 	virtual void PostInitAssetEditor() override;
 	virtual void InitToolMenuContext(FToolMenuContext& MenuContext) override;
 	virtual void GetSaveableObjects(TArray<UObject*>& OutObjects) const override;
+	virtual bool ShouldReopenEditorForSavedAsset(const UObject* Asset) const override;
+	virtual void OnAssetsSavedAs(const TArray<UObject*>& SavedObjects) override;
 
 	// IAssetEditorInstance
 	virtual bool IsPrimaryEditor() const override { return true; };
diff --git a/Engine/Source/Editor/UnrealEd/Private/Toolkits/AssetEditorToolkit.cpp b/Engine/Source/Editor/UnrealEd/Private/Toolkits/AssetEditorToolkit.cpp
--- a/Engine/Source/Editor/UnrealEd/Private/Toolkits/AssetEditorToolkit.cpp
+++ b/Engine/Source/Editor/UnrealEd/Private/Toolkits/AssetEditorToolkit.cpp
@@ -692,7 +692,10 @@ void FAssetEditorToolkit::SaveAssetAs_Execute()
 	}
 	for (auto Object : SavedObjects)
 	{
-		ObjectsToReopen.AddUnique(Object);
+		if (ShouldReopenEditorForSavedAsset(Object))
+		{
+			ObjectsToReopen.AddUnique(Object);
+		}
 	}
 	for (auto Object : EditingObjects)
 	{
@@ -701,6 +704,8 @@ void FAssetEditorToolkit::SaveAssetAs_Execute()
 	}
 	AssetEditorSubsystem->OpenEditorForAssets_Advanced(ObjectsToReopen, ToolkitMode, MyToolkitHost.ToSharedRef());
 	// end hack
+
+	OnAssetsSavedAs(SavedObjects);
 }
 
 
diff --git a/Engine/Source/Editor/UnrealEd/Public/Toolkits/AssetEditorToolkit.h b/Engine/Source/Editor/UnrealEd/Public/Toolkits/AssetEditorToolkit.h
--- a/Engine/Source/Editor/UnrealEd/Public/Toolkits/AssetEditorToolkit.h
+++ b/Engine/Source/Editor/UnrealEd/Public/Toolkits/AssetEditorToolkit.h
@@ -326,6 +326,12 @@ protected:
 	/** Called when "Save As" is clicked for this asset */
 	UNREALED_API virtual void SaveAssetAs_Execute();
 
+	/** Whether to open a new asset editor for the newly-saved Asset. Called from SaveAssetAs_Execute for all newly-saved Assets sequentially. */
+	virtual bool ShouldReopenEditorForSavedAsset(const UObject* Asset) const { return true; }
+
+	/** Called from SaveAssetAs_Execute when assets have been saved */
+	UNREALED_API virtual void OnAssetsSavedAs(const TArray<UObject*>& SavedObjects) {}
+
 	/** Called to test if "Find in Content Browser" should be enabled for this asset */
 	virtual bool CanFindInContentBrowser() const { return true; }
 
