Added an input option to specify the name of the weight map containing the inpainting information.

#rnx

diff --git a/Engine/Plugins/Experimental/ChaosClothAssetEditor/Source/ChaosClothAssetDataflowNodes/Private/ChaosClothAsset/ChaosClothAssetDataflowNodesModule.cpp b/Engine/Plugins/Experimental/ChaosClothAssetEditor/Source/ChaosClothAssetDataflowNodes/Private/ChaosClothAsset/ChaosClothAssetDataflowNodesModule.cpp
--- a/Engine/Plugins/Experimental/ChaosClothAssetEditor/Source/ChaosClothAssetDataflowNodes/Private/ChaosClothAsset/ChaosClothAssetDataflowNodesModule.cpp
+++ b/Engine/Plugins/Experimental/ChaosClothAssetEditor/Source/ChaosClothAssetDataflowNodes/Private/ChaosClothAsset/ChaosClothAssetDataflowNodesModule.cpp
@@ -107,6 +107,7 @@ void FChaosClothAssetDataflowNodesModule::StartupModule()
 	{
 		PropertyModule->RegisterCustomPropertyTypeLayout("ChaosClothAssetWeightedValue", FOnGetPropertyTypeCustomizationInstance::CreateStatic(&FWeightedValueCustomization::MakeInstance));
 		PropertyModule->RegisterCustomPropertyTypeLayout("ChaosClothAssetWeightedValueNonAnimatable", FOnGetPropertyTypeCustomizationInstance::CreateStatic(&FWeightedValueCustomization::MakeInstance));
+		PropertyModule->RegisterCustomPropertyTypeLayout("ChaosClothAssetWeightedValueNonAnimatableNoLowHighRange", FOnGetPropertyTypeCustomizationInstance::CreateStatic(&FWeightedValueCustomization::MakeInstance));
 	}
 }
 
@@ -117,6 +118,7 @@ void FChaosClothAssetDataflowNodesModule::ShutdownModule()
 	{
 		PropertyModule->UnregisterCustomPropertyTypeLayout("ChaosClothAssetWeightedValue");
 		PropertyModule->UnregisterCustomPropertyTypeLayout("ChaosClothAssetWeightedValueNonAnimatable");
+		PropertyModule->UnregisterCustomPropertyTypeLayout("ChaosClothAssetWeightedValueNonAnimatableNoLowHighRange");
 	}
 }
 
diff --git a/Engine/Plugins/Experimental/ChaosClothAssetEditor/Source/ChaosClothAssetDataflowNodes/Private/ChaosClothAsset/TransferSkinWeightsNode.cpp b/Engine/Plugins/Experimental/ChaosClothAssetEditor/Source/ChaosClothAssetDataflowNodes/Private/ChaosClothAsset/TransferSkinWeightsNode.cpp
--- a/Engine/Plugins/Experimental/ChaosClothAssetEditor/Source/ChaosClothAssetDataflowNodes/Private/ChaosClothAsset/TransferSkinWeightsNode.cpp
+++ b/Engine/Plugins/Experimental/ChaosClothAssetEditor/Source/ChaosClothAssetDataflowNodes/Private/ChaosClothAsset/TransferSkinWeightsNode.cpp
@@ -47,7 +47,7 @@ namespace UE::Chaos::ClothAsset::Private
 		// Convert the sim mesh to DynamicMesh. 
 		// @todo: FTransferBoneWeights should accept raw data arrays (vertices/triangles) to avoid this conversion
 		FClothPatternToDynamicMesh PatternToDynamicMesh;
-		constexpr bool bDisableAttributes = true;
+		constexpr bool bDisableAttributes = false;
 		PatternToDynamicMesh.Convert(ClothCollection, INDEX_NONE, EClothPatternVertexType::Sim3D, WeldedSimMesh, bDisableAttributes);
 
 		WeldedSimMesh.EnableAttributes();
@@ -74,6 +74,7 @@ namespace UE::Chaos::ClothAsset::Private
 		const double SmoothingStrength,
 		bool bUseParallel,
 		const TSharedRef<FManagedArrayCollection>& ClothCollection,
+		const FString& InpaintMaskWeightMapName,
 		UE::Geometry::FTransferBoneWeights& TransferBoneWeights)
 	{
 		using namespace UE::Geometry;
@@ -100,6 +101,7 @@ namespace UE::Chaos::ClothAsset::Private
 		TransferBoneWeights.NumSmoothingIterations = NumSmoothingIterations;
 		TransferBoneWeights.SmoothingStrength = SmoothingStrength;
 		TransferBoneWeights.LayeredMeshSupport = true; // multilayerd clothing
+		TransferBoneWeights.ForceInpaintWeightMapName = FName(InpaintMaskWeightMapName);
 
 		if (!ensure(TransferBoneWeights.Validate() == EOperationValidationResult::Ok))
 		{
@@ -312,6 +314,7 @@ FChaosClothAssetTransferSkinWeightsNode::FChaosClothAssetTransferSkinWeightsNode
 {
 	RegisterInputConnection(&Collection);
 	RegisterOutputConnection(&Collection, &Collection);
+	RegisterInputConnection(&InpaintMask.WeightMap, GET_MEMBER_NAME_CHECKED(FChaosClothAssetWeightedValueNonAnimatableNoLowHighRange, WeightMap));
 }
 
 void FChaosClothAssetTransferSkinWeightsNode::Evaluate(Dataflow::FContext& Context, const FDataflowOutput* Out) const
@@ -340,6 +343,8 @@ void FChaosClothAssetTransferSkinWeightsNode::Evaluate(Dataflow::FContext& Conte
 				return;
 			}
 
+			InpaintMask.WeightMap_Override = GetValue<FString>(Context, &InpaintMask.WeightMap, FString());
+
 			//
 			// Convert source Skeletal Mesh to Dynamic Mesh.
 			//
@@ -365,7 +370,8 @@ void FChaosClothAssetTransferSkinWeightsNode::Evaluate(Dataflow::FContext& Conte
 			bool bTransferResult = false;
 			if (TransferMethod == EChaosClothAssetTransferSkinWeightsMethod::InpaintWeights)
 			{
-				bTransferResult = TransferInpaintWeights(TargetRefSkeleton, NormalThreshold, RadiusPercentage, NumSmoothingIterations, SmoothingStrength, bUseParallel, ClothCollection, TransferBoneWeights);
+				const FString& InpainMaskWeightMapName = GetValue<FString>(Context, &InpaintMask.WeightMap);
+				bTransferResult = TransferInpaintWeights(TargetRefSkeleton, NormalThreshold, RadiusPercentage, NumSmoothingIterations, SmoothingStrength, bUseParallel, ClothCollection, InpainMaskWeightMapName, TransferBoneWeights);
 			}
 			else if (TransferMethod == EChaosClothAssetTransferSkinWeightsMethod::ClosestPointOnSurface)
 			{
diff --git a/Engine/Plugins/Experimental/ChaosClothAssetEditor/Source/ChaosClothAssetDataflowNodes/Public/ChaosClothAsset/TransferSkinWeightsNode.h b/Engine/Plugins/Experimental/ChaosClothAssetEditor/Source/ChaosClothAssetDataflowNodes/Public/ChaosClothAsset/TransferSkinWeightsNode.h
--- a/Engine/Plugins/Experimental/ChaosClothAssetEditor/Source/ChaosClothAssetDataflowNodes/Public/ChaosClothAsset/TransferSkinWeightsNode.h
+++ b/Engine/Plugins/Experimental/ChaosClothAssetEditor/Source/ChaosClothAssetDataflowNodes/Public/ChaosClothAsset/TransferSkinWeightsNode.h
@@ -4,6 +4,7 @@
 
 #include "Dataflow/DataflowNode.h"
 #include "GeometryCollection/ManagedArrayCollection.h"
+#include "ChaosClothAsset/WeightedValue.h"
 #include "TransferSkinWeightsNode.generated.h"
 
 class USkeletalMesh;
@@ -73,6 +74,10 @@ public:
 	UPROPERTY(EditAnywhere, Category = "Transfer Skin Weights|Transfer Method", Meta = (UIMin = 0, UIMax = 1, ClampMin = 0, ClampMax = 1, DisplayName = "Smoothing Strength", EditCondition="TransferMethod==EChaosClothAssetTransferSkinWeightsMethod::InpaintWeights"))
 	float SmoothingStrength = 0.1;
 
+    /** Optional mask where a non-zero value indicates that we want the skinning weights for the vertex to be computed automatically instead of it being copied over from the source mesh. */
+	UPROPERTY(EditAnywhere, Category = "Transfer Skin Weights|Transfer Method", Meta = (DisplayName = "Inpaint Mask", EditCondition="TransferMethod==EChaosClothAssetTransferSkinWeightsMethod::InpaintWeights"))
+	FChaosClothAssetWeightedValueNonAnimatableNoLowHighRange InpaintMask = { TEXT("InpaintMask") };
+
 	FChaosClothAssetTransferSkinWeightsNode(const Dataflow::FNodeParameters& InParam, FGuid InGuid = FGuid::NewGuid());
 
 private:
diff --git a/Engine/Plugins/Experimental/ChaosClothAssetEditor/Source/ChaosClothAssetDataflowNodes/Public/ChaosClothAsset/WeightedValue.h b/Engine/Plugins/Experimental/ChaosClothAssetEditor/Source/ChaosClothAssetDataflowNodes/Public/ChaosClothAsset/WeightedValue.h
--- a/Engine/Plugins/Experimental/ChaosClothAssetEditor/Source/ChaosClothAssetDataflowNodes/Public/ChaosClothAsset/WeightedValue.h
+++ b/Engine/Plugins/Experimental/ChaosClothAssetEditor/Source/ChaosClothAssetDataflowNodes/Public/ChaosClothAsset/WeightedValue.h
@@ -81,3 +81,18 @@ struct FChaosClothAssetWeightedValueNonAnimatable
 	UPROPERTY(VisibleAnywhere, Transient, Category = "Weighted Value")
 	mutable FString WeightMap_Override;  // _Override has a special meaning to the property customization, mutable because this property is set while getting the original value
 };
+
+
+USTRUCT()
+struct FChaosClothAssetWeightedValueNonAnimatableNoLowHighRange
+{
+	GENERATED_BODY()
+
+	/** The name of the weight map for this property. */
+	UPROPERTY(EditAnywhere, Category = "Weighted Value", Meta = (DataflowInput))
+	FString WeightMap = TEXT("WeightMap");
+
+	/** The weight map override value for when the WeightMap has a connection that replaces the provided weight map value. */
+	UPROPERTY(VisibleAnywhere, Transient, Category = "Weighted Value")
+	mutable FString WeightMap_Override;  // _Override has a special meaning to the property customization, mutable because this property is set while getting the original value
+};
\ No newline at end of file
