Cloth Editor: when starting a tool, check if the current view mode is supported by the tool. If not, change to a valid view mode before starting the tool and then change back after the tool ends.
- Added a new IChaosClothAssetEditorToolBuilder interface for defining which view modes are supported for each tool

diff --git a/Engine/Plugins/Experimental/ChaosClothAssetEditor/Source/ChaosClothAssetEditor/Private/ChaosClothAsset/ClothEditorMode.cpp b/Engine/Plugins/Experimental/ChaosClothAssetEditor/Source/ChaosClothAssetEditor/Private/ChaosClothAsset/ClothEditorMode.cpp
--- a/Engine/Plugins/Experimental/ChaosClothAssetEditor/Source/ChaosClothAssetEditor/Private/ChaosClothAsset/ClothEditorMode.cpp
+++ b/Engine/Plugins/Experimental/ChaosClothAssetEditor/Source/ChaosClothAssetEditor/Private/ChaosClothAsset/ClothEditorMode.cpp
@@ -63,6 +63,7 @@
 #include "MaterialDomain.h"
 #include "Dataflow/DataflowSNode.h"
 #include "Components/SkeletalMeshComponent.h"
+#include "ChaosClothAsset/ClothEditorToolBuilder.h"
 
 #include UE_INLINE_GENERATED_CPP_BY_NAME(ClothEditorMode)
 
@@ -134,8 +135,12 @@ void UChaosClothAssetEditorMode::AddToolTargetFactories()
 	GetInteractiveToolsContext()->TargetManager->AddTargetFactory(NewObject<UClothComponentToolTargetFactory>(GetToolManager()));
 }
 
-
-void UChaosClothAssetEditorMode::RegisterClothTool(TSharedPtr<FUICommandInfo> UICommand, FString ToolIdentifier, UInteractiveToolBuilder* Builder, UEditorInteractiveToolsContext* const ToolsContext, EToolsContextScope ToolScope)
+void UChaosClothAssetEditorMode::RegisterClothTool(TSharedPtr<FUICommandInfo> UICommand, 
+	FString ToolIdentifier, 
+	UInteractiveToolBuilder* Builder, 
+	const UE::Chaos::ClothAsset::IChaosClothAssetEditorToolBuilder* ClothToolBuilder,
+	UEditorInteractiveToolsContext* const ToolsContext, 
+	EToolsContextScope ToolScope)
 {
 	if (!Toolkit.IsValid())
 	{
@@ -158,8 +163,18 @@ void UChaosClothAssetEditorMode::RegisterClothTool(TSharedPtr<FUICommandInfo> UI
 	const TSharedRef<FUICommandList>& CommandList = Toolkit->GetToolkitCommands();
 
 	CommandList->MapAction(UICommand,
-		FExecuteAction::CreateWeakLambda(ToolsContext, [this, ToolsContext, ToolIdentifier]()
+		FExecuteAction::CreateWeakLambda(ToolsContext, [this, ToolsContext, ToolIdentifier, ClothToolBuilder]()
 		{
+			TArray<UE::Chaos::ClothAsset::EClothPatternVertexType> SupportedModes;
+			ClothToolBuilder->GetSupportedViewModes(SupportedModes);
+
+			if (SupportedModes.Num() > 0 && !SupportedModes.Contains(this->GetConstructionViewMode()))
+			{
+				// Need to switch view modes before starting the tool
+				SavedConstructionViewMode = this->GetConstructionViewMode();	// remember the current view mode so we can restore it later
+				this->SetConstructionViewMode(SupportedModes[0]);
+			}
+
 			ActiveToolsContext = ToolsContext;
 			ToolsContext->StartTool(ToolIdentifier);
 		}),
@@ -212,15 +227,16 @@ void UChaosClothAssetEditorMode::RegisterTools()
 
 	UEditorInteractiveToolsContext* const ConstructionViewportToolsContext = GetInteractiveToolsContext();
 
-	RegisterClothTool(CommandInfos.BeginAttributeEditorTool, FChaosClothAssetEditorCommands::BeginAttributeEditorToolIdentifier, NewObject<UAttributeEditorToolBuilder>(), ConstructionViewportToolsContext);
-
-	RegisterClothTool(CommandInfos.BeginWeightMapPaintTool, FChaosClothAssetEditorCommands::BeginWeightMapPaintToolIdentifier, NewObject<UClothEditorWeightMapPaintToolBuilder>(), ConstructionViewportToolsContext);
+	UClothEditorWeightMapPaintToolBuilder* WeightMapPaintToolBuilder = NewObject<UClothEditorWeightMapPaintToolBuilder>();
+	RegisterClothTool(CommandInfos.BeginWeightMapPaintTool, FChaosClothAssetEditorCommands::BeginWeightMapPaintToolIdentifier, WeightMapPaintToolBuilder, WeightMapPaintToolBuilder, ConstructionViewportToolsContext);
 	RegisterAddNodeCommand(CommandInfos.AddWeightMapNode, FChaosClothAssetAddWeightMapNode::StaticType(), CommandInfos.BeginWeightMapPaintTool);
 
-	RegisterClothTool(CommandInfos.BeginTransferSkinWeightsTool, FChaosClothAssetEditorCommands::BeginTransferSkinWeightsToolIdentifier, NewObject<UClothTransferSkinWeightsToolBuilder>(), ConstructionViewportToolsContext);
+	UClothTransferSkinWeightsToolBuilder* TransferToolBuilder = NewObject<UClothTransferSkinWeightsToolBuilder>();
+	RegisterClothTool(CommandInfos.BeginTransferSkinWeightsTool, FChaosClothAssetEditorCommands::BeginTransferSkinWeightsToolIdentifier, TransferToolBuilder, TransferToolBuilder, ConstructionViewportToolsContext);
 	RegisterAddNodeCommand(CommandInfos.AddTransferSkinWeightsNode, FChaosClothAssetTransferSkinWeightsNode::StaticType(), CommandInfos.BeginTransferSkinWeightsTool);
 
-	RegisterClothTool(CommandInfos.BeginMeshSelectionTool, FChaosClothAssetEditorCommands::BeginMeshSelectionToolIdentifier, NewObject<UClothMeshSelectionToolBuilder>(), ConstructionViewportToolsContext);
+	UClothMeshSelectionToolBuilder* SelectionToolBuilder = NewObject<UClothMeshSelectionToolBuilder>();
+	RegisterClothTool(CommandInfos.BeginMeshSelectionTool, FChaosClothAssetEditorCommands::BeginMeshSelectionToolIdentifier, SelectionToolBuilder, SelectionToolBuilder, ConstructionViewportToolsContext);
 	RegisterAddNodeCommand(CommandInfos.AddMeshSelectionNode, FChaosClothAssetSelectionNode::StaticType(), CommandInfos.BeginMeshSelectionTool);
 }
 
@@ -261,6 +277,12 @@ void UChaosClothAssetEditorMode::OnToolEnded(UInteractiveToolManager* Manager, U
 {
 	UE::Chaos::ClothAsset::FChaosClothAssetEditorCommands::UpdateToolCommandBinding(Tool, ToolCommandList, true);
 
+	if (SavedConstructionViewMode)
+	{
+		SetConstructionViewMode(*SavedConstructionViewMode);
+		SavedConstructionViewMode.Reset();
+	}
+
 	bCanChangeConstructionViewMode = true;
 
 	ReinitializeDynamicMeshComponents();
@@ -754,6 +776,21 @@ void UChaosClothAssetEditorMode::ModeTick(float DeltaTime)
 		bShouldClearTeleportFlag = true;		// clear the flag next tick
 	}
 
+
+	if (!NodeTypeForPendingToolStart.IsNone() && !GetToolManager()->HasActiveTool(EToolSide::Left))
+	{
+		const TSharedRef<FUICommandList> CommandList = Toolkit->GetToolkitCommands();
+		const UE::Chaos::ClothAsset::FChaosClothAssetEditorCommands& CommandInfos = UE::Chaos::ClothAsset::FChaosClothAssetEditorCommands::Get();
+
+		if (const TSharedPtr<const FUICommandInfo>* const Command = NodeTypeToToolCommandMap.Find(NodeTypeForPendingToolStart))
+		{
+			CommandList->TryExecuteAction(Command->ToSharedRef());
+		}
+
+		NodeTypeForPendingToolStart = FName();
+	}
+
+
 	if (PreviewScene->GetWorld())
 	{
 		PreviewScene->GetWorld()->Tick(ELevelTick::LEVELTICK_All, DeltaTime);
@@ -910,20 +947,12 @@ void UChaosClothAssetEditorMode::SetDataflowGraphEditor(TSharedPtr<SDataflowGrap
 
 void UChaosClothAssetEditorMode::StartToolForSelectedNode(const UObject* SelectedNode)
 {
-	// Start the tool associated with the selected node, if any
-
-	const TSharedRef<FUICommandList> CommandList = Toolkit->GetToolkitCommands();
-	const UE::Chaos::ClothAsset::FChaosClothAssetEditorCommands& CommandInfos = UE::Chaos::ClothAsset::FChaosClothAssetEditorCommands::Get();
-
 	if (const UDataflowEdNode* const EdNode = Cast<UDataflowEdNode>(SelectedNode))
 	{
 		if (const TSharedPtr<const FDataflowNode> DataflowNode = EdNode->GetDataflowNode())
 		{
-			const FName DataflowNodeName = DataflowNode->GetType();
-			if (const TSharedPtr<const FUICommandInfo>* const Command = NodeTypeToToolCommandMap.Find(DataflowNodeName))
-			{
-				CommandList->TryExecuteAction(Command->ToSharedRef());
-			}
+			const FName DataflowNodeType = DataflowNode->GetType();
+			NodeTypeForPendingToolStart = DataflowNodeType;
 		}
 	}
 }
diff --git a/Engine/Plugins/Experimental/ChaosClothAssetEditor/Source/ChaosClothAssetEditor/Private/ChaosClothAsset/ClothEditorMode.h b/Engine/Plugins/Experimental/ChaosClothAssetEditor/Source/ChaosClothAssetEditor/Private/ChaosClothAsset/ClothEditorMode.h
--- a/Engine/Plugins/Experimental/ChaosClothAssetEditor/Source/ChaosClothAssetEditor/Private/ChaosClothAsset/ClothEditorMode.h
+++ b/Engine/Plugins/Experimental/ChaosClothAssetEditor/Source/ChaosClothAssetEditor/Private/ChaosClothAsset/ClothEditorMode.h
@@ -35,6 +35,7 @@ class FChaosClothPreviewScene;
 class FChaosClothAssetEditorModeToolkit;
 class FChaosClothAssetEditorToolkit;
 class FChaosClothEditorRestSpaceViewportClient;
+class IChaosClothAssetEditorToolBuilder;
 }
 class UEdGraphNode;
 
@@ -117,7 +118,8 @@ private:
 	// Use this function to register tools rather than UEdMode::RegisterTool() because we need to specify the ToolsContext
 	void RegisterClothTool(TSharedPtr<FUICommandInfo> UICommand, 
 		FString ToolIdentifier, 
-		UInteractiveToolBuilder* Builder, 
+		UInteractiveToolBuilder* Builder,
+		const UE::Chaos::ClothAsset::IChaosClothAssetEditorToolBuilder* ClothToolBuilder,
 		UEditorInteractiveToolsContext* UseToolsContext, 
 		EToolsContextScope ToolScope = EToolsContextScope::Default);
 
@@ -215,6 +217,12 @@ private:
 	UE::Chaos::ClothAsset::EClothPatternVertexType ConstructionViewMode;
 	bool bCanChangeConstructionViewMode = true;
 
+	// The Construction view mode that was active before starting the current tool. When the tool ends, restore this view mode
+	TOptional<UE::Chaos::ClothAsset::EClothPatternVertexType> SavedConstructionViewMode;
+
+	// Dataflow node type whose corresponding tool should be started on the next Tick
+	FName NodeTypeForPendingToolStart;
+
 	bool bConstructionViewWireframe = false;
 
 	// Create dynamic mesh components from the cloth component's rest space info
diff --git a/Engine/Plugins/Experimental/ChaosClothAssetEditor/Source/ChaosClothAssetEditor/Private/ChaosClothAsset/ClothEditorToolkit.cpp b/Engine/Plugins/Experimental/ChaosClothAssetEditor/Source/ChaosClothAssetEditor/Private/ChaosClothAsset/ClothEditorToolkit.cpp
--- a/Engine/Plugins/Experimental/ChaosClothAssetEditor/Source/ChaosClothAssetEditor/Private/ChaosClothAsset/ClothEditorToolkit.cpp
+++ b/Engine/Plugins/Experimental/ChaosClothAssetEditor/Source/ChaosClothAssetEditor/Private/ChaosClothAsset/ClothEditorToolkit.cpp
@@ -1030,9 +1030,15 @@ void FChaosClothAssetEditorToolkit::OnNodeSingleClicked(UObject* ClickedNode) co
 	UChaosClothAssetEditorMode* const ClothMode = CastChecked<UChaosClothAssetEditorMode>(EditorModeManager->GetActiveScriptableMode(UChaosClothAssetEditorMode::EM_ChaosClothAssetEditorModeId));
 	if (ClothMode)
 	{
-		// Start the corresponding tool
+		// Close any running tool
 		UEditorInteractiveToolsContext* const ToolsContext = ClothMode->GetInteractiveToolsContext();
 		checkf(ToolsContext, TEXT("No valid ToolsContext found for FChaosClothAssetEditorToolkit"));
+		if (ToolsContext->HasActiveTool())
+		{
+			ToolsContext->EndTool(EToolShutdownType::Completed);
+		}
+
+		// Start the corresponding tool
 		ClothMode->StartToolForSelectedNode(ClickedNode);
 	}
 }
diff --git a/Engine/Plugins/Experimental/ChaosClothAssetEditor/Source/ChaosClothAssetEditorTools/Private/ChaosClothAsset/ClothMeshSelectionTool.cpp b/Engine/Plugins/Experimental/ChaosClothAssetEditor/Source/ChaosClothAssetEditorTools/Private/ChaosClothAsset/ClothMeshSelectionTool.cpp
--- a/Engine/Plugins/Experimental/ChaosClothAssetEditor/Source/ChaosClothAssetEditorTools/Private/ChaosClothAsset/ClothMeshSelectionTool.cpp
+++ b/Engine/Plugins/Experimental/ChaosClothAssetEditor/Source/ChaosClothAssetEditorTools/Private/ChaosClothAsset/ClothMeshSelectionTool.cpp
@@ -26,6 +26,13 @@
 
 // ------------------- Builder -------------------
 
+void UClothMeshSelectionToolBuilder::GetSupportedViewModes(TArray<UE::Chaos::ClothAsset::EClothPatternVertexType>& Modes) const
+{
+	Modes.Add(UE::Chaos::ClothAsset::EClothPatternVertexType::Sim3D);
+	Modes.Add(UE::Chaos::ClothAsset::EClothPatternVertexType::Sim2D);
+	Modes.Add(UE::Chaos::ClothAsset::EClothPatternVertexType::Render);
+}
+
 const FToolTargetTypeRequirements& UClothMeshSelectionToolBuilder::GetTargetRequirements() const
 {
 	static FToolTargetTypeRequirements TypeRequirements(
diff --git a/Engine/Plugins/Experimental/ChaosClothAssetEditor/Source/ChaosClothAssetEditorTools/Private/ChaosClothAsset/ClothTransferSkinWeightsTool.cpp b/Engine/Plugins/Experimental/ChaosClothAssetEditor/Source/ChaosClothAssetEditorTools/Private/ChaosClothAsset/ClothTransferSkinWeightsTool.cpp
--- a/Engine/Plugins/Experimental/ChaosClothAssetEditor/Source/ChaosClothAssetEditorTools/Private/ChaosClothAsset/ClothTransferSkinWeightsTool.cpp
+++ b/Engine/Plugins/Experimental/ChaosClothAssetEditor/Source/ChaosClothAssetEditorTools/Private/ChaosClothAsset/ClothTransferSkinWeightsTool.cpp
@@ -15,6 +15,7 @@
 #include "InteractiveToolObjects.h"
 #include "MeshOpPreviewHelpers.h"
 #include "Components/SkeletalMeshComponent.h"
+#include "ChaosClothAsset/ClothPatternVertexType.h"
 
 #define LOCTEXT_NAMESPACE "ClothTransferSkinWeightsTool"
 
@@ -43,6 +44,11 @@ namespace UE::Chaos::ClothAsset::Private
 
 // ------------------- Builder -------------------
 
+void UClothTransferSkinWeightsToolBuilder::GetSupportedViewModes(TArray<UE::Chaos::ClothAsset::EClothPatternVertexType>& Modes) const
+{
+	Modes.Add(UE::Chaos::ClothAsset::EClothPatternVertexType::Sim3D);
+}
+
 USingleSelectionMeshEditingTool* UClothTransferSkinWeightsToolBuilder::CreateNewTool(const FToolBuilderState& SceneState) const
 {
 	UClothTransferSkinWeightsTool* NewTool = NewObject<UClothTransferSkinWeightsTool>(SceneState.ToolManager);
diff --git a/Engine/Plugins/Experimental/ChaosClothAssetEditor/Source/ChaosClothAssetEditorTools/Private/ChaosClothAsset/ClothWeightMapPaintTool.cpp b/Engine/Plugins/Experimental/ChaosClothAssetEditor/Source/ChaosClothAssetEditorTools/Private/ChaosClothAsset/ClothWeightMapPaintTool.cpp
--- a/Engine/Plugins/Experimental/ChaosClothAssetEditor/Source/ChaosClothAssetEditorTools/Private/ChaosClothAsset/ClothWeightMapPaintTool.cpp
+++ b/Engine/Plugins/Experimental/ChaosClothAssetEditor/Source/ChaosClothAssetEditorTools/Private/ChaosClothAsset/ClothWeightMapPaintTool.cpp
@@ -66,6 +66,13 @@ namespace UE::Chaos::ClothAsset::Private
 /*
  * ToolBuilder
  */
+
+void UClothEditorWeightMapPaintToolBuilder::GetSupportedViewModes(TArray<UE::Chaos::ClothAsset::EClothPatternVertexType>& Modes) const
+{
+	Modes.Add(UE::Chaos::ClothAsset::EClothPatternVertexType::Sim3D);
+	Modes.Add(UE::Chaos::ClothAsset::EClothPatternVertexType::Sim2D);
+}
+
 UMeshSurfacePointTool* UClothEditorWeightMapPaintToolBuilder::CreateNewTool(const FToolBuilderState& SceneState) const
 {
 	UClothEditorWeightMapPaintTool* PaintTool = NewObject<UClothEditorWeightMapPaintTool>(SceneState.ToolManager);
diff --git a/Engine/Plugins/Experimental/ChaosClothAssetEditor/Source/ChaosClothAssetEditorTools/Public/ChaosClothAsset/ClothEditorToolBuilder.h b/Engine/Plugins/Experimental/ChaosClothAssetEditor/Source/ChaosClothAssetEditorTools/Public/ChaosClothAsset/ClothEditorToolBuilder.h
new file mode 100644
--- /dev/null
+++ b/Engine/Plugins/Experimental/ChaosClothAssetEditor/Source/ChaosClothAssetEditorTools/Public/ChaosClothAsset/ClothEditorToolBuilder.h
@@ -0,0 +1,22 @@
+// Copyright Epic Games, Inc. All Rights Reserved.
+
+#pragma once
+
+#include "Containers/Array.h"
+
+namespace UE::Chaos::ClothAsset
+{
+	enum class EClothPatternVertexType : uint8;
+
+	// Interface for all UInteractiveToolBuilders for Cloth Editor tools
+	class IChaosClothAssetEditorToolBuilder
+	{
+	public:
+		virtual ~IChaosClothAssetEditorToolBuilder() {}
+
+		/** Returns all Construction View modes that this tool can operate in. The first element should be the preferred mode to switch to if necessary. */
+		virtual void GetSupportedViewModes(TArray<EClothPatternVertexType>& Modes) const = 0;
+	};
+}
+
+
diff --git a/Engine/Plugins/Experimental/ChaosClothAssetEditor/Source/ChaosClothAssetEditorTools/Public/ChaosClothAsset/ClothMeshSelectionTool.h b/Engine/Plugins/Experimental/ChaosClothAssetEditor/Source/ChaosClothAssetEditorTools/Public/ChaosClothAsset/ClothMeshSelectionTool.h
--- a/Engine/Plugins/Experimental/ChaosClothAssetEditor/Source/ChaosClothAssetEditorTools/Public/ChaosClothAsset/ClothMeshSelectionTool.h
+++ b/Engine/Plugins/Experimental/ChaosClothAssetEditor/Source/ChaosClothAssetEditorTools/Public/ChaosClothAsset/ClothMeshSelectionTool.h
@@ -3,6 +3,7 @@
 #pragma once
 
 #include "BaseTools/SingleSelectionMeshEditingTool.h"
+#include "ChaosClothAsset/ClothEditorToolBuilder.h"
 #include "ClothMeshSelectionTool.generated.h"
 
 class UPolygonSelectionMechanic;
@@ -16,12 +17,13 @@ namespace UE::Geometry
 }
 
 UCLASS()
-class CHAOSCLOTHASSETEDITORTOOLS_API UClothMeshSelectionToolBuilder : public UInteractiveToolWithToolTargetsBuilder
+class CHAOSCLOTHASSETEDITORTOOLS_API UClothMeshSelectionToolBuilder : public UInteractiveToolWithToolTargetsBuilder, public UE::Chaos::ClothAsset::IChaosClothAssetEditorToolBuilder
 {
 	GENERATED_BODY()
 
 private:
 
+	virtual void GetSupportedViewModes(TArray<UE::Chaos::ClothAsset::EClothPatternVertexType>& Modes) const override;
 	virtual bool CanBuildTool(const FToolBuilderState& SceneState) const override;
 	virtual UInteractiveTool* BuildTool(const FToolBuilderState& SceneState) const override;
 	virtual const FToolTargetTypeRequirements& GetTargetRequirements() const override;
diff --git a/Engine/Plugins/Experimental/ChaosClothAssetEditor/Source/ChaosClothAssetEditorTools/Public/ChaosClothAsset/ClothTransferSkinWeightsTool.h b/Engine/Plugins/Experimental/ChaosClothAssetEditor/Source/ChaosClothAssetEditorTools/Public/ChaosClothAsset/ClothTransferSkinWeightsTool.h
--- a/Engine/Plugins/Experimental/ChaosClothAssetEditor/Source/ChaosClothAssetEditorTools/Public/ChaosClothAsset/ClothTransferSkinWeightsTool.h
+++ b/Engine/Plugins/Experimental/ChaosClothAssetEditor/Source/ChaosClothAssetEditorTools/Public/ChaosClothAsset/ClothTransferSkinWeightsTool.h
@@ -5,6 +5,7 @@
 #include "BaseTools/SingleSelectionMeshEditingTool.h"
 #include "ModelingOperators.h"
 #include "Transforms/TransformGizmoDataBinder.h"
+#include "ChaosClothAsset/ClothEditorToolBuilder.h"
 #include "ClothTransferSkinWeightsTool.generated.h"
 
 
@@ -43,12 +44,13 @@ public:
 
 
 UCLASS()
-class CHAOSCLOTHASSETEDITORTOOLS_API UClothTransferSkinWeightsToolBuilder : public USingleSelectionMeshEditingToolBuilder
+class CHAOSCLOTHASSETEDITORTOOLS_API UClothTransferSkinWeightsToolBuilder : public USingleSelectionMeshEditingToolBuilder, public UE::Chaos::ClothAsset::IChaosClothAssetEditorToolBuilder
 {
 	GENERATED_BODY()
 
 private:
 
+	virtual void GetSupportedViewModes(TArray<UE::Chaos::ClothAsset::EClothPatternVertexType>& Modes) const override;
 	virtual USingleSelectionMeshEditingTool* CreateNewTool(const FToolBuilderState& SceneState) const override;	
 
 };
diff --git a/Engine/Plugins/Experimental/ChaosClothAssetEditor/Source/ChaosClothAssetEditorTools/Public/ChaosClothAsset/ClothWeightMapPaintTool.h b/Engine/Plugins/Experimental/ChaosClothAssetEditor/Source/ChaosClothAssetEditorTools/Public/ChaosClothAsset/ClothWeightMapPaintTool.h
--- a/Engine/Plugins/Experimental/ChaosClothAssetEditor/Source/ChaosClothAssetEditorTools/Public/ChaosClothAsset/ClothWeightMapPaintTool.h
+++ b/Engine/Plugins/Experimental/ChaosClothAssetEditor/Source/ChaosClothAssetEditorTools/Public/ChaosClothAsset/ClothWeightMapPaintTool.h
@@ -22,7 +22,7 @@
 #include "TransformTypes.h"
 #include "ToolDataVisualizer.h"
 #include "GroupTopology.h"
-
+#include "ChaosClothAsset/ClothEditorToolBuilder.h"
 #include "ClothWeightMapPaintTool.generated.h"
 
 class UMeshElementsVisualizer;
@@ -54,11 +54,12 @@ DECLARE_CYCLE_STAT(TEXT("WeightMapPaintTool_Normals_Compute"), WeightMapPaintToo
  * Tool Builder
  */
 UCLASS()
-class CHAOSCLOTHASSETEDITORTOOLS_API UClothEditorWeightMapPaintToolBuilder : public UMeshSurfacePointMeshEditingToolBuilder
+class CHAOSCLOTHASSETEDITORTOOLS_API UClothEditorWeightMapPaintToolBuilder : public UMeshSurfacePointMeshEditingToolBuilder, public UE::Chaos::ClothAsset::IChaosClothAssetEditorToolBuilder
 {
 	GENERATED_BODY()
 
-public:
+private:
+	virtual void GetSupportedViewModes(TArray<UE::Chaos::ClothAsset::EClothPatternVertexType>& Modes) const override;
 	virtual UMeshSurfacePointTool* CreateNewTool(const FToolBuilderState& SceneState) const override;
 };
 
