Renamed Translation to Location to have consistent naming with FTransform UI.

#rnx

diff --git a/Engine/Plugins/Experimental/ChaosClothAssetEditor/Source/ChaosClothAssetEditor/Private/ChaosClothAsset/ClothEditorPreviewScene.h b/Engine/Plugins/Experimental/ChaosClothAssetEditor/Source/ChaosClothAssetEditor/Private/ChaosClothAsset/ClothEditorPreviewScene.h
--- a/Engine/Plugins/Experimental/ChaosClothAssetEditor/Source/ChaosClothAssetEditor/Private/ChaosClothAsset/ClothEditorPreviewScene.h
+++ b/Engine/Plugins/Experimental/ChaosClothAssetEditor/Source/ChaosClothAssetEditor/Private/ChaosClothAsset/ClothEditorPreviewScene.h
@@ -42,7 +42,7 @@ public:
 	UPROPERTY(EditAnywhere, Transient, Category = "SkeletalMesh")
 	TObjectPtr<UAnimationAsset> AnimationAsset;
 
-	UPROPERTY(EditAnywhere, Transient, Category = "Transform")
+	UPROPERTY(EditAnywhere, Transient, Category = "Transform", Meta=(DisplayName="Location"))
 	FVector3d Translation = FVector3d::ZeroVector;
 
 	UPROPERTY(EditAnywhere, Transient, Category = "Transform")
