Sync between the transfer skin weight node and the transfer skin weight tool when changing same parameters (source mesh, transforms).

#rnx

diff --git a/Engine/Plugins/Experimental/ChaosClothAssetEditor/Source/ChaosClothAssetEditorTools/Private/ChaosClothAsset/ClothTransferSkinWeightsTool.cpp b/Engine/Plugins/Experimental/ChaosClothAssetEditor/Source/ChaosClothAssetEditorTools/Private/ChaosClothAsset/ClothTransferSkinWeightsTool.cpp
--- a/Engine/Plugins/Experimental/ChaosClothAssetEditor/Source/ChaosClothAssetEditorTools/Private/ChaosClothAsset/ClothTransferSkinWeightsTool.cpp
+++ b/Engine/Plugins/Experimental/ChaosClothAssetEditor/Source/ChaosClothAssetEditorTools/Private/ChaosClothAsset/ClothTransferSkinWeightsTool.cpp
@@ -113,9 +113,24 @@ void UClothTransferSkinWeightsTool::Setup()
 	SourceMeshComponent->RegisterComponent();
 
 	// Watch for property changes
+	ToolProperties->WatchProperty(TransferSkinWeightsNode->Transform, 
+	[this](FTransform Transform) 
+	{
+		SetSRTPropertiesFromTransform(Transform);
+	}, 
+	[this](FTransform Transform1, FTransform Transform2)
+	{ 
+		return !Transform1.Equals(Transform2, 0.f);
+	});
 
-	ToolProperties->WatchProperty(ToolProperties->SourceMesh, [this](TObjectPtr<USkeletalMesh>) { UpdateSourceMesh(); });
+	ToolProperties->WatchProperty(ToolProperties->SourceMesh, [this](TObjectPtr<USkeletalMesh> Mesh) { UpdateSourceMesh(Mesh); });
 
+	ToolProperties->WatchProperty(TransferSkinWeightsNode->SkeletalMesh, [this](TObjectPtr<USkeletalMesh> Mesh) 
+	{ 
+		ToolProperties->SourceMesh = Mesh; 
+		UpdateSourceMesh(Mesh); 
+	});
+	
 	ToolProperties->WatchProperty(ToolProperties->bHideSourceMesh, [this](bool bNewHideSourceMesh) 
 	{ 
 		check(SourceMeshComponent);
@@ -179,7 +194,7 @@ void UClothTransferSkinWeightsTool::Setup()
 	DataBinder->BindToInitializedGizmo(SourceMeshTransformGizmo, SourceMeshTransformProxy);
 
 
-	UpdateSourceMesh();
+	UpdateSourceMesh(ToolProperties->SourceMesh);
 
 	UE::ToolTarget::HideSourceObject(Target);
 }
@@ -270,15 +285,15 @@ void UClothTransferSkinWeightsTool::SetSRTPropertiesFromTransform(const FTransfo
 	ToolProperties->SourceMeshScale = Transform.GetScale3D();
 }
 
-
-void UClothTransferSkinWeightsTool::UpdateSourceMesh()
+void UClothTransferSkinWeightsTool::UpdateSourceMesh(TObjectPtr<USkeletalMesh> Mesh)
 {
 	checkf(ToolProperties, TEXT("ToolProperties is expected to be non-null. Be sure to run Setup() on this tool when it is created."));
 
-	if (ToolProperties->SourceMesh)
+	SourceMeshComponent->SetSkeletalMeshAsset(Mesh);
+
+	if (Mesh)
 	{
 		// Set up source mesh (from the SkeletalMesh)
-		SourceMeshComponent->SetSkeletalMeshAsset(ToolProperties->SourceMesh);	
 		SourceMeshParentActor->SetActorTransform(TransformFromProperties());
 		SourceMeshComponent->SetVisibility(!ToolProperties->bHideSourceMesh);
 
diff --git a/Engine/Plugins/Experimental/ChaosClothAssetEditor/Source/ChaosClothAssetEditorTools/Public/ChaosClothAsset/ClothTransferSkinWeightsTool.h b/Engine/Plugins/Experimental/ChaosClothAssetEditor/Source/ChaosClothAssetEditorTools/Public/ChaosClothAsset/ClothTransferSkinWeightsTool.h
--- a/Engine/Plugins/Experimental/ChaosClothAssetEditor/Source/ChaosClothAssetEditorTools/Public/ChaosClothAsset/ClothTransferSkinWeightsTool.h
+++ b/Engine/Plugins/Experimental/ChaosClothAssetEditor/Source/ChaosClothAssetEditorTools/Public/ChaosClothAsset/ClothTransferSkinWeightsTool.h
@@ -29,13 +29,13 @@ public:
 	UPROPERTY(EditAnywhere, Category = Source)
 	TObjectPtr<USkeletalMesh> SourceMesh;
 
-	UPROPERTY(EditAnywhere, Category = Source, meta = (EditCondition = "SourceMesh != nullptr"))
+	UPROPERTY(EditAnywhere, Category = "Transform", meta = (DisplayName = "Location", EditCondition = "SourceMesh != nullptr"))
 	FVector3d SourceMeshTranslation;
 
-	UPROPERTY(EditAnywhere, Category = Source, meta = (EditCondition = "SourceMesh != nullptr"))
+	UPROPERTY(EditAnywhere, Category = "Transform", meta = (DisplayName = "Rotation", EditCondition = "SourceMesh != nullptr"))
 	FVector3d SourceMeshRotation;
 
-	UPROPERTY(EditAnywhere, Category = Source, meta = (EditCondition = "SourceMesh != nullptr"))
+	UPROPERTY(EditAnywhere, Category = "Transform", meta = (DisplayName = "Scale", AllowPreserveRatio, EditCondition = "SourceMesh != nullptr"))
 	FVector3d SourceMeshScale;
 
 	UPROPERTY(EditAnywhere, Category = Source)
@@ -80,7 +80,7 @@ private:
 	FTransform TransformFromProperties() const;
 	void SetSRTPropertiesFromTransform(const FTransform& Transform) const;
 
-	void UpdateSourceMesh();
+	void UpdateSourceMesh(TObjectPtr<USkeletalMesh> Mesh);
 
 	void OpFinishedCallback(const UE::Geometry::FDynamicMeshOperator* Op);
 
