Cloth Editor Map Paint Tool:
- don't wait for brush to move to start applying stamp
- make the effect of the Strength parameter a little more pronounced for values less than 1
- don't reset Brush and Filter properties to default every time the tool starts

diff --git a/Engine/Plugins/Experimental/ChaosClothAssetEditor/Source/ChaosClothAssetEditorTools/Private/ChaosClothAsset/ClothWeightMapPaintBrushOps.h b/Engine/Plugins/Experimental/ChaosClothAssetEditor/Source/ChaosClothAssetEditorTools/Private/ChaosClothAsset/ClothWeightMapPaintBrushOps.h
--- a/Engine/Plugins/Experimental/ChaosClothAssetEditor/Source/ChaosClothAssetEditorTools/Private/ChaosClothAsset/ClothWeightMapPaintBrushOps.h
+++ b/Engine/Plugins/Experimental/ChaosClothAssetEditor/Source/ChaosClothAssetEditorTools/Private/ChaosClothAsset/ClothWeightMapPaintBrushOps.h
@@ -111,7 +111,7 @@ public:
 
 	virtual bool IgnoreZeroMovements() const override
 	{
-		return true;
+		return false;
 	}
 
 	virtual void ApplyStampByVertices(
@@ -148,7 +148,7 @@ public:
 
 	virtual bool IgnoreZeroMovements() const override
 	{
-		return true;
+		return false;
 	}
 
 	virtual void ApplyStampByVertices(
diff --git a/Engine/Plugins/Experimental/ChaosClothAssetEditor/Source/ChaosClothAssetEditorTools/Private/ChaosClothAsset/ClothWeightMapPaintTool.cpp b/Engine/Plugins/Experimental/ChaosClothAssetEditor/Source/ChaosClothAssetEditorTools/Private/ChaosClothAsset/ClothWeightMapPaintTool.cpp
--- a/Engine/Plugins/Experimental/ChaosClothAssetEditor/Source/ChaosClothAssetEditorTools/Private/ChaosClothAsset/ClothWeightMapPaintTool.cpp
+++ b/Engine/Plugins/Experimental/ChaosClothAssetEditor/Source/ChaosClothAssetEditorTools/Private/ChaosClothAsset/ClothWeightMapPaintTool.cpp
@@ -211,10 +211,8 @@ void UClothEditorWeightMapPaintTool::Setup()
 		[this](EClothEditorWeightMapPaintInteractionType NewType) { UpdateSubToolType(NewType); });
 	FilterProperties->WatchProperty(FilterProperties->BrushSize,
 		[this](float NewSize) { UMeshSculptToolBase::BrushProperties->BrushSize.AdaptiveSize = NewSize; });
-	FilterProperties->RestoreProperties(this);
 	FilterProperties->BrushSize = UMeshSculptToolBase::BrushProperties->BrushSize.AdaptiveSize;
-	FilterProperties->Strength = 1.0;
-	FilterProperties->AttributeValue = 1.0;
+	FilterProperties->RestoreProperties(this);
 	AddToolPropertySource(FilterProperties);
 
 	InitializeIndicator();
@@ -516,7 +514,7 @@ void UClothEditorWeightMapPaintTool::OnBeginStroke(const FRay& WorldRay)
 	if (PaintBrushOpProperties)
 	{
 		PaintBrushOpProperties->AttributeValue = FilterProperties->AttributeValue;
-		PaintBrushOpProperties->Strength = FilterProperties->Strength;
+		PaintBrushOpProperties->Strength = FilterProperties->Strength * FilterProperties->Strength;
 	}
 	if (EraseBrushOpProperties)
 	{
@@ -524,7 +522,7 @@ void UClothEditorWeightMapPaintTool::OnBeginStroke(const FRay& WorldRay)
 	}
 	if (SmoothBrushOpProperties)
 	{
-		SmoothBrushOpProperties->Strength = FilterProperties->Strength;
+		SmoothBrushOpProperties->Strength = FilterProperties->Strength * FilterProperties->Strength;
 	}
 
 	// initialize first "Last Stamp", so that we can assume all stamps in stroke have a valid previous stamp
diff --git a/Engine/Plugins/Experimental/ChaosClothAssetEditor/Source/ChaosClothAssetEditorTools/Public/ChaosClothAsset/ClothWeightMapPaintTool.h b/Engine/Plugins/Experimental/ChaosClothAssetEditor/Source/ChaosClothAssetEditorTools/Public/ChaosClothAsset/ClothWeightMapPaintTool.h
--- a/Engine/Plugins/Experimental/ChaosClothAssetEditor/Source/ChaosClothAssetEditorTools/Public/ChaosClothAsset/ClothWeightMapPaintTool.h
+++ b/Engine/Plugins/Experimental/ChaosClothAssetEditor/Source/ChaosClothAssetEditorTools/Public/ChaosClothAsset/ClothWeightMapPaintTool.h
@@ -137,18 +137,17 @@ public:
 	EClothEditorWeightMapPaintBrushType PrimaryBrushType = EClothEditorWeightMapPaintBrushType::Paint;
 
 	/** Relative size of brush */
-	UPROPERTY(EditAnywhere, Category = ActionType, meta = (DisplayName = "Brush Size", UIMin = "0.0", UIMax = "1.0", ClampMin = "0.0", ClampMax = "10.0", 
+	UPROPERTY(EditAnywhere, Category = Brush, meta = (DisplayName = "Brush Size", UIMin = "0.0", UIMax = "1.0", ClampMin = "0.0", ClampMax = "10.0", 
 		HideEditConditionToggle, EditConditionHides, EditCondition = "SubToolType != EClothEditorWeightMapPaintInteractionType::PolyLasso"))
 	float BrushSize = 0.25f;
 
 	/** The new value to paint on the mesh */
-	UPROPERTY(EditAnywhere, Category = ActionType, meta = (UIMin = 0, ClampMin = 0, UIMax = 1, ClampMax = 1))
+	UPROPERTY(EditAnywhere, Category = Brush, meta = (UIMin = 0, ClampMin = 0, UIMax = 1, ClampMax = 1))
 	double AttributeValue = 1;
 
 	/** How quickly each brush stroke will drive mesh values towards the desired value */
-	UPROPERTY(EditAnywhere, Category = ActionType, meta = (UIMin = 0, ClampMin = 0, UIMax = 1, ClampMax = 1))
-	double Strength = 1;
-
+	UPROPERTY(EditAnywhere, Category = Brush, meta = (UIMin = 0, ClampMin = 0, UIMax = 1, ClampMax = 1))
+	double Strength = 0.5;
 
 	/** The Region affected by the current operation will be bounded by edge angles larger than this threshold */
 	UPROPERTY(EditAnywhere, Category = Filters, meta = (UIMin = "0.0", UIMax = "180.0", EditCondition = "SubToolType != EClothEditorWeightMapPaintInteractionType::PolyLasso && BrushAreaMode == EClothEditorWeightMapPaintBrushAreaType::Connected"))
