Cloth Editor:
- Allow changing 2d/3d/render mode while a tool is active. This is accomplished by quitting the current tool and saving the results to the corresponding node, then switching view modes, then restarting the tool again.
- Disable switching to unsupported view modes for the active tool

diff --git a/Engine/Plugins/Experimental/ChaosClothAssetEditor/Source/ChaosClothAssetEditor/Private/ChaosClothAsset/ClothEditorMode.cpp b/Engine/Plugins/Experimental/ChaosClothAssetEditor/Source/ChaosClothAssetEditor/Private/ChaosClothAsset/ClothEditorMode.cpp
--- a/Engine/Plugins/Experimental/ChaosClothAssetEditor/Source/ChaosClothAssetEditor/Private/ChaosClothAsset/ClothEditorMode.cpp
+++ b/Engine/Plugins/Experimental/ChaosClothAssetEditor/Source/ChaosClothAssetEditor/Private/ChaosClothAsset/ClothEditorMode.cpp
@@ -138,7 +138,7 @@ void UChaosClothAssetEditorMode::AddToolTargetFactories()
 void UChaosClothAssetEditorMode::RegisterClothTool(TSharedPtr<FUICommandInfo> UICommand, 
 	FString ToolIdentifier, 
 	UInteractiveToolBuilder* Builder, 
-	const UE::Chaos::ClothAsset::IChaosClothAssetEditorToolBuilder* ClothToolBuilder,
+	const IChaosClothAssetEditorToolBuilder* ClothToolBuilder,
 	UEditorInteractiveToolsContext* const ToolsContext, 
 	EToolsContextScope ToolScope)
 {
@@ -165,13 +165,20 @@ void UChaosClothAssetEditorMode::RegisterClothTool(TSharedPtr<FUICommandInfo> UI
 	CommandList->MapAction(UICommand,
 		FExecuteAction::CreateWeakLambda(ToolsContext, [this, ToolsContext, ToolIdentifier, ClothToolBuilder]()
 		{
+			// Check if we need to switch view modes before starting the tool
 			TArray<UE::Chaos::ClothAsset::EClothPatternVertexType> SupportedModes;
 			ClothToolBuilder->GetSupportedViewModes(SupportedModes);
 
 			if (SupportedModes.Num() > 0 && !SupportedModes.Contains(this->GetConstructionViewMode()))
 			{
-				// Need to switch view modes before starting the tool
-				SavedConstructionViewMode = this->GetConstructionViewMode();	// remember the current view mode so we can restore it later
+				if (!bShouldRestoreSavedConstructionViewMode)
+				{
+					// remember the current view mode so we can restore it later
+					SavedConstructionViewMode = this->GetConstructionViewMode();	
+					bShouldRestoreSavedConstructionViewMode = true;
+				}
+
+				// switch to the preferred view mode for the tool that's about to start
 				this->SetConstructionViewMode(SupportedModes[0]);
 			}
 
@@ -269,23 +276,21 @@ void UChaosClothAssetEditorMode::OnToolStarted(UInteractiveToolManager* Manager,
 	using namespace UE::Chaos::ClothAsset;
 
 	FChaosClothAssetEditorCommands::UpdateToolCommandBinding(Tool, ToolCommandList, false);
-
-	bCanChangeConstructionViewMode = false;
 }
 
 void UChaosClothAssetEditorMode::OnToolEnded(UInteractiveToolManager* Manager, UInteractiveTool* Tool)
 {
 	UE::Chaos::ClothAsset::FChaosClothAssetEditorCommands::UpdateToolCommandBinding(Tool, ToolCommandList, true);
 
-	if (SavedConstructionViewMode)
+	if (bShouldRestoreSavedConstructionViewMode)
 	{
-		SetConstructionViewMode(*SavedConstructionViewMode);
-		SavedConstructionViewMode.Reset();
+		SetConstructionViewMode(SavedConstructionViewMode);
+		bShouldRestoreSavedConstructionViewMode = false;
+	}
+	else
+	{
+		ReinitializeDynamicMeshComponents();
 	}
-
-	bCanChangeConstructionViewMode = true;
-
-	ReinitializeDynamicMeshComponents();
 
 	if (TSharedPtr<SDataflowGraphEditor> GraphEditor = DataflowGraphEditor.Pin())
 	{
@@ -859,10 +864,27 @@ FBox UChaosClothAssetEditorMode::PreviewBoundingBox() const
 
 void UChaosClothAssetEditorMode::SetConstructionViewMode(UE::Chaos::ClothAsset::EClothPatternVertexType InMode)
 {
+	// We will first check if there is an active tool. If so, we'll shut down the tool and save the results to the Node, then change view modes, then restart the tool again.
+	bool bEndedActiveTool = false;
+	UInteractiveToolManager* const ToolManager = GetInteractiveToolsContext()->ToolManager;
+	checkf(ToolManager, TEXT("No valid ToolManager found for UChaosClothAssetEditorMode"));
+	if (UInteractiveTool* const ActiveTool = ToolManager->GetActiveTool(EToolSide::Left))
+	{
+		// avoid switching back to the previous view mode when the tool ends here
+		const bool bTempShouldRestoreVal = bShouldRestoreSavedConstructionViewMode;
+		bShouldRestoreSavedConstructionViewMode = false;
+
+		ToolManager->PostActiveToolShutdownRequest(ActiveTool, EToolShutdownType::Accept);
+		bEndedActiveTool = true;
+
+		// now we can restore the previous view mode the next time the tool ends
+		bShouldRestoreSavedConstructionViewMode = bTempShouldRestoreVal;
+	}
+
 	ConstructionViewMode = InMode;
 	ReinitializeDynamicMeshComponents();
 
-	TSharedPtr<UE::Chaos::ClothAsset::FChaosClothEditorRestSpaceViewportClient> VC = RestSpaceViewportClient.Pin();
+	const TSharedPtr<UE::Chaos::ClothAsset::FChaosClothEditorRestSpaceViewportClient> VC = RestSpaceViewportClient.Pin();
 	if (VC.IsValid())
 	{
 		VC->SetConstructionViewMode(ConstructionViewMode);
@@ -870,6 +892,20 @@ void UChaosClothAssetEditorMode::SetConstructionViewMode(UE::Chaos::ClothAsset::
 
 	// If we are switching to a mode with a valid mesh for the first time, focus the camera on it
 	FirstTimeFocusRestSpaceViewport();
+
+	if (bEndedActiveTool)
+	{
+		// If we ended the active tool in order to change modes, restart it now
+		if (const TSharedPtr<const SDataflowGraphEditor> PinnedGraphEditor = DataflowGraphEditor.Pin())
+		{
+			const FGraphPanelSelectionSet& SelectedNodes = PinnedGraphEditor->GetSelectedNodes();
+			if (SelectedNodes.Num() == 1)
+			{
+				StartToolForSelectedNode(*SelectedNodes.CreateConstIterator());
+			}
+		}
+	}
+
 }
 
 UE::Chaos::ClothAsset::EClothPatternVertexType UChaosClothAssetEditorMode::GetConstructionViewMode() const
@@ -877,9 +913,22 @@ UE::Chaos::ClothAsset::EClothPatternVertexType UChaosClothAssetEditorMode::GetCo
 	return ConstructionViewMode;
 }
 
-bool UChaosClothAssetEditorMode::CanChangeConstructionViewMode() const
+bool UChaosClothAssetEditorMode::CanChangeConstructionViewModeTo(UE::Chaos::ClothAsset::EClothPatternVertexType NewViewMode) const
 {
-	return bCanChangeConstructionViewMode;
+	if (!GetToolManager()->HasActiveTool(EToolSide::Left))
+	{
+		return true;
+	}
+
+	const UInteractiveToolBuilder* const ActiveToolBuilder = GetToolManager()->GetActiveToolBuilder(EToolSide::Left);
+	checkf(ActiveToolBuilder, TEXT("No Active Tool Builder found despite having an Active Tool"));
+
+	const IChaosClothAssetEditorToolBuilder* const ClothToolBuilder = Cast<const IChaosClothAssetEditorToolBuilder>(ActiveToolBuilder);
+	checkf(ClothToolBuilder, TEXT("Cloth Editor has an active Tool Builder that does not implement IChaosClothAssetEditorToolBuilder"));
+
+	TArray<UE::Chaos::ClothAsset::EClothPatternVertexType> SupportedViewModes;
+	ClothToolBuilder->GetSupportedViewModes(SupportedViewModes);
+	return SupportedViewModes.Contains(NewViewMode);
 }
 
 void UChaosClothAssetEditorMode::ToggleConstructionViewWireframe()
diff --git a/Engine/Plugins/Experimental/ChaosClothAssetEditor/Source/ChaosClothAssetEditor/Private/ChaosClothAsset/ClothEditorMode.h b/Engine/Plugins/Experimental/ChaosClothAssetEditor/Source/ChaosClothAssetEditor/Private/ChaosClothAsset/ClothEditorMode.h
--- a/Engine/Plugins/Experimental/ChaosClothAssetEditor/Source/ChaosClothAssetEditor/Private/ChaosClothAsset/ClothEditorMode.h
+++ b/Engine/Plugins/Experimental/ChaosClothAssetEditor/Source/ChaosClothAssetEditor/Private/ChaosClothAsset/ClothEditorMode.h
@@ -35,8 +35,8 @@ class FChaosClothPreviewScene;
 class FChaosClothAssetEditorModeToolkit;
 class FChaosClothAssetEditorToolkit;
 class FChaosClothEditorRestSpaceViewportClient;
-class IChaosClothAssetEditorToolBuilder;
 }
+class IChaosClothAssetEditorToolBuilder;
 class UEdGraphNode;
 
 /**
@@ -63,7 +63,7 @@ public:
 
 	void SetConstructionViewMode(UE::Chaos::ClothAsset::EClothPatternVertexType InMode);
 	UE::Chaos::ClothAsset::EClothPatternVertexType GetConstructionViewMode() const;
-	bool CanChangeConstructionViewMode() const;
+	bool CanChangeConstructionViewModeTo(UE::Chaos::ClothAsset::EClothPatternVertexType NewViewMode) const;
 
 	void ToggleConstructionViewWireframe();
 	bool IsConstructionViewWireframeActive() const
@@ -119,7 +119,7 @@ private:
 	void RegisterClothTool(TSharedPtr<FUICommandInfo> UICommand, 
 		FString ToolIdentifier, 
 		UInteractiveToolBuilder* Builder,
-		const UE::Chaos::ClothAsset::IChaosClothAssetEditorToolBuilder* ClothToolBuilder,
+		const IChaosClothAssetEditorToolBuilder* ClothToolBuilder,
 		UEditorInteractiveToolsContext* UseToolsContext, 
 		EToolsContextScope ToolScope = EToolsContextScope::Default);
 
@@ -215,10 +215,12 @@ private:
 	void RestSpaceViewportResized(FViewport* RestspaceViewport, uint32 Unused);
 
 	UE::Chaos::ClothAsset::EClothPatternVertexType ConstructionViewMode;
-	bool bCanChangeConstructionViewMode = true;
 
-	// The Construction view mode that was active before starting the current tool. When the tool ends, restore this view mode
-	TOptional<UE::Chaos::ClothAsset::EClothPatternVertexType> SavedConstructionViewMode;
+	// The Construction view mode that was active before starting the current tool. When the tool ends, restore this view mode if bShouldRestoreSavedConstructionViewMode is true.
+	UE::Chaos::ClothAsset::EClothPatternVertexType SavedConstructionViewMode;
+
+	// Whether we should restore the previous view mode when a tool ends
+	bool bShouldRestoreSavedConstructionViewMode = false;
 
 	// Dataflow node type whose corresponding tool should be started on the next Tick
 	FName NodeTypeForPendingToolStart;
diff --git a/Engine/Plugins/Experimental/ChaosClothAssetEditor/Source/ChaosClothAssetEditor/Private/ChaosClothAsset/SClothEditorRestSpaceViewport.cpp b/Engine/Plugins/Experimental/ChaosClothAssetEditor/Source/ChaosClothAssetEditor/Private/ChaosClothAsset/SClothEditorRestSpaceViewport.cpp
--- a/Engine/Plugins/Experimental/ChaosClothAssetEditor/Source/ChaosClothAssetEditor/Private/ChaosClothAsset/SClothEditorRestSpaceViewport.cpp
+++ b/Engine/Plugins/Experimental/ChaosClothAssetEditor/Source/ChaosClothAssetEditor/Private/ChaosClothAsset/SClothEditorRestSpaceViewport.cpp
@@ -54,7 +54,7 @@ void SChaosClothAssetEditorRestSpaceViewport::BindCommands()
 		{ 
 			if (const UChaosClothAssetEditorMode* const ClothEdMode = GetEdMode())
 			{
-				return ClothEdMode->CanChangeConstructionViewMode();
+				return ClothEdMode->CanChangeConstructionViewModeTo(EClothPatternVertexType::Sim2D);
 			}
 			return false; 
 		}),
@@ -81,7 +81,7 @@ void SChaosClothAssetEditorRestSpaceViewport::BindCommands()
 		{ 
 			if (const UChaosClothAssetEditorMode* const ClothEdMode = GetEdMode())
 			{
-				return ClothEdMode->CanChangeConstructionViewMode();
+				return ClothEdMode->CanChangeConstructionViewModeTo(EClothPatternVertexType::Sim3D);
 			}
 			return false; 
 		}),
@@ -109,7 +109,7 @@ void SChaosClothAssetEditorRestSpaceViewport::BindCommands()
 		{ 
 			if (const UChaosClothAssetEditorMode* const ClothEdMode = GetEdMode())
 			{
-				return ClothEdMode->CanChangeConstructionViewMode();
+				return ClothEdMode->CanChangeConstructionViewModeTo(EClothPatternVertexType::Render);
 			}
 			return false; 
 		}),
diff --git a/Engine/Plugins/Experimental/ChaosClothAssetEditor/Source/ChaosClothAssetEditorTools/Public/ChaosClothAsset/ClothEditorToolBuilder.h b/Engine/Plugins/Experimental/ChaosClothAssetEditor/Source/ChaosClothAssetEditorTools/Public/ChaosClothAsset/ClothEditorToolBuilder.h
--- a/Engine/Plugins/Experimental/ChaosClothAssetEditor/Source/ChaosClothAssetEditorTools/Public/ChaosClothAsset/ClothEditorToolBuilder.h
+++ b/Engine/Plugins/Experimental/ChaosClothAssetEditor/Source/ChaosClothAssetEditorTools/Public/ChaosClothAsset/ClothEditorToolBuilder.h
@@ -3,20 +3,28 @@
 #pragma once
 
 #include "Containers/Array.h"
+#include "ClothEditorToolBuilder.generated.h"
 
 namespace UE::Chaos::ClothAsset
 {
 	enum class EClothPatternVertexType : uint8;
+}
 
-	// Interface for all UInteractiveToolBuilders for Cloth Editor tools
-	class IChaosClothAssetEditorToolBuilder
-	{
-	public:
-		virtual ~IChaosClothAssetEditorToolBuilder() {}
 
-		/** Returns all Construction View modes that this tool can operate in. The first element should be the preferred mode to switch to if necessary. */
-		virtual void GetSupportedViewModes(TArray<EClothPatternVertexType>& Modes) const = 0;
-	};
-}
+UINTERFACE(MinimalAPI)
+class UChaosClothAssetEditorToolBuilder : public UInterface
+{
+	GENERATED_BODY()
+};
+
+
+class IChaosClothAssetEditorToolBuilder
+{
+	GENERATED_BODY()
+
+public:
 
+	/** Returns all Construction View modes that this tool can operate in. The first element should be the preferred mode to switch to if necessary. */
+	virtual void GetSupportedViewModes(TArray<UE::Chaos::ClothAsset::EClothPatternVertexType>& Modes) const = 0;
+};
 
diff --git a/Engine/Plugins/Experimental/ChaosClothAssetEditor/Source/ChaosClothAssetEditorTools/Public/ChaosClothAsset/ClothMeshSelectionTool.h b/Engine/Plugins/Experimental/ChaosClothAssetEditor/Source/ChaosClothAssetEditorTools/Public/ChaosClothAsset/ClothMeshSelectionTool.h
--- a/Engine/Plugins/Experimental/ChaosClothAssetEditor/Source/ChaosClothAssetEditorTools/Public/ChaosClothAsset/ClothMeshSelectionTool.h
+++ b/Engine/Plugins/Experimental/ChaosClothAssetEditor/Source/ChaosClothAssetEditorTools/Public/ChaosClothAsset/ClothMeshSelectionTool.h
@@ -17,7 +17,7 @@ namespace UE::Geometry
 }
 
 UCLASS()
-class CHAOSCLOTHASSETEDITORTOOLS_API UClothMeshSelectionToolBuilder : public UInteractiveToolWithToolTargetsBuilder, public UE::Chaos::ClothAsset::IChaosClothAssetEditorToolBuilder
+class CHAOSCLOTHASSETEDITORTOOLS_API UClothMeshSelectionToolBuilder : public UInteractiveToolWithToolTargetsBuilder, public IChaosClothAssetEditorToolBuilder
 {
 	GENERATED_BODY()
 
diff --git a/Engine/Plugins/Experimental/ChaosClothAssetEditor/Source/ChaosClothAssetEditorTools/Public/ChaosClothAsset/ClothTransferSkinWeightsTool.h b/Engine/Plugins/Experimental/ChaosClothAssetEditor/Source/ChaosClothAssetEditorTools/Public/ChaosClothAsset/ClothTransferSkinWeightsTool.h
--- a/Engine/Plugins/Experimental/ChaosClothAssetEditor/Source/ChaosClothAssetEditorTools/Public/ChaosClothAsset/ClothTransferSkinWeightsTool.h
+++ b/Engine/Plugins/Experimental/ChaosClothAssetEditor/Source/ChaosClothAssetEditorTools/Public/ChaosClothAsset/ClothTransferSkinWeightsTool.h
@@ -44,7 +44,7 @@ public:
 
 
 UCLASS()
-class CHAOSCLOTHASSETEDITORTOOLS_API UClothTransferSkinWeightsToolBuilder : public USingleSelectionMeshEditingToolBuilder, public UE::Chaos::ClothAsset::IChaosClothAssetEditorToolBuilder
+class CHAOSCLOTHASSETEDITORTOOLS_API UClothTransferSkinWeightsToolBuilder : public USingleSelectionMeshEditingToolBuilder, public IChaosClothAssetEditorToolBuilder
 {
 	GENERATED_BODY()
 
diff --git a/Engine/Plugins/Experimental/ChaosClothAssetEditor/Source/ChaosClothAssetEditorTools/Public/ChaosClothAsset/ClothWeightMapPaintTool.h b/Engine/Plugins/Experimental/ChaosClothAssetEditor/Source/ChaosClothAssetEditorTools/Public/ChaosClothAsset/ClothWeightMapPaintTool.h
--- a/Engine/Plugins/Experimental/ChaosClothAssetEditor/Source/ChaosClothAssetEditorTools/Public/ChaosClothAsset/ClothWeightMapPaintTool.h
+++ b/Engine/Plugins/Experimental/ChaosClothAssetEditor/Source/ChaosClothAssetEditorTools/Public/ChaosClothAsset/ClothWeightMapPaintTool.h
@@ -54,7 +54,7 @@ DECLARE_CYCLE_STAT(TEXT("WeightMapPaintTool_Normals_Compute"), WeightMapPaintToo
  * Tool Builder
  */
 UCLASS()
-class CHAOSCLOTHASSETEDITORTOOLS_API UClothEditorWeightMapPaintToolBuilder : public UMeshSurfacePointMeshEditingToolBuilder, public UE::Chaos::ClothAsset::IChaosClothAssetEditorToolBuilder
+class CHAOSCLOTHASSETEDITORTOOLS_API UClothEditorWeightMapPaintToolBuilder : public UMeshSurfacePointMeshEditingToolBuilder, public IChaosClothAssetEditorToolBuilder
 {
 	GENERATED_BODY()
 
