Sync between the weightmap node and weightmap tool when chaning the name of the map.

#rnx

diff --git a/Engine/Plugins/Experimental/ChaosClothAssetEditor/Source/ChaosClothAssetEditorTools/Private/ChaosClothAsset/ClothTransferSkinWeightsTool.cpp b/Engine/Plugins/Experimental/ChaosClothAssetEditor/Source/ChaosClothAssetEditorTools/Private/ChaosClothAsset/ClothTransferSkinWeightsTool.cpp
--- a/Engine/Plugins/Experimental/ChaosClothAssetEditor/Source/ChaosClothAssetEditorTools/Private/ChaosClothAsset/ClothTransferSkinWeightsTool.cpp
+++ b/Engine/Plugins/Experimental/ChaosClothAssetEditor/Source/ChaosClothAssetEditorTools/Private/ChaosClothAsset/ClothTransferSkinWeightsTool.cpp
@@ -114,11 +114,11 @@ void UClothTransferSkinWeightsTool::Setup()
 
 	// Watch for property changes
 	ToolProperties->WatchProperty(TransferSkinWeightsNode->Transform, 
-	[this](FTransform Transform) 
+	[this](const FTransform& Transform) 
 	{
 		SetSRTPropertiesFromTransform(Transform);
 	}, 
-	[this](FTransform Transform1, FTransform Transform2)
+	[this](const FTransform& Transform1, const FTransform& Transform2)
 	{ 
 		return !Transform1.Equals(Transform2, 0.f);
 	});
@@ -127,8 +127,7 @@ void UClothTransferSkinWeightsTool::Setup()
 
 	ToolProperties->WatchProperty(TransferSkinWeightsNode->SkeletalMesh, [this](TObjectPtr<USkeletalMesh> Mesh) 
 	{ 
-		ToolProperties->SourceMesh = Mesh; 
-		UpdateSourceMesh(Mesh); 
+		ToolProperties->SourceMesh = Mesh; // this triggers the WatchProperty of ToolProperties->SourceMesh above
 	});
 	
 	ToolProperties->WatchProperty(ToolProperties->bHideSourceMesh, [this](bool bNewHideSourceMesh) 
@@ -289,6 +288,11 @@ void UClothTransferSkinWeightsTool::UpdateSourceMesh(TObjectPtr<USkeletalMesh> M
 {
 	checkf(ToolProperties, TEXT("ToolProperties is expected to be non-null. Be sure to run Setup() on this tool when it is created."));
 
+	if (Mesh == SourceMeshComponent->GetSkeletalMeshAsset())
+	{
+		return;
+	}
+
 	SourceMeshComponent->SetSkeletalMeshAsset(Mesh);
 
 	if (Mesh)
diff --git a/Engine/Plugins/Experimental/ChaosClothAssetEditor/Source/ChaosClothAssetEditorTools/Private/ChaosClothAsset/ClothWeightMapPaintTool.cpp b/Engine/Plugins/Experimental/ChaosClothAssetEditor/Source/ChaosClothAssetEditorTools/Private/ChaosClothAsset/ClothWeightMapPaintTool.cpp
--- a/Engine/Plugins/Experimental/ChaosClothAssetEditor/Source/ChaosClothAssetEditorTools/Private/ChaosClothAsset/ClothWeightMapPaintTool.cpp
+++ b/Engine/Plugins/Experimental/ChaosClothAssetEditor/Source/ChaosClothAssetEditorTools/Private/ChaosClothAsset/ClothWeightMapPaintTool.cpp
@@ -196,12 +196,10 @@ void UClothEditorWeightMapPaintTool::Setup()
 
 	UpdateWeightMapProperties = NewObject<UClothEditorUpdateWeightMapProperties>(this);
 	UpdateWeightMapProperties->Name = WeightMapNodeToUpdate->Name;
-	UpdateWeightMapProperties->WatchProperty(UpdateWeightMapProperties->Name, [this, OriginalName = WeightMapNodeToUpdate->Name](const FString& NewName)
+
+	UpdateWeightMapProperties->WatchProperty(WeightMapNodeToUpdate->Name, [this](const FString& NewName)
 	{
-		if (NewName != OriginalName)
-		{
-			bAnyChangeMade = true;
-		}
+		UpdateWeightMapProperties->Name = NewName;
 	});
 	AddToolPropertySource(UpdateWeightMapProperties);
 
@@ -1515,7 +1513,7 @@ void UClothEditorWeightMapPaintTool::OnTick(float DeltaTime)
 
 bool UClothEditorWeightMapPaintTool::CanAccept() const
 {
-	return bAnyChangeMade;
+	return bAnyChangeMade || UpdateWeightMapProperties->Name != WeightMapNodeToUpdate->Name;
 }
 
 
