Cloth Editor: when a new tool starts, manually set keyboard focus on the construction viewport so that tool hotkeys work without having to click in the viewport first.
Also restore some previously removed code in the ViewportClient that allows Tools to process key events.

diff --git a/Engine/Plugins/Experimental/ChaosClothAssetEditor/Source/ChaosClothAssetEditor/Private/ChaosClothAsset/ClothEditorModeToolkit.cpp b/Engine/Plugins/Experimental/ChaosClothAssetEditor/Source/ChaosClothAssetEditor/Private/ChaosClothAsset/ClothEditorModeToolkit.cpp
--- a/Engine/Plugins/Experimental/ChaosClothAssetEditor/Source/ChaosClothAssetEditor/Private/ChaosClothAsset/ClothEditorModeToolkit.cpp
+++ b/Engine/Plugins/Experimental/ChaosClothAssetEditor/Source/ChaosClothAssetEditor/Private/ChaosClothAsset/ClothEditorModeToolkit.cpp
@@ -20,6 +20,7 @@
 #include "SPrimaryButton.h"
 #include "IDetailsView.h"
 #include "Widgets/Images/SImage.h"
+#include "Framework/Application/SlateApplication.h"
 	
 #define LOCTEXT_NAMESPACE "FChaosClothAssetEditorModeToolkit"
 
@@ -239,6 +240,9 @@ void FChaosClothAssetEditorModeToolkit::OnToolStarted(UInteractiveToolManager* M
 	if (SBaseCharacterFXEditorViewport* Widget = GetViewportWidgetForManager(Manager))
 	{
 		Widget->AddOverlayWidget(ViewportOverlayWidget.ToSharedRef());
+
+		// Manually set keyboard focus to this viewport. Otherwise the user has to click on the viewport before Tool keyboard shortcuts will work.
+		FSlateApplication::Get().SetUserFocus(FSlateApplication::Get().GetUserIndexForKeyboard(), Widget->AsShared(), EFocusCause::SetDirectly);
 	}
 }
 
diff --git a/Engine/Plugins/Experimental/ChaosClothAssetEditor/Source/ChaosClothAssetEditor/Private/ChaosClothAsset/ClothEditorRestSpaceViewportClient.cpp b/Engine/Plugins/Experimental/ChaosClothAssetEditor/Source/ChaosClothAssetEditor/Private/ChaosClothAsset/ClothEditorRestSpaceViewportClient.cpp
--- a/Engine/Plugins/Experimental/ChaosClothAssetEditor/Source/ChaosClothAssetEditor/Private/ChaosClothAsset/ClothEditorRestSpaceViewportClient.cpp
+++ b/Engine/Plugins/Experimental/ChaosClothAssetEditor/Source/ChaosClothAssetEditor/Private/ChaosClothAsset/ClothEditorRestSpaceViewportClient.cpp
@@ -117,6 +117,17 @@ bool FChaosClothEditorRestSpaceViewportClient::ShouldOrbitCamera() const
 
 bool FChaosClothEditorRestSpaceViewportClient::InputKey(const FInputKeyEventArgs& EventArgs)
 {
+	// See if any tool commands want to handle the key event
+	const TSharedPtr<FUICommandList> PinnedToolCommandList = ToolCommandList.Pin();
+	if (EventArgs.Event != IE_Released && PinnedToolCommandList.IsValid())
+	{
+		const FModifierKeysState KeyState = FSlateApplication::Get().GetModifierKeys();
+		if (PinnedToolCommandList->ProcessCommandBindings(EventArgs.Key, KeyState, (EventArgs.Event == IE_Repeat)))
+		{
+			return true;
+		}
+	}
+
 	if (ConstructionViewMode != EClothPatternVertexType::Sim2D)
 	{
 		return FEditorViewportClient::InputKey(EventArgs);
