Cloth Editor:
- StaticMeshImportNode: fix EditCondition parsing for the UVScale property
- Weight Map Paint Tool: fix EditConditions that reference a non-existant property

diff --git a/Engine/Plugins/Experimental/ChaosClothAssetEditor/Source/ChaosClothAssetDataflowNodes/Public/ChaosClothAsset/StaticMeshImportNode.h b/Engine/Plugins/Experimental/ChaosClothAssetEditor/Source/ChaosClothAssetDataflowNodes/Public/ChaosClothAsset/StaticMeshImportNode.h
--- a/Engine/Plugins/Experimental/ChaosClothAssetEditor/Source/ChaosClothAssetDataflowNodes/Public/ChaosClothAsset/StaticMeshImportNode.h
+++ b/Engine/Plugins/Experimental/ChaosClothAssetEditor/Source/ChaosClothAssetDataflowNodes/Public/ChaosClothAsset/StaticMeshImportNode.h
@@ -36,11 +36,11 @@ public:
 	bool bImportAsRenderMesh = true;
 
 	/* UV Channel used to populate Sim Mesh positions */
-	UPROPERTY(EditAnywhere, Category = "Static Mesh Import", meta = (EditCondition = "bImportAsSimMesh"))
+	UPROPERTY(EditAnywhere, Category = "Static Mesh Import", meta = (ClampMin = "-1", EditCondition = "bImportAsSimMesh"))
 	int32 UVChannel = 0;
 
 	/* Apply this scale to the UVs when populating Sim Mesh positions. */
-	UPROPERTY(EditAnywhere, Category = "Static Mesh Import", Meta = (AllowPreserveRatio, EditCondition = "bImportAsSimMesh && UVChannel != INDEX_NONE"))
+	UPROPERTY(EditAnywhere, Category = "Static Mesh Import", Meta = (AllowPreserveRatio, EditCondition = "bImportAsSimMesh && (UVChannel >= 0)"))
 	FVector2f UVScale = { 1.f, 1.f };
 
 	FChaosClothAssetStaticMeshImportNode(const Dataflow::FNodeParameters& InParam, FGuid InGuid = FGuid::NewGuid());
diff --git a/Engine/Plugins/Experimental/ChaosClothAssetEditor/Source/ChaosClothAssetEditorTools/Public/ChaosClothAsset/ClothWeightMapPaintTool.h b/Engine/Plugins/Experimental/ChaosClothAssetEditor/Source/ChaosClothAssetEditorTools/Public/ChaosClothAsset/ClothWeightMapPaintTool.h
--- a/Engine/Plugins/Experimental/ChaosClothAssetEditor/Source/ChaosClothAssetEditorTools/Public/ChaosClothAsset/ClothWeightMapPaintTool.h
+++ b/Engine/Plugins/Experimental/ChaosClothAssetEditor/Source/ChaosClothAssetEditorTools/Public/ChaosClothAsset/ClothWeightMapPaintTool.h
@@ -150,15 +150,15 @@ public:
 	double Strength = 0.5;
 
 	/** The Region affected by the current operation will be bounded by edge angles larger than this threshold */
-	UPROPERTY(EditAnywhere, Category = Filters, meta = (UIMin = "0.0", UIMax = "180.0", EditCondition = "SubToolType != EClothEditorWeightMapPaintInteractionType::PolyLasso && BrushAreaMode == EClothEditorWeightMapPaintBrushAreaType::Connected"))
+	UPROPERTY(EditAnywhere, Category = Filters, meta = (UIMin = "0.0", UIMax = "180.0", EditCondition = "SubToolType != EClothEditorWeightMapPaintInteractionType::PolyLasso"))
 	float AngleThreshold = 180.0f;
 
 	/** The Region affected by the current operation will be bounded by UV borders/seams */
-	UPROPERTY(EditAnywhere, Category = Filters, meta = (EditCondition = "SubToolType != EClothEditorWeightMapPaintInteractionType::PolyLasso && BrushAreaMode == EClothEditorWeightMapPaintBrushAreaType::Connected"))
+	UPROPERTY(EditAnywhere, Category = Filters, meta = (EditCondition = "SubToolType != EClothEditorWeightMapPaintInteractionType::PolyLasso"))
 	bool bUVSeams = false;
 
 	/** The Region affected by the current operation will be bounded by Hard Normal edges/seams */
-	UPROPERTY(EditAnywhere, Category = Filters, meta = (EditCondition = "SubToolType != EClothEditorWeightMapPaintInteractionType::PolyLasso && BrushAreaMode == EClothEditorWeightMapPaintBrushAreaType::Connected"))
+	UPROPERTY(EditAnywhere, Category = Filters, meta = (EditCondition = "SubToolType != EClothEditorWeightMapPaintInteractionType::PolyLasso"))
 	bool bNormalSeams = false;
 
 	/** Control which triangles can be affected by the current operation based on visibility. Applied after all other filters. */
