Cloth Editor:
- fix for tool not starting when launched from the tool panel
- fix for tool starting when more than one node is selected

diff --git a/Engine/Plugins/Experimental/ChaosClothAssetEditor/Source/ChaosClothAssetEditor/Private/ChaosClothAsset/ClothEditorMode.cpp b/Engine/Plugins/Experimental/ChaosClothAssetEditor/Source/ChaosClothAssetEditor/Private/ChaosClothAsset/ClothEditorMode.cpp
--- a/Engine/Plugins/Experimental/ChaosClothAssetEditor/Source/ChaosClothAssetEditor/Private/ChaosClothAsset/ClothEditorMode.cpp
+++ b/Engine/Plugins/Experimental/ChaosClothAssetEditor/Source/ChaosClothAssetEditor/Private/ChaosClothAsset/ClothEditorMode.cpp
@@ -202,7 +202,11 @@ void UChaosClothAssetEditorMode::RegisterAddNodeCommand(TSharedPtr<FUICommandInf
 		const FName ConnectionType = FManagedArrayCollection::StaticType();
 		UEdGraphNode* const CurrentlySelectedNode = GetSingleSelectedNodeWithOutputType(ConnectionType);
 		checkf(CurrentlySelectedNode, TEXT("No node with FManagedArrayCollection output is currently selected in the Dataflow graph"));
-		verifyf(CreateAndConnectNewNode(NewNodeType, *CurrentlySelectedNode, ConnectionType), TEXT("Failed to create a new node: %s"), *NewNodeType.ToString());
+
+		const UEdGraphNode* const NewNode = CreateAndConnectNewNode(NewNodeType, *CurrentlySelectedNode, ConnectionType);
+		verifyf(NewNode, TEXT("Failed to create a new node: %s"), *NewNodeType.ToString());
+
+		StartToolForSelectedNode(NewNode);
 	};
 
 	auto CanAddNode = [this](const FName& NewNodeType) -> bool
diff --git a/Engine/Plugins/Experimental/ChaosClothAssetEditor/Source/ChaosClothAssetEditor/Private/ChaosClothAsset/ClothEditorToolkit.cpp b/Engine/Plugins/Experimental/ChaosClothAssetEditor/Source/ChaosClothAssetEditor/Private/ChaosClothAsset/ClothEditorToolkit.cpp
--- a/Engine/Plugins/Experimental/ChaosClothAssetEditor/Source/ChaosClothAssetEditor/Private/ChaosClothAsset/ClothEditorToolkit.cpp
+++ b/Engine/Plugins/Experimental/ChaosClothAssetEditor/Source/ChaosClothAssetEditor/Private/ChaosClothAsset/ClothEditorToolkit.cpp
@@ -1030,16 +1030,19 @@ void FChaosClothAssetEditorToolkit::OnNodeSingleClicked(UObject* ClickedNode) co
 	UChaosClothAssetEditorMode* const ClothMode = CastChecked<UChaosClothAssetEditorMode>(EditorModeManager->GetActiveScriptableMode(UChaosClothAssetEditorMode::EM_ChaosClothAssetEditorModeId));
 	if (ClothMode)
 	{
-		// Close any running tool
-		UEditorInteractiveToolsContext* const ToolsContext = ClothMode->GetInteractiveToolsContext();
-		checkf(ToolsContext, TEXT("No valid ToolsContext found for FChaosClothAssetEditorToolkit"));
-		if (ToolsContext->HasActiveTool())
+		if (GraphEditor && GraphEditor->GetSingleSelectedNode() == ClickedNode)
 		{
-			ToolsContext->EndTool(EToolShutdownType::Completed);
-		}
+			// Close any running tool
+			UEditorInteractiveToolsContext* const ToolsContext = ClothMode->GetInteractiveToolsContext();
+			checkf(ToolsContext, TEXT("No valid ToolsContext found for FChaosClothAssetEditorToolkit"));
+			if (ToolsContext->HasActiveTool())
+			{
+				ToolsContext->EndTool(EToolShutdownType::Completed);
+			}
 
-		// Start the corresponding tool
-		ClothMode->StartToolForSelectedNode(ClickedNode);
+			// Start the corresponding tool
+			ClothMode->StartToolForSelectedNode(ClickedNode);
+		}
 	}
 }
 
