Refactored TransferInpainWeights method by splitting the body into multiple functions.
Added an option to specify custom sim mesh for weight transferring only.
Added more error checking, handling of empty meshes, and added comments.

#rnx

diff --git a/Engine/Plugins/Experimental/ChaosClothAssetEditor/Source/ChaosClothAssetDataflowNodes/Private/ChaosClothAsset/TransferSkinWeightsNode.cpp b/Engine/Plugins/Experimental/ChaosClothAssetEditor/Source/ChaosClothAssetDataflowNodes/Private/ChaosClothAsset/TransferSkinWeightsNode.cpp
--- a/Engine/Plugins/Experimental/ChaosClothAssetEditor/Source/ChaosClothAssetDataflowNodes/Private/ChaosClothAsset/TransferSkinWeightsNode.cpp
+++ b/Engine/Plugins/Experimental/ChaosClothAssetEditor/Source/ChaosClothAssetDataflowNodes/Private/ChaosClothAsset/TransferSkinWeightsNode.cpp
@@ -24,8 +24,11 @@
 
 #define LOCTEXT_NAMESPACE "ChaosClothAssetTransferSkinWeightsNode"
 
+static const FName InpaintWeightMaskName(TEXT("_InpaintWeightMask"));
+
 namespace UE::Chaos::ClothAsset::Private
 {
+	/** Convert the USkeletalMesh to DynamicMesh. */
 	static void SkeletalMeshToDynamicMesh(USkeletalMesh* FromSkeletalMeshAsset, int32 LodIndex, FDynamicMesh3& ToDynamicMesh)
 	{	
 		FMeshDescription SourceMesh;
@@ -34,8 +37,7 @@
 		Converter.Convert(&SourceMesh, ToDynamicMesh);
 	}
 
-	// Convert the cloth into DynamicMesh
-	// @todo: This should instead be handled by a cloth lod to dynamic mesh converter similar to functions in the ClothPatternToDynamicMesh.h
+	/** Convert the ClothCollection to DynamicMesh. */
 	static bool SimClothToDynamicMesh(
 		const TSharedRef<FManagedArrayCollection>& ClothCollection,
 		const FReferenceSkeleton& TargetRefSkeleton, // the reference skeleton to add to the dynamic mesh
@@ -44,16 +46,23 @@
 		using namespace UE::Geometry;
 		using namespace UE::Chaos::ClothAsset;
 
+		// Check if ClothCollection is empty
+		FCollectionClothFacade ClothFacade(ClothCollection);
+		const int32 NumVertices = ClothFacade.GetNumSimVertices3D();
+		const int32 NumFaces = ClothFacade.GetNumSimFaces();
+
+		if (NumVertices <= 0 || NumFaces <= 0)
+		{
+			UE_LOG(LogChaosClothAssetDataflowNodes, Warning, TEXT("TransferSkinWeightsNode: Failed to convert the Cloth Collection to Dynamic Mesh. Cloth Collection is empty."));
+			return false;
+		}
+
 		// Convert the sim mesh to DynamicMesh. 
 		// @todo: FTransferBoneWeights should accept raw data arrays (vertices/triangles) to avoid this conversion
 		FClothPatternToDynamicMesh PatternToDynamicMesh;
 		constexpr bool bDisableAttributes = false;
 		PatternToDynamicMesh.Convert(ClothCollection, INDEX_NONE, EClothPatternVertexType::Sim3D, WeldedSimMesh, bDisableAttributes);
 
-		WeldedSimMesh.EnableAttributes();
-		FMeshNormals NormalsUtil(&WeldedSimMesh);
-		NormalsUtil.InitializeOverlayToPerVertexNormals(WeldedSimMesh.Attributes()->PrimaryNormals());
-
 		// Setup the skeleton.
 		// @note we can't simply copy the bone attributes from the source USkeletalMesh because the cloth  
 		// asset reference skeleton comes from the USkeleton not the USkeletalMesh
@@ -66,58 +75,14 @@
 		return true;
 	}
 
-	static bool TransferInpaintWeights(
-		const FReferenceSkeleton& TargetRefSkeleton,
-		const double NormalThreshold,
-		const double RadiusPercentage,
-		const int32 NumSmoothingIterations,
-		const double SmoothingStrength,
-		bool bUseParallel,
-		const TSharedRef<FManagedArrayCollection>& ClothCollection,
-		const FString& InpaintMaskWeightMapName,
-		UE::Geometry::FTransferBoneWeights& TransferBoneWeights)
+	/** Copy the skin weights from DynamicMesh to Collection, handling split vertices. */
+	static void CopySkinWeightsFromDynamicMeshToCloth(const FDynamicMesh3& WeldedSimMesh, 
+													  const bool bUseParallel,
+													  const TSharedRef<FManagedArrayCollection>& ClothCollection)
 	{
 		using namespace UE::Geometry;
 
-		//
-		// Convert cloth sim mesh LOD to the welded dynamic sim mesh.
-		//
-		
-		// We cannot handle it if there are orphaned vertices, so strip those now.
-		UE::Chaos::ClothAsset::FClothGeometryTools::CleanupAndCompactMesh(ClothCollection);
-
-		FDynamicMesh3 WeldedSimMesh;
-		if (!ensure(SimClothToDynamicMesh(ClothCollection, TargetRefSkeleton, WeldedSimMesh)))
-		{
-			UE_LOG(LogChaosClothAssetDataflowNodes, Warning, TEXT("TransferSkinWeightsNode: Failed to weld the simulation mesh for LOD."));
-			return false;
-		}
-
-		//
-		// Transfer the weights from the body to the welded sim mesh.
-		//
-		TransferBoneWeights.NormalThreshold = FMathd::DegToRad * NormalThreshold;
-		TransferBoneWeights.SearchRadius = RadiusPercentage * WeldedSimMesh.GetBounds().DiagonalLength();
-		TransferBoneWeights.NumSmoothingIterations = NumSmoothingIterations;
-		TransferBoneWeights.SmoothingStrength = SmoothingStrength;
-		TransferBoneWeights.LayeredMeshSupport = true; // multilayerd clothing
-		TransferBoneWeights.ForceInpaintWeightMapName = FName(InpaintMaskWeightMapName);
-
-		if (!ensure(TransferBoneWeights.Validate() == EOperationValidationResult::Ok))
-		{
-			UE_LOG(LogChaosClothAssetDataflowNodes, Warning, TEXT("TransferSkinWeightsNode: Transfer method parameters are invalid."));
-			return false;
-		}
-		if (!ensure(TransferBoneWeights.TransferWeightsToMesh(WeldedSimMesh, FSkeletalMeshAttributes::DefaultSkinWeightProfileName)))
-		{
-			UE_LOG(LogChaosClothAssetDataflowNodes, Warning, TEXT("TransferSkinWeightsNode: Transferring skin weights failed."));
-			return false;
-		}
-
 		UE::Chaos::ClothAsset::FCollectionClothFacade ClothFacade(ClothCollection);
-		const FName InpaintWeightMaskName = TEXT("_InpaintWeightMask");
-		ClothFacade.AddWeightMap(InpaintWeightMaskName);
-		TArrayView<float> InpaintWeightMask = ClothFacade.GetWeightMap(InpaintWeightMaskName);
 
 		//
 		// Copy the new bone weight data from the welded sim mesh back to the cloth patterns.
@@ -133,7 +98,7 @@
 				SimMeshToDynamicMesh[NonManifoldMapping.GetOriginalNonManifoldVertexID(DynamicMeshVert)].Add(DynamicMeshVert);
 			}
 
-			ParallelFor(ClothFacade.GetNumSimVertices3D(), [&ClothFacade, &WeldedSimMesh, &SimMeshToDynamicMesh, &TransferBoneWeights, &InpaintWeightMask](int32 SimVertexID)
+			ParallelFor(ClothFacade.GetNumSimVertices3D(), [&ClothFacade, &WeldedSimMesh, &SimMeshToDynamicMesh](int32 SimVertexID)
 			{
 				const FDynamicMeshVertexSkinWeightsAttribute* const OutAttribute = WeldedSimMesh.Attributes()->GetSkinWeightsAttribute(FSkeletalMeshAttributes::DefaultSkinWeightProfileName);
 				checkSlow(OutAttribute);
@@ -150,13 +115,9 @@
 					checkSlow(WeldedSimMesh.IsVertex(WeldedID));
 					OutAttribute->GetValue(WeldedID, ClothFacade.GetSimBoneIndices()[SimVertexID],
 						ClothFacade.GetSimBoneWeights()[SimVertexID]);
-					InpaintWeightMask[SimVertexID] = TransferBoneWeights.MatchedVertices[WeldedID] ? 1.0f : 0.0f;
 				}
 				else
-				{
-					// TODO: The logic below should be moved to BoneWeights.h and implemented as Merge method which takes
-					// an array of bone weights and computes their average.
-					
+				{					
 					// Need to merge data because dynamic mesh split the original vertex
 					TMap<int32, TPair<float, int32>> CombinedData;
 					for (const int32 WeldedID : SimMeshToDynamicMesh[SimVertexID])
@@ -214,16 +175,12 @@
 					{
 						Weight *= WeightsSumRecip;
 					}
-
-					const int32 WeldedID = SimMeshToDynamicMesh[SimVertexID][0]; // Any welded id can be used here
-					InpaintWeightMask[SimVertexID] = TransferBoneWeights.MatchedVertices[WeldedID] ? 1.0f : 0.0f;
 				}
 			}, bUseParallel ? EParallelForFlags::None : EParallelForFlags::ForceSingleThread);
 		}
 		else
 		{
-
-			ParallelFor(WeldedSimMesh.MaxVertexID(), [&ClothFacade, &WeldedSimMesh, &NonManifoldMapping, &InpaintWeightMask, &TransferBoneWeights](int32 WeldedID)
+			ParallelFor(WeldedSimMesh.MaxVertexID(), [&ClothFacade, &WeldedSimMesh, &NonManifoldMapping](int32 WeldedID)
 			{
 				const FDynamicMeshVertexSkinWeightsAttribute* OutAttribute = WeldedSimMesh.Attributes()->GetSkinWeightsAttribute(FSkeletalMeshAttributes::DefaultSkinWeightProfileName);
 
@@ -232,27 +189,188 @@
 				checkSlow(WeldedID < ClothFacade.GetNumSimVertices3D());
 				OutAttribute->GetValue(WeldedID, ClothFacade.GetSimBoneIndices()[WeldedID],
 					ClothFacade.GetSimBoneWeights()[WeldedID]);
-				InpaintWeightMask[WeldedID] = TransferBoneWeights.MatchedVertices[WeldedID] ? 1.0f : 0.0f;
 			}, bUseParallel ? EParallelForFlags::None : EParallelForFlags::ForceSingleThread);
+		}
+	}
+
+	/** Copy the inpaint mask array to Collection, handling split vertices. */
+	static void CopyInpaintMapFromDynamicMeshToCloth(const FDynamicMesh3& WeldedSimMesh, 
+													 const FName WeightMapName,
+													 const bool bUseParallel, 
+													 const TArray<bool>& MatchedVertices,
+													 const TSharedRef<FManagedArrayCollection>& ClothCollection)
+	{
+
+		using namespace UE::Geometry;
+
+		UE::Chaos::ClothAsset::FCollectionClothFacade ClothFacade(ClothCollection);
+
+		TArrayView<float> InpaintWeightMask = ClothFacade.GetWeightMap(WeightMapName);
+
+		const FNonManifoldMappingSupport NonManifoldMapping(WeldedSimMesh);
+		TArray<TArray<int32>> SimMeshToDynamicMesh;
+		if (NonManifoldMapping.IsNonManifoldVertexInSource())
+		{
+			// WeldedSimMesh indices don't match cloth collection. 
+			SimMeshToDynamicMesh.SetNum(ClothFacade.GetNumSimVertices3D());
+			for (int32 DynamicMeshVert = 0; DynamicMeshVert < WeldedSimMesh.VertexCount(); ++DynamicMeshVert)
+			{
+				SimMeshToDynamicMesh[NonManifoldMapping.GetOriginalNonManifoldVertexID(DynamicMeshVert)].Add(DynamicMeshVert);
+			}
+
+			ParallelFor(ClothFacade.GetNumSimVertices3D(), [&ClothFacade, &WeldedSimMesh, &SimMeshToDynamicMesh, &InpaintWeightMask, &MatchedVertices](int32 SimVertexID)
+			{
+				if (!ensure(SimMeshToDynamicMesh[SimVertexID].Num() > 0))
+				{
+					return;
+				}
+				if (SimMeshToDynamicMesh[SimVertexID].Num() == 1)
+				{
+					const int32 WeldedID = SimMeshToDynamicMesh[SimVertexID][0];
+					checkSlow(WeldedSimMesh.IsVertex(WeldedID));
+					InpaintWeightMask[SimVertexID] = MatchedVertices[WeldedID] ? 1.0f : 0.0f;
+				}
+				else
+				{
+					const int32 WeldedID = SimMeshToDynamicMesh[SimVertexID][0]; // Any welded id can be used here
+					InpaintWeightMask[SimVertexID] = MatchedVertices[WeldedID] ? 1.0f : 0.0f;
+				}
+			}, bUseParallel ? EParallelForFlags::None : EParallelForFlags::ForceSingleThread);
+		}
+		else
+		{
+			ParallelFor(WeldedSimMesh.MaxVertexID(), [&ClothFacade, &WeldedSimMesh, &NonManifoldMapping, &InpaintWeightMask, &MatchedVertices](int32 WeldedID)
+			{
+				checkSlow(WeldedSimMesh.IsVertex(WeldedID));
+				checkSlow(WeldedID < ClothFacade.GetNumSimVertices3D());
+				InpaintWeightMask[WeldedID] = MatchedVertices[WeldedID] ? 1.0f : 0.0f;
+			}, bUseParallel ? EParallelForFlags::None : EParallelForFlags::ForceSingleThread);
+		}
+	}
+
+	/** Transfer skin weights to sim cloth. */
+	static bool TransferInpaintWeightsToSim(
+		const FReferenceSkeleton& TargetRefSkeleton,
+		const double NormalThreshold,
+		const double RadiusPercentage,
+		const int32 NumSmoothingIterations,
+		const double SmoothingStrength,
+		bool bUseParallel,
+		const TSharedRef<FManagedArrayCollection>& ClothCollection,
+		const FName& InpaintMaskWeightMapName,
+		UE::Geometry::FTransferBoneWeights& TransferBoneWeights,
+		FDynamicMesh3& WeldedSimMesh)
+	{
+		using namespace UE::Geometry;
+
+		// We cannot handle it if there are orphaned vertices, so strip those now.
+		UE::Chaos::ClothAsset::FClothGeometryTools::CleanupAndCompactMesh(ClothCollection);
+		UE::Chaos::ClothAsset::FCollectionClothFacade ClothFacade(ClothCollection);
+
+		// Convert cloth sim mesh LOD to the welded dynamic sim mesh.
+		if (!SimClothToDynamicMesh(ClothCollection, TargetRefSkeleton, WeldedSimMesh))
+		{
+			UE_LOG(LogChaosClothAssetDataflowNodes, Warning, TEXT("TransferSkinWeightsNode: Failed to weld the simulation mesh for LOD."));
+			return false;
+		}
+
+		// Transfer the weights from the body to the welded sim mesh.
+		TransferBoneWeights.NormalThreshold = FMathd::DegToRad * NormalThreshold;
+		TransferBoneWeights.SearchRadius = RadiusPercentage * WeldedSimMesh.GetBounds().DiagonalLength();
+		TransferBoneWeights.NumSmoothingIterations = NumSmoothingIterations;
+		TransferBoneWeights.SmoothingStrength = SmoothingStrength;
+		TransferBoneWeights.LayeredMeshSupport = true; // multilayerd clothing
+		TransferBoneWeights.ForceInpaintWeightMapName = InpaintMaskWeightMapName;
+
+		if (!ensure(TransferBoneWeights.Validate() == EOperationValidationResult::Ok))
+		{
+			UE_LOG(LogChaosClothAssetDataflowNodes, Warning, TEXT("TransferSkinWeightsNode: Transfer method parameters are invalid."));
+			return false;
+		}
+		if (!ensure(TransferBoneWeights.TransferWeightsToMesh(WeldedSimMesh, FSkeletalMeshAttributes::DefaultSkinWeightProfileName)))
+		{
+			UE_LOG(LogChaosClothAssetDataflowNodes, Warning, TEXT("TransferSkinWeightsNode: Transferring skin weights failed."));
+			return false;
+		}
+
+		
+		ClothFacade.AddWeightMap(InpaintWeightMaskName);
+
+		return true;
+	}
+
+	/** Transfer skin weights to sim and render cloth. */
+	static bool TransferInpaintWeights(
+		const FReferenceSkeleton& TargetRefSkeleton,
+		const double NormalThreshold,
+		const double RadiusPercentage,
+		const int32 NumSmoothingIterations,
+		const double SmoothingStrength,
+		bool bUseParallel,
+		const TSharedRef<FManagedArrayCollection>& ClothCollection,
+		const TSharedRef<FManagedArrayCollection>& UserTransferCollection,
+		const FName InpaintMaskWeightMapName,
+		UE::Geometry::FTransferBoneWeights& TransferBoneWeights)
+	{
+		using namespace UE::Geometry;
+
+		//
+		// Convert cloth sim mesh LOD to the welded dynamic sim mesh and transfer weights.
+		//
+		FDynamicMesh3 WeldedSimMesh;
+		if (!TransferInpaintWeightsToSim(TargetRefSkeleton, NormalThreshold, RadiusPercentage, NumSmoothingIterations, SmoothingStrength, bUseParallel, ClothCollection, InpaintMaskWeightMapName, TransferBoneWeights, WeldedSimMesh))
+		{
+			return false;
 		}
 
 		//
-		// Compute the bone weights for the render mesh by transferring weights from the sim mesh
+		// Copy the new bone weight data and inpaint mask from the welded sim mesh back to the sim cloth patterns.
+		//
+		CopySkinWeightsFromDynamicMeshToCloth(WeldedSimMesh, true, ClothCollection);
+		CopyInpaintMapFromDynamicMeshToCloth(WeldedSimMesh, InpaintWeightMaskName, true, TransferBoneWeights.MatchedVertices, ClothCollection);
+
+		//
+		// Check if the custom mesh collection for transferring weights is provided, in which case use it instead of the sim mesh
+		//
+		if (FCollectionClothFacade(UserTransferCollection).IsValid())
+		{
+			FDynamicMesh3 TransferWeightSimMesh;
+			if (TransferInpaintWeightsToSim(TargetRefSkeleton, NormalThreshold, RadiusPercentage, NumSmoothingIterations, SmoothingStrength, bUseParallel, UserTransferCollection, InpaintMaskWeightMapName, TransferBoneWeights, TransferWeightSimMesh))
+			{
+				WeldedSimMesh = MoveTemp(TransferWeightSimMesh); 
+			}
+			else
+			{
+				UE_LOG(LogChaosClothAssetDataflowNodes, Warning, TEXT("TransferSkinWeightsNode: Custom mesh collection for transferring weights is invalid, falling back to simulation mesh."));
+			}
+		}
+
+		//
+		// Compute the bone weights for the render mesh by transferring weights either from the sim mesh or custom user-specified mesh
 		// @todo If render mesh eventually supports welding, we should be transferring weights from the 
 		// body instead, same as we do for the sim mesh.
 		//
-		FTransferBoneWeights SimToRenderMeshTransfer(&WeldedSimMesh, FSkeletalMeshAttributes::DefaultSkinWeightProfileName);
-		SimToRenderMeshTransfer.bUseParallel = bUseParallel;
-		SimToRenderMeshTransfer.TransferMethod = FTransferBoneWeights::ETransferBoneWeightsMethod::ClosestPointOnSurface;
-		ParallelFor(ClothFacade.GetNumRenderVertices(), [&SimToRenderMeshTransfer, &ClothFacade](int32 VertexID)
+		if (WeldedSimMesh.VertexCount() > 0 && WeldedSimMesh.TriangleCount() > 0)
 		{
-			SimToRenderMeshTransfer.TransferWeightsToPoint(ClothFacade.GetRenderBoneIndices()[VertexID],
-				ClothFacade.GetRenderBoneWeights()[VertexID],
-				ClothFacade.GetRenderPosition()[VertexID]);
+			FTransferBoneWeights SimToRenderMeshTransfer(&WeldedSimMesh, FSkeletalMeshAttributes::DefaultSkinWeightProfileName);
+			SimToRenderMeshTransfer.bUseParallel = bUseParallel;
+			SimToRenderMeshTransfer.TransferMethod = FTransferBoneWeights::ETransferBoneWeightsMethod::ClosestPointOnSurface;
 
-		}, bUseParallel ? EParallelForFlags::None : EParallelForFlags::ForceSingleThread);
+			UE::Chaos::ClothAsset::FCollectionClothFacade ClothFacade(ClothCollection);
+			ParallelFor(ClothFacade.GetNumRenderVertices(), [&SimToRenderMeshTransfer, &ClothFacade](int32 VertexID)
+				{
+					SimToRenderMeshTransfer.TransferWeightsToPoint(ClothFacade.GetRenderBoneIndices()[VertexID],
+					ClothFacade.GetRenderBoneWeights()[VertexID],
+					ClothFacade.GetRenderPosition()[VertexID]);
 
-		return true;
+				}, bUseParallel ? EParallelForFlags::None : EParallelForFlags::ForceSingleThread);
+
+			return true;
+		}
+		else
+		{
+			return false;
+		}
 	}
 
 	static bool TransferClosestPointOnSurface(
@@ -315,6 +433,7 @@
 	RegisterInputConnection(&Collection);
 	RegisterOutputConnection(&Collection, &Collection);
 	RegisterInputConnection(&InpaintMask.WeightMap, GET_MEMBER_NAME_CHECKED(FChaosClothAssetWeightedValueNonAnimatableNoLowHighRange, WeightMap));
+	RegisterInputConnection(&TransferWeightsCollection);
 }
 
 void FChaosClothAssetTransferSkinWeightsNode::Evaluate(Dataflow::FContext& Context, const FDataflowOutput* Out) const
@@ -370,8 +489,10 @@
 			bool bTransferResult = false;
 			if (TransferMethod == EChaosClothAssetTransferSkinWeightsMethod::InpaintWeights)
 			{
+				FManagedArrayCollection TransferWeightsManagedCollection = GetValue<FManagedArrayCollection>(Context, &TransferWeightsCollection);
+				const TSharedRef<FManagedArrayCollection> ClothTransferWeightCollection = MakeShared<FManagedArrayCollection>(MoveTemp(TransferWeightsManagedCollection));
 				const FString& InpainMaskWeightMapName = GetValue<FString>(Context, &InpaintMask.WeightMap);
-				bTransferResult = TransferInpaintWeights(TargetRefSkeleton, NormalThreshold, RadiusPercentage, NumSmoothingIterations, SmoothingStrength, bUseParallel, ClothCollection, InpainMaskWeightMapName, TransferBoneWeights);
+				bTransferResult = TransferInpaintWeights(TargetRefSkeleton, NormalThreshold, RadiusPercentage, NumSmoothingIterations, SmoothingStrength, bUseParallel, ClothCollection, ClothTransferWeightCollection, FName(InpainMaskWeightMapName), TransferBoneWeights);
 			}
 			else if (TransferMethod == EChaosClothAssetTransferSkinWeightsMethod::ClosestPointOnSurface)
 			{
diff --git a/Engine/Plugins/Experimental/ChaosClothAssetEditor/Source/ChaosClothAssetDataflowNodes/Public/ChaosClothAsset/TransferSkinWeightsNode.h b/Engine/Plugins/Experimental/ChaosClothAssetEditor/Source/ChaosClothAssetDataflowNodes/Public/ChaosClothAsset/TransferSkinWeightsNode.h
--- a/Engine/Plugins/Experimental/ChaosClothAssetEditor/Source/ChaosClothAssetDataflowNodes/Public/ChaosClothAsset/TransferSkinWeightsNode.h
+++ b/Engine/Plugins/Experimental/ChaosClothAssetEditor/Source/ChaosClothAssetDataflowNodes/Public/ChaosClothAsset/TransferSkinWeightsNode.h
@@ -49,7 +49,7 @@ public:
 
 	/** Algorithm used for the transfer method. Use the simple ClosestPointOnSurface method or the more complex InpaintWeights method for better results. */
 	UPROPERTY(EditAnywhere, Category = "Transfer Skin Weights|Transfer Method", Meta = (DisplayName = "Algorithm"))
-	EChaosClothAssetTransferSkinWeightsMethod TransferMethod = EChaosClothAssetTransferSkinWeightsMethod::ClosestPointOnSurface;
+	EChaosClothAssetTransferSkinWeightsMethod TransferMethod = EChaosClothAssetTransferSkinWeightsMethod::InpaintWeights;
 
 	/**
 	 * Percentage of the bounding box diagonal of the simulation mesh to use as search radius for the InpaintWeights method.
@@ -78,6 +78,10 @@ public:
 	UPROPERTY(EditAnywhere, Category = "Transfer Skin Weights|Transfer Method", Meta = (DisplayName = "Inpaint Mask", EditCondition="TransferMethod==EChaosClothAssetTransferSkinWeightsMethod::InpaintWeights"))
 	FChaosClothAssetWeightedValueNonAnimatableNoLowHighRange InpaintMask = { TEXT("InpaintMask") };
 
+	/** Optional mesh that will be used for transferring weights to the render mesh instead of the simulation mesh. */
+	UPROPERTY(Meta = (Dataflowinput))
+	FManagedArrayCollection TransferWeightsCollection;
+
 	FChaosClothAssetTransferSkinWeightsNode(const Dataflow::FNodeParameters& InParam, FGuid InGuid = FGuid::NewGuid());
 
 private:
