Chaos Cloth - Added PostEditChangeProperty and Dataflow evaluation code to reformat and prevent the creation of non valid weightmap/selection/attribute names.

#rnx

diff --git a/Engine/Plugins/Experimental/ChaosClothAssetEditor/Source/ChaosClothAssetDataflowNodes/Private/ChaosClothAsset/ClothDataflowTools.cpp b/Engine/Plugins/Experimental/ChaosClothAssetEditor/Source/ChaosClothAssetDataflowNodes/Private/ChaosClothAsset/ClothDataflowTools.cpp
--- a/Engine/Plugins/Experimental/ChaosClothAssetEditor/Source/ChaosClothAssetDataflowNodes/Private/ChaosClothAsset/ClothDataflowTools.cpp
+++ b/Engine/Plugins/Experimental/ChaosClothAssetEditor/Source/ChaosClothAssetDataflowNodes/Private/ChaosClothAsset/ClothDataflowTools.cpp
@@ -450,10 +450,4 @@ namespace UE::Chaos::ClothAsset
 
 		UE_LOG(LogChaosClothAssetDataflowNodes, Display, TEXT("%s"), *Text.ToString());
 	}
-
-	FString FClothDataflowTools::MakeAttributeName(const FString& String)
-	{
-		return SlugStringForValidName(String, TEXT("_")).Replace(TEXT("\\"), TEXT("_"));
-	}
-
 }  // End namespace UE::Chaos::ClothAsset
diff --git a/Engine/Plugins/Experimental/ChaosClothAssetEditor/Source/ChaosClothAssetDataflowNodes/Private/ChaosClothAsset/ClothDataflowTools.h b/Engine/Plugins/Experimental/ChaosClothAssetEditor/Source/ChaosClothAssetDataflowNodes/Private/ChaosClothAsset/ClothDataflowTools.h
--- a/Engine/Plugins/Experimental/ChaosClothAssetEditor/Source/ChaosClothAssetDataflowNodes/Private/ChaosClothAsset/ClothDataflowTools.h
+++ b/Engine/Plugins/Experimental/ChaosClothAssetEditor/Source/ChaosClothAssetDataflowNodes/Private/ChaosClothAsset/ClothDataflowTools.h
@@ -22,8 +22,6 @@ namespace UE::Chaos::ClothAsset
 		static void AddSimPatternsFromSkeletalMeshSection(const TSharedRef<FManagedArrayCollection>& ClothCollection, const FSkeletalMeshLODModel& SkeletalMeshModel, const int32 SectionIndex, const int32 UVChannelIndex, const FVector2f& UVScale);
 
 		static void LogAndToastWarning(const FDataflowNode& DataflowNode, const FText& Headline, const FText& Details);
-
-		static FString MakeAttributeName(const FString& String);
 	};
 }  // End namespace UE::Chaos::ClothAsset
 
diff --git a/Engine/Plugins/Experimental/ChaosClothAssetEditor/Source/ChaosClothAssetDataflowNodes/Private/ChaosClothAsset/SimulationBaseConfigNodePrivate.h b/Engine/Plugins/Experimental/ChaosClothAssetEditor/Source/ChaosClothAssetDataflowNodes/Private/ChaosClothAsset/SimulationBaseConfigNodePrivate.h
--- a/Engine/Plugins/Experimental/ChaosClothAssetEditor/Source/ChaosClothAssetDataflowNodes/Private/ChaosClothAsset/SimulationBaseConfigNodePrivate.h
+++ b/Engine/Plugins/Experimental/ChaosClothAssetEditor/Source/ChaosClothAssetDataflowNodes/Private/ChaosClothAsset/SimulationBaseConfigNodePrivate.h
@@ -3,6 +3,7 @@
 #pragma once
 
 #include "Dataflow/DataflowInputOutput.h"
+#include "ChaosClothAsset/WeightedValue.h"
 
 /** Macro for adding and setting value properties from inside AddProperties() */
 #define UE_CHAOS_CLOTHASSET_SIMULATIONCONFIG_SETPROPERTY(Property) \
@@ -19,13 +20,6 @@
 	const int32 Property##KeyIndex = AddPropertyHelper(Properties, FName(TEXT(#Property))); \
 	Properties.SetValue(Property##KeyIndex, (int32)Property);
 
-/** Macro for adding and setting weighted value properties from inside AddProperties() */
-#define UE_CHAOS_CLOTHASSET_SIMULATIONCONFIG_SETPROPERTYWEIGHTED(Property) \
-	const int32 Property##KeyIndex = AddPropertyHelper(Properties, FName(TEXT(#Property)), Property.bIsAnimatable); \
-	Properties.SetWeightedValue(Property##KeyIndex, Property.Low, Property.High); \
-	Properties.SetStringValue(Property##KeyIndex, GetValue<FString>(Context, &Property.WeightMap)); \
-	Property.WeightMap_Override = GetValue<FString>(Context, &Property.WeightMap, FString());
-
 /** Macro for adding and setting weighted value properties from inside AddProperties() and check for a similar property being overriden. */
 #define UE_CHAOS_CLOTHASSET_SIMULATIONCONFIG_SETPROPERTYCHECKED1(Property, SimilarProperty1) \
 	const int32 Property##KeyIndex = AddPropertyHelper(Properties, FName(TEXT(#Property)), true, \
@@ -44,13 +38,22 @@
 		{ FName(TEXT(#SimilarProperty1)), FName(TEXT(#SimilarProperty2) )}); \
 	Properties.SetValue(Property##KeyIndex, (int32)Property);
 
+/** Macro for adding and setting weighted value properties from inside AddProperties() */
+#define UE_CHAOS_CLOTHASSET_SIMULATIONCONFIG_SETPROPERTYWEIGHTED(Property) \
+	const int32 Property##KeyIndex = AddPropertyHelper(Properties, FName(TEXT(#Property)), Property.bIsAnimatable); \
+	Properties.SetWeightedValue(Property##KeyIndex, Property.Low, Property.High); \
+	Properties.SetStringValue(Property##KeyIndex, GetValue<FString>(Context, &Property.WeightMap)); \
+	Property.WeightMap_Override = GetValue<FString>(Context, &Property.WeightMap, UE::Chaos::ClothAsset::FWeightMapTools::NotOverridden); \
+	UE::Chaos::ClothAsset::FWeightMapTools::MakeWeightMapName(Property.WeightMap);
+
 /** Macro for adding and setting weighted value properties from inside AddProperties() and check for a similar property being overriden. */
 #define UE_CHAOS_CLOTHASSET_SIMULATIONCONFIG_SETPROPERTYWEIGHTEDCHECKED1(Property, SimilarProperty1) \
 	const int32 Property##KeyIndex = AddPropertyHelper(Properties, FName(TEXT(#Property)), Property.bIsAnimatable, \
 		{ FName(TEXT(#SimilarProperty1)) }); \
 	Properties.SetWeightedValue(Property##KeyIndex, Property.Low, Property.High); \
 	Properties.SetStringValue(Property##KeyIndex, GetValue<FString>(Context, &Property.WeightMap)); \
-	Property.WeightMap_Override = GetValue<FString>(Context, &Property.WeightMap, FString());
+	Property.WeightMap_Override = GetValue<FString>(Context, &Property.WeightMap, UE::Chaos::ClothAsset::FWeightMapTools::NotOverridden); \
+	UE::Chaos::ClothAsset::FWeightMapTools::MakeWeightMapName(Property.WeightMap);
 
 /** Macro for adding and setting weighted value properties from inside AddProperties() and check for two similar properties being overriden. */
 #define UE_CHAOS_CLOTHASSET_SIMULATIONCONFIG_SETPROPERTYWEIGHTEDCHECKED2(Property, SimilarProperty1, SimilarProperty2) \
@@ -58,7 +61,8 @@
 		{ FName(TEXT(#SimilarProperty1)), FName(TEXT(#SimilarProperty2)) }); \
 	Properties.SetWeightedValue(Property##KeyIndex, Property.Low, Property.High); \
 	Properties.SetStringValue(Property##KeyIndex, GetValue<FString>(Context, &Property.WeightMap)); \
-	Property.WeightMap_Override = GetValue<FString>(Context, &Property.WeightMap, FString());
+	Property.WeightMap_Override = GetValue<FString>(Context, &Property.WeightMap, UE::Chaos::ClothAsset::FWeightMapTools::NotOverridden); \
+	UE::Chaos::ClothAsset::FWeightMapTools::MakeWeightMapName(Property.WeightMap);
 
 /** Macro for adding and setting weighted value properties from inside AddProperties() and check for three similar properties being overriden. */
 #define UE_CHAOS_CLOTHASSET_SIMULATIONCONFIG_SETPROPERTYWEIGHTEDCHECKED3(Property, SimilarProperty1, SimilarProperty2, SimilarProperty3) \
@@ -66,7 +70,8 @@
 		{ FName(TEXT(#SimilarProperty1)), FName(TEXT(#SimilarProperty2)), FName(TEXT(#SimilarProperty3)) }); \
 	Properties.SetWeightedValue(Property##KeyIndex, Property.Low, Property.High); \
 	Properties.SetStringValue(Property##KeyIndex, GetValue<FString>(Context, &Property.WeightMap)); \
-	Property.WeightMap_Override = GetValue<FString>(Context, &Property.WeightMap, FString());
+	Property.WeightMap_Override = GetValue<FString>(Context, &Property.WeightMap, UE::Chaos::ClothAsset::FWeightMapTools::NotOverridden); \
+	UE::Chaos::ClothAsset::FWeightMapTools::MakeWeightMapName(Property.WeightMap);
 
 /** Macro for adding and setting weighted value properties from inside AddProperties() and check for four similar properties being overriden. */
 #define UE_CHAOS_CLOTHASSET_SIMULATIONCONFIG_SETPROPERTYWEIGHTEDCHECKED4(Property, SimilarProperty1, SimilarProperty2, SimilarProperty3, SimilarProperty4) \
@@ -74,4 +79,6 @@
 		{ FName(TEXT(#SimilarProperty1)), FName(TEXT(#SimilarProperty2)), FName(TEXT(#SimilarProperty3)), FName(TEXT(#SimilarProperty4)) }); \
 	Properties.SetWeightedValue(Property##KeyIndex, Property.Low, Property.High); \
 	Properties.SetStringValue(Property##KeyIndex, GetValue<FString>(Context, &Property.WeightMap)); \
-	Property.WeightMap_Override = GetValue<FString>(Context, &Property.WeightMap, FString());
+	Property.WeightMap_Override = GetValue<FString>(Context, &Property.WeightMap, UE::Chaos::ClothAsset::FWeightMapTools::NotOverridden); \
+	UE::Chaos::ClothAsset::FWeightMapTools::MakeWeightMapName(Property.WeightMap);
+
diff --git a/Engine/Plugins/Experimental/ChaosClothAssetEditor/Source/ChaosClothAssetDataflowNodes/Private/ChaosClothAsset/WeightedValueCustomization.cpp b/Engine/Plugins/Experimental/ChaosClothAssetEditor/Source/ChaosClothAssetDataflowNodes/Private/ChaosClothAsset/WeightedValueCustomization.cpp
--- a/Engine/Plugins/Experimental/ChaosClothAssetEditor/Source/ChaosClothAssetDataflowNodes/Private/ChaosClothAsset/WeightedValueCustomization.cpp
+++ b/Engine/Plugins/Experimental/ChaosClothAssetEditor/Source/ChaosClothAssetDataflowNodes/Private/ChaosClothAsset/WeightedValueCustomization.cpp
@@ -1,6 +1,7 @@
 // Copyright Epic Games, Inc. All Rights Reserved.
 
 #include "ChaosClothAsset/WeightedValueCustomization.h"
+#include "ChaosClothAsset/WeightedValue.h"
 #include "Widgets/Input/SCheckBox.h"
 #include "Widgets/Input/SNumericEntryBox.h"
 #include "Widgets/Input/SEditableTextBox.h"
@@ -168,8 +169,9 @@ namespace UE::Chaos::ClothAsset
 						{
 							OverrideHandlePtr->GetValue(Text);
 						}
-						if (Text.IsEmpty())
+						if (Text == UE::Chaos::ClothAsset::FWeightMapTools::NotOverridden)
 						{
+							Text.Empty();  // GetValue seems to concatenate the text if the string isn't emptied first
 							if (const TSharedPtr<IPropertyHandle> HandlePtr = HandleWeakPtr.Pin())
 							{
 								HandlePtr->GetValue(Text);
@@ -192,7 +194,7 @@ namespace UE::Chaos::ClothAsset
 						{
 							OverrideHandlePtr->GetValue(Text);
 						}
-						return Text.IsEmpty();
+						return Text == UE::Chaos::ClothAsset::FWeightMapTools::NotOverridden;
 					})
 				.Font(IPropertyTypeCustomizationUtils::GetRegularFont());
 		}
diff --git a/Engine/Plugins/Experimental/ChaosClothAssetEditor/Source/ChaosClothAssetDataflowNodes/Public/ChaosClothAsset/AddWeightMapNode.h b/Engine/Plugins/Experimental/ChaosClothAssetEditor/Source/ChaosClothAssetDataflowNodes/Public/ChaosClothAsset/AddWeightMapNode.h
--- a/Engine/Plugins/Experimental/ChaosClothAssetEditor/Source/ChaosClothAssetDataflowNodes/Public/ChaosClothAsset/AddWeightMapNode.h
+++ b/Engine/Plugins/Experimental/ChaosClothAssetEditor/Source/ChaosClothAssetDataflowNodes/Public/ChaosClothAsset/AddWeightMapNode.h
@@ -20,7 +20,7 @@ public:
 
 	/** The name to be set as a weight map attribute. */
 	UPROPERTY(EditAnywhere, Category = "Add Weight Map", Meta = (DataflowOutput))
-	FString Name;
+	mutable FString Name;  // Mutable so that it can be name checked in the evaluate function
 
 	UPROPERTY()
 	TArray<float> VertexWeights;
diff --git a/Engine/Plugins/Experimental/ChaosClothAssetEditor/Source/ChaosClothAssetDataflowNodes/Public/ChaosClothAsset/WeightedValue.h b/Engine/Plugins/Experimental/ChaosClothAssetEditor/Source/ChaosClothAssetDataflowNodes/Public/ChaosClothAsset/WeightedValue.h
--- a/Engine/Plugins/Experimental/ChaosClothAssetEditor/Source/ChaosClothAssetDataflowNodes/Public/ChaosClothAsset/WeightedValue.h
+++ b/Engine/Plugins/Experimental/ChaosClothAssetEditor/Source/ChaosClothAssetDataflowNodes/Public/ChaosClothAsset/WeightedValue.h
@@ -5,6 +5,30 @@
 #include "Containers/UnrealString.h"
 #include "WeightedValue.generated.h"
 
+namespace UE::Chaos::ClothAsset
+{
+	struct FWeightMapTools
+	{
+		/**
+		 * String code used to disable the override field in the UI.
+		 * Uses a character that is removed from the names made by MakeWeightMapName to avoid conflicts.
+		 */
+		inline static const FString NotOverridden = TEXT("#");
+
+		/**
+		 * Modifies a string to make it suitable as an attribute name.
+		 * Replaces any deemed special characters and spaces by underscores, and remove leading and ending underscores from the name.
+		 * Note that leading underscores are reserved for non user/internal attribute names.
+		 */
+		static void MakeWeightMapName(FString& InOutString)
+		{
+			InOutString = SlugStringForValidName(InOutString, TEXT("_")).Replace(TEXT("\\"), TEXT("_"));
+			bool bCharsWereRemoved;
+			do { InOutString.TrimCharInline(TEXT('_'), &bCharsWereRemoved); } while (bCharsWereRemoved);
+		}
+	};
+}
+
 USTRUCT()
 struct FChaosClothAssetWeightedValue
 {
@@ -39,11 +63,11 @@ struct FChaosClothAssetWeightedValue
 
 	/** The name of the weight map for this property. */
 	UPROPERTY(EditAnywhere, Category = "Weighted Value", Meta = (DataflowInput))
-	FString WeightMap = TEXT("WeightMap");
+	mutable FString WeightMap = TEXT("WeightMap");  // Mutable so that it can be name checked in the evaluate function
 
 	/** The weight map override value for when the WeightMap has a connection that replaces the provided weight map value. */
 	UPROPERTY(VisibleAnywhere, Transient, Category = "Weighted Value")
-	mutable FString WeightMap_Override;  // _Override has a special meaning to the property customization, mutable because this property is set while getting the original value
+	mutable FString WeightMap_Override = UE::Chaos::ClothAsset::FWeightMapTools::NotOverridden;  // _Override has a special meaning to the property customization, mutable because this property is set while getting the original value
 };
 
 USTRUCT()
@@ -75,11 +99,11 @@ struct FChaosClothAssetWeightedValueNonAnimatable
 
 	/** The name of the weight map for this property. */
 	UPROPERTY(EditAnywhere, Category = "Weighted Value", Meta = (DataflowInput))
-	FString WeightMap = TEXT("WeightMap");
+	mutable FString WeightMap = TEXT("WeightMap");  // Mutable so that it can be name checked in the evaluate function
 
 	/** The weight map override value for when the WeightMap has a connection that replaces the provided weight map value. */
 	UPROPERTY(VisibleAnywhere, Transient, Category = "Weighted Value")
-	mutable FString WeightMap_Override;  // _Override has a special meaning to the property customization, mutable because this property is set while getting the original value
+	mutable FString WeightMap_Override = UE::Chaos::ClothAsset::FWeightMapTools::NotOverridden;  // _Override has a special meaning to the property customization, mutable because this property is set while getting the original value
 };
 
 
diff --git a/Engine/Plugins/Experimental/ChaosClothAssetEditor/Source/ChaosClothAssetEditorTools/Private/ChaosClothAsset/ClothMeshSelectionTool.cpp b/Engine/Plugins/Experimental/ChaosClothAssetEditor/Source/ChaosClothAssetEditorTools/Private/ChaosClothAsset/ClothMeshSelectionTool.cpp
--- a/Engine/Plugins/Experimental/ChaosClothAssetEditor/Source/ChaosClothAssetEditorTools/Private/ChaosClothAsset/ClothMeshSelectionTool.cpp
+++ b/Engine/Plugins/Experimental/ChaosClothAssetEditor/Source/ChaosClothAssetEditorTools/Private/ChaosClothAsset/ClothMeshSelectionTool.cpp
@@ -5,6 +5,7 @@
 #include "ChaosClothAsset/ClothPatternVertexType.h"
 #include "ChaosClothAsset/SelectionNode.h"
 #include "ChaosClothAsset/CollectionClothFacade.h"
+#include "ChaosClothAsset/WeightedValue.h"
 #include "InteractiveToolManager.h"
 #include "PreviewMesh.h"
 #include "TargetInterfaces/PrimitiveComponentBackedTarget.h"
@@ -61,6 +62,16 @@ UInteractiveTool* UClothMeshSelectionToolBuilder::BuildTool(const FToolBuilderSt
 	return NewTool;
 }
 
+// ------------------- Properties -------------------
+
+void UClothMeshSelectionToolProperties::PostEditChangeProperty(FPropertyChangedEvent& PropertyChangedEvent)
+{
+	if (PropertyChangedEvent.GetPropertyName() == GET_MEMBER_NAME_CHECKED(UClothMeshSelectionToolProperties, Name))
+	{
+		UE::Chaos::ClothAsset::FWeightMapTools::MakeWeightMapName(Name);
+	}
+}
+
 
 // ------------------- Tool -------------------
 	
diff --git a/Engine/Plugins/Experimental/ChaosClothAssetEditor/Source/ChaosClothAssetEditorTools/Private/ChaosClothAsset/ClothWeightMapPaintTool.cpp b/Engine/Plugins/Experimental/ChaosClothAssetEditor/Source/ChaosClothAssetEditorTools/Private/ChaosClothAsset/ClothWeightMapPaintTool.cpp
--- a/Engine/Plugins/Experimental/ChaosClothAssetEditor/Source/ChaosClothAssetEditorTools/Private/ChaosClothAsset/ClothWeightMapPaintTool.cpp
+++ b/Engine/Plugins/Experimental/ChaosClothAssetEditor/Source/ChaosClothAssetEditorTools/Private/ChaosClothAsset/ClothWeightMapPaintTool.cpp
@@ -34,6 +34,7 @@
 #include "ChaosClothAsset/ClothEditorContextObject.h"
 #include "ChaosClothAsset/ClothPatternVertexType.h"
 #include "ChaosClothAsset/CollectionClothFacade.h"
+#include "ChaosClothAsset/WeightedValue.h"
 #include "ContextObjectStore.h"
 #include "Dataflow/DataflowEdNode.h"
 #include "Dataflow/DataflowObject.h"
@@ -88,7 +89,7 @@ UMeshSurfacePointTool* UClothEditorWeightMapPaintToolBuilder::CreateNewTool(cons
 
 
 /*
- * Properties
+ * ToolActions
  */
 void UClothEditorMeshWeightMapPaintToolActions::PostAction(EClothEditorWeightMapPaintToolActions Action)
 {
@@ -99,7 +100,16 @@ void UClothEditorMeshWeightMapPaintToolActions::PostAction(EClothEditorWeightMap
 }
 
 
-
+/*
+ * Properties
+ */
+void UClothEditorUpdateWeightMapProperties::PostEditChangeProperty(FPropertyChangedEvent& PropertyChangedEvent)
+{
+	if (PropertyChangedEvent.GetPropertyName() == GET_MEMBER_NAME_CHECKED(UClothEditorUpdateWeightMapProperties, Name))
+	{
+		UE::Chaos::ClothAsset::FWeightMapTools::MakeWeightMapName(Name);
+	}
+}
 
 
 /*
diff --git a/Engine/Plugins/Experimental/ChaosClothAssetEditor/Source/ChaosClothAssetEditorTools/Public/ChaosClothAsset/ClothMeshSelectionTool.h b/Engine/Plugins/Experimental/ChaosClothAssetEditor/Source/ChaosClothAssetEditorTools/Public/ChaosClothAsset/ClothMeshSelectionTool.h
--- a/Engine/Plugins/Experimental/ChaosClothAssetEditor/Source/ChaosClothAssetEditorTools/Public/ChaosClothAsset/ClothMeshSelectionTool.h
+++ b/Engine/Plugins/Experimental/ChaosClothAssetEditor/Source/ChaosClothAssetEditorTools/Public/ChaosClothAsset/ClothMeshSelectionTool.h
@@ -45,6 +45,10 @@ public:
 	UPROPERTY(EditAnywhere, Category = Visualization, meta = (DisplayName = "Show Edges"))
 	bool bShowEdges = false;
 
+private:
+
+	virtual void PostEditChangeProperty(FPropertyChangedEvent& PropertyChangedEvent) override;
+
 };
 
 UCLASS()
diff --git a/Engine/Plugins/Experimental/ChaosClothAssetEditor/Source/ChaosClothAssetEditorTools/Public/ChaosClothAsset/ClothWeightMapPaintTool.h b/Engine/Plugins/Experimental/ChaosClothAssetEditor/Source/ChaosClothAssetEditorTools/Public/ChaosClothAsset/ClothWeightMapPaintTool.h
--- a/Engine/Plugins/Experimental/ChaosClothAssetEditor/Source/ChaosClothAssetEditorTools/Public/ChaosClothAsset/ClothWeightMapPaintTool.h
+++ b/Engine/Plugins/Experimental/ChaosClothAssetEditor/Source/ChaosClothAssetEditorTools/Public/ChaosClothAsset/ClothWeightMapPaintTool.h
@@ -222,6 +222,10 @@ public:
 
 	UPROPERTY(EditAnywhere, Category = UpdateNode, meta = (DisplayName = "Name"))
 	FString Name;
+
+private:
+
+	virtual void PostEditChangeProperty(struct FPropertyChangedEvent& PropertyChangedEvent) override;
 };
 
 /**
