Cloth Editor Weight Map Paint Tool:
- hide/show tool properties for different combinations of tool and brush types
- add high/low values for the gradient tool
- remove an unimplemented visibility filter from the UI

diff --git a/Engine/Plugins/Experimental/ChaosClothAssetEditor/Source/ChaosClothAssetEditorTools/Private/ChaosClothAsset/ClothWeightMapPaintTool.cpp b/Engine/Plugins/Experimental/ChaosClothAssetEditor/Source/ChaosClothAssetEditorTools/Private/ChaosClothAsset/ClothWeightMapPaintTool.cpp
--- a/Engine/Plugins/Experimental/ChaosClothAssetEditor/Source/ChaosClothAssetEditorTools/Private/ChaosClothAsset/ClothWeightMapPaintTool.cpp
+++ b/Engine/Plugins/Experimental/ChaosClothAssetEditor/Source/ChaosClothAssetEditorTools/Private/ChaosClothAsset/ClothWeightMapPaintTool.cpp
@@ -1071,7 +1071,7 @@ void UClothEditorWeightMapPaintTool::ComputeGradient()
 			}
 		}
 
-		const float Value = FMath::LerpStable(0.0f, 1.0f, DistanceToLowSq / (DistanceToLowSq + DistanceToHighSq));
+		const float Value = FMath::LerpStable(FilterProperties->GradientLowValue, FilterProperties->GradientHighValue, DistanceToLowSq / (DistanceToLowSq + DistanceToHighSq));
 		if (bHaveDynamicMeshToWeightConversion)
 		{
 			for (const int32 Idx : WeightToDynamicMesh[DynamicMeshToWeight[VertexIndex]])
diff --git a/Engine/Plugins/Experimental/ChaosClothAssetEditor/Source/ChaosClothAssetEditorTools/Public/ChaosClothAsset/ClothWeightMapPaintTool.h b/Engine/Plugins/Experimental/ChaosClothAssetEditor/Source/ChaosClothAssetEditorTools/Public/ChaosClothAsset/ClothWeightMapPaintTool.h
--- a/Engine/Plugins/Experimental/ChaosClothAssetEditor/Source/ChaosClothAssetEditorTools/Public/ChaosClothAsset/ClothWeightMapPaintTool.h
+++ b/Engine/Plugins/Experimental/ChaosClothAssetEditor/Source/ChaosClothAssetEditorTools/Public/ChaosClothAsset/ClothWeightMapPaintTool.h
@@ -114,7 +114,6 @@ UENUM()
 enum class EClothEditorWeightMapPaintVisibilityType : uint8
 {
 	None,
-	FrontFacing,
 	Unoccluded
 };
 
@@ -133,36 +132,56 @@ public:
 	EClothEditorWeightMapPaintInteractionType SubToolType = EClothEditorWeightMapPaintInteractionType::Brush;
 
 	UPROPERTY(EditAnywhere, Category = ActionType, meta = (DisplayName = "Brush Mode", 
-		HideEditConditionToggle, EditConditionHides, EditCondition = "SubToolType != EClothEditorWeightMapPaintInteractionType::PolyLasso"))
+		HideEditConditionToggle, EditConditionHides, EditCondition = "SubToolType == EClothEditorWeightMapPaintInteractionType::Brush"))
 	EClothEditorWeightMapPaintBrushType PrimaryBrushType = EClothEditorWeightMapPaintBrushType::Paint;
 
 	/** Relative size of brush */
 	UPROPERTY(EditAnywhere, Category = Brush, meta = (DisplayName = "Brush Size", UIMin = "0.0", UIMax = "1.0", ClampMin = "0.0", ClampMax = "10.0", 
-		HideEditConditionToggle, EditConditionHides, EditCondition = "SubToolType != EClothEditorWeightMapPaintInteractionType::PolyLasso"))
+		HideEditConditionToggle, EditConditionHides, EditCondition = "SubToolType == EClothEditorWeightMapPaintInteractionType::Brush"))
 	float BrushSize = 0.25f;
 
 	/** The new value to paint on the mesh */
-	UPROPERTY(EditAnywhere, Category = Brush, meta = (UIMin = 0, ClampMin = 0, UIMax = 1, ClampMax = 1))
+	UPROPERTY(EditAnywhere, Category = Brush, meta = (UIMin = 0, ClampMin = 0, UIMax = 1, ClampMax = 1,
+		HideEditConditionToggle, EditConditionHides, EditCondition = 
+		"(SubToolType == EClothEditorWeightMapPaintInteractionType::Brush && PrimaryBrushType == EClothEditorWeightMapPaintBrushType::Paint) || SubToolType == EClothEditorWeightMapPaintInteractionType::Fill || SubToolType == EClothEditorWeightMapPaintInteractionType::PolyLasso"))
 	double AttributeValue = 1;
 
 	/** How quickly each brush stroke will drive mesh values towards the desired value */
-	UPROPERTY(EditAnywhere, Category = Brush, meta = (UIMin = 0, ClampMin = 0, UIMax = 1, ClampMax = 1))
+	UPROPERTY(EditAnywhere, Category = Brush, meta = (UIMin = 0, ClampMin = 0, UIMax = 1, ClampMax = 1,
+		HideEditConditionToggle, EditConditionHides, EditCondition = 
+		"SubToolType == EClothEditorWeightMapPaintInteractionType::Brush || SubToolType == EClothEditorWeightMapPaintInteractionType::Fill"))
 	double Strength = 0.5;
 
+	/** The Gradient upper limit value */
+	UPROPERTY(EditAnywhere, Category = Gradient, meta = (UIMin = 0, ClampMin = 0, UIMax = 1, ClampMax = 1,
+		HideEditConditionToggle, EditConditionHides, EditCondition = "SubToolType == EClothEditorWeightMapPaintInteractionType::Gradient"))
+	double GradientHighValue = 1.0;
+
+	/** The Gradient lower limit value */
+	UPROPERTY(EditAnywhere, Category = Gradient, meta = (UIMin = 0, ClampMin = 0, UIMax = 1, ClampMax = 1,
+		HideEditConditionToggle, EditConditionHides, EditCondition = "SubToolType == EClothEditorWeightMapPaintInteractionType::Gradient"))
+	double GradientLowValue = 0.0;
+
+
 	/** The Region affected by the current operation will be bounded by edge angles larger than this threshold */
-	UPROPERTY(EditAnywhere, Category = Filters, meta = (UIMin = "0.0", UIMax = "180.0", EditCondition = "SubToolType != EClothEditorWeightMapPaintInteractionType::PolyLasso"))
+	UPROPERTY(EditAnywhere, Category = Filters, meta = (UIMin = "0.0", ClampMin = "0.0", UIMax = "180.0", ClampMax = "180.0",
+		HideEditConditionToggle, EditConditionHides, EditCondition = 
+		"SubToolType == EClothEditorWeightMapPaintInteractionType::Brush || SubToolType == EClothEditorWeightMapPaintInteractionType::Fill"))
 	float AngleThreshold = 180.0f;
 
 	/** The Region affected by the current operation will be bounded by UV borders/seams */
-	UPROPERTY(EditAnywhere, Category = Filters, meta = (EditCondition = "SubToolType != EClothEditorWeightMapPaintInteractionType::PolyLasso"))
+	UPROPERTY(EditAnywhere, Category = Filters, meta = (HideEditConditionToggle, EditConditionHides, EditCondition = 
+		"SubToolType == EClothEditorWeightMapPaintInteractionType::Brush || SubToolType == EClothEditorWeightMapPaintInteractionType::Fill"))
 	bool bUVSeams = false;
 
 	/** The Region affected by the current operation will be bounded by Hard Normal edges/seams */
-	UPROPERTY(EditAnywhere, Category = Filters, meta = (EditCondition = "SubToolType != EClothEditorWeightMapPaintInteractionType::PolyLasso"))
+	UPROPERTY(EditAnywhere, Category = Filters, meta = (HideEditConditionToggle, EditConditionHides, EditCondition = 
+		"SubToolType == EClothEditorWeightMapPaintInteractionType::Brush || SubToolType == EClothEditorWeightMapPaintInteractionType::Fill"))
 	bool bNormalSeams = false;
 
 	/** Control which triangles can be affected by the current operation based on visibility. Applied after all other filters. */
-	UPROPERTY(EditAnywhere, Category = Filters)
+	UPROPERTY(EditAnywhere, Category = Filters, meta = (HideEditConditionToggle, EditConditionHides, EditCondition =
+		"SubToolType == EClothEditorWeightMapPaintInteractionType::Brush || SubToolType == EClothEditorWeightMapPaintInteractionType::Fill || SubToolType == EClothEditorWeightMapPaintInteractionType::PolyLasso"))
 	EClothEditorWeightMapPaintVisibilityType VisibilityFilter = EClothEditorWeightMapPaintVisibilityType::None;
 
 	/** The weight value at the brush indicator */
