Chaos Cloth Editor: disable wireframe mode when using paint and selection tools.

This follows how changing construction view mode works with tools (e.g., adding a new virtual method to check if wireframe is allowed, going back to wireframe mode if it was disabled by the tool).

diff --git a/Engine/Plugins/Experimental/ChaosClothAssetEditor/Source/ChaosClothAssetEditor/Private/ChaosClothAsset/ClothEditorMode.cpp b/Engine/Plugins/Experimental/ChaosClothAssetEditor/Source/ChaosClothAssetEditor/Private/ChaosClothAsset/ClothEditorMode.cpp
--- a/Engine/Plugins/Experimental/ChaosClothAssetEditor/Source/ChaosClothAssetEditor/Private/ChaosClothAsset/ClothEditorMode.cpp
+++ b/Engine/Plugins/Experimental/ChaosClothAssetEditor/Source/ChaosClothAssetEditor/Private/ChaosClothAsset/ClothEditorMode.cpp
@@ -182,6 +182,17 @@ void UChaosClothAssetEditorMode::RegisterClothTool(TSharedPtr<FUICommandInfo> UI
 				this->SetConstructionViewMode(SupportedModes[0]);
 			}
 
+			// Check if we need to disable wireframe mode before starting tool.
+			const bool bCanSetWireframeActive = ClothToolBuilder->CanSetConstructionViewWireframeActive();
+			if (!bCanSetWireframeActive)
+			{
+				if (!bShouldRestoreConstructionViewWireframe)
+				{
+					bShouldRestoreConstructionViewWireframe = bConstructionViewWireframe;
+				}
+				bConstructionViewWireframe = false;
+			}
+
 			ActiveToolsContext = ToolsContext;
 			ToolsContext->StartTool(ToolIdentifier);
 		}),
@@ -286,6 +297,12 @@ void UChaosClothAssetEditorMode::OnToolEnded(UInteractiveToolManager* Manager, U
 {
 	UE::Chaos::ClothAsset::FChaosClothAssetEditorCommands::UpdateToolCommandBinding(Tool, ToolCommandList, true);
 
+	if (bShouldRestoreConstructionViewWireframe)
+	{
+		bConstructionViewWireframe = true;
+		bShouldRestoreConstructionViewWireframe = false;
+	}
+
 	if (bShouldRestoreSavedConstructionViewMode)
 	{
 		SetConstructionViewMode(SavedConstructionViewMode);
@@ -976,6 +993,21 @@ void UChaosClothAssetEditorMode::ToggleConstructionViewWireframe()
 	ReinitializeDynamicMeshComponents();
 }
 
+bool UChaosClothAssetEditorMode::CanSetConstructionViewWireframeActive() const
+{
+	if (!GetToolManager()->HasActiveTool(EToolSide::Left))
+	{
+		return true;
+	}
+
+	const UInteractiveToolBuilder* const ActiveToolBuilder = GetToolManager()->GetActiveToolBuilder(EToolSide::Left);
+	checkf(ActiveToolBuilder, TEXT("No Active Tool Builder found despite having an Active Tool"));
+
+	const IChaosClothAssetEditorToolBuilder* const ClothToolBuilder = Cast<const IChaosClothAssetEditorToolBuilder>(ActiveToolBuilder);
+	checkf(ClothToolBuilder, TEXT("Cloth Editor has an active Tool Builder that does not implement IChaosClothAssetEditorToolBuilder"));
+	return ClothToolBuilder->CanSetConstructionViewWireframeActive();
+}
+
 void UChaosClothAssetEditorMode::SetRestSpaceViewportClient(TWeakPtr<UE::Chaos::ClothAsset::FChaosClothEditorRestSpaceViewportClient, ESPMode::ThreadSafe> InViewportClient)
 {
 	RestSpaceViewportClient = InViewportClient;
diff --git a/Engine/Plugins/Experimental/ChaosClothAssetEditor/Source/ChaosClothAssetEditor/Private/ChaosClothAsset/ClothEditorMode.h b/Engine/Plugins/Experimental/ChaosClothAssetEditor/Source/ChaosClothAssetEditor/Private/ChaosClothAsset/ClothEditorMode.h
--- a/Engine/Plugins/Experimental/ChaosClothAssetEditor/Source/ChaosClothAssetEditor/Private/ChaosClothAsset/ClothEditorMode.h
+++ b/Engine/Plugins/Experimental/ChaosClothAssetEditor/Source/ChaosClothAssetEditor/Private/ChaosClothAsset/ClothEditorMode.h
@@ -66,6 +66,7 @@ public:
 	bool CanChangeConstructionViewModeTo(UE::Chaos::ClothAsset::EClothPatternVertexType NewViewMode) const;
 
 	void ToggleConstructionViewWireframe();
+	bool CanSetConstructionViewWireframeActive() const;
 	bool IsConstructionViewWireframeActive() const
 	{
 		return bConstructionViewWireframe;
@@ -232,6 +233,7 @@ private:
 	FName NodeTypeForPendingToolStart;
 
 	bool bConstructionViewWireframe = false;
+	bool bShouldRestoreConstructionViewWireframe = false;
 
 	// Create dynamic mesh components from the cloth component's rest space info
 	void ReinitializeDynamicMeshComponents();
diff --git a/Engine/Plugins/Experimental/ChaosClothAssetEditor/Source/ChaosClothAssetEditorTools/Public/ChaosClothAsset/ClothEditorToolBuilder.h b/Engine/Plugins/Experimental/ChaosClothAssetEditor/Source/ChaosClothAssetEditorTools/Public/ChaosClothAsset/ClothEditorToolBuilder.h
--- a/Engine/Plugins/Experimental/ChaosClothAssetEditor/Source/ChaosClothAssetEditorTools/Public/ChaosClothAsset/ClothEditorToolBuilder.h
+++ b/Engine/Plugins/Experimental/ChaosClothAssetEditor/Source/ChaosClothAssetEditorTools/Public/ChaosClothAsset/ClothEditorToolBuilder.h
@@ -27,5 +27,8 @@ public:
 
 	/** Returns all Construction View modes that this tool can operate in. The first element should be the preferred mode to switch to if necessary. */
 	virtual void GetSupportedViewModes(TArray<UE::Chaos::ClothAsset::EClothPatternVertexType>& Modes) const = 0;
+
+	/** Returns whether or not view can be set to wireframe when this tool is active.. */
+	virtual bool CanSetConstructionViewWireframeActive() const { return true; }
 };
 
diff --git a/Engine/Plugins/Experimental/ChaosClothAssetEditor/Source/ChaosClothAssetEditorTools/Public/ChaosClothAsset/ClothMeshSelectionTool.h b/Engine/Plugins/Experimental/ChaosClothAssetEditor/Source/ChaosClothAssetEditorTools/Public/ChaosClothAsset/ClothMeshSelectionTool.h
--- a/Engine/Plugins/Experimental/ChaosClothAssetEditor/Source/ChaosClothAssetEditorTools/Public/ChaosClothAsset/ClothMeshSelectionTool.h
+++ b/Engine/Plugins/Experimental/ChaosClothAssetEditor/Source/ChaosClothAssetEditorTools/Public/ChaosClothAsset/ClothMeshSelectionTool.h
@@ -27,6 +27,7 @@ private:
 	virtual bool CanBuildTool(const FToolBuilderState& SceneState) const override;
 	virtual UInteractiveTool* BuildTool(const FToolBuilderState& SceneState) const override;
 	virtual const FToolTargetTypeRequirements& GetTargetRequirements() const override;
+	virtual bool CanSetConstructionViewWireframeActive() const { return false; }
 };
 
 UCLASS()
diff --git a/Engine/Plugins/Experimental/ChaosClothAssetEditor/Source/ChaosClothAssetEditorTools/Public/ChaosClothAsset/ClothWeightMapPaintTool.h b/Engine/Plugins/Experimental/ChaosClothAssetEditor/Source/ChaosClothAssetEditorTools/Public/ChaosClothAsset/ClothWeightMapPaintTool.h
--- a/Engine/Plugins/Experimental/ChaosClothAssetEditor/Source/ChaosClothAssetEditorTools/Public/ChaosClothAsset/ClothWeightMapPaintTool.h
+++ b/Engine/Plugins/Experimental/ChaosClothAssetEditor/Source/ChaosClothAssetEditorTools/Public/ChaosClothAsset/ClothWeightMapPaintTool.h
@@ -61,6 +61,7 @@ class CHAOSCLOTHASSETEDITORTOOLS_API UClothEditorWeightMapPaintToolBuilder : pub
 private:
 	virtual void GetSupportedViewModes(TArray<UE::Chaos::ClothAsset::EClothPatternVertexType>& Modes) const override;
 	virtual UMeshSurfacePointTool* CreateNewTool(const FToolBuilderState& SceneState) const override;
+	virtual bool CanSetConstructionViewWireframeActive() const { return false; }
 };
 
 
