Cloth Editor: add toolbar icons, remove placeholders

diff --git a/Engine/Plugins/Experimental/ChaosClothAssetEditor/Source/ChaosClothAssetEditor/Private/ChaosClothAsset/ClothEditorStyle.cpp b/Engine/Plugins/Experimental/ChaosClothAssetEditor/Source/ChaosClothAssetEditor/Private/ChaosClothAsset/ClothEditorStyle.cpp
--- a/Engine/Plugins/Experimental/ChaosClothAssetEditor/Source/ChaosClothAssetEditor/Private/ChaosClothAsset/ClothEditorStyle.cpp
+++ b/Engine/Plugins/Experimental/ChaosClothAssetEditor/Source/ChaosClothAssetEditor/Private/ChaosClothAsset/ClothEditorStyle.cpp
@@ -22,12 +22,9 @@ FString FChaosClothAssetEditorStyle::InContent(const FString& RelativePath, cons
 FChaosClothAssetEditorStyle::FChaosClothAssetEditorStyle()
 	: FSlateStyleSet(StyleName)
 {
-	// Used FUVEditorStyle and FModelingToolsEditorModeStyle and as models
-
 	SetCoreContentRoot(FPaths::EngineContentDir() / TEXT("Editor/Slate"));
 
-	// TODO: Set this once we get our own UI content
-	FString ContentDir = IPluginManager::Get().FindPlugin(TEXT("ModelingToolsEditorMode"))->GetContentDir();
+	FString ContentDir = IPluginManager::Get().FindPlugin(TEXT("ChaosClothAssetEditor"))->GetContentDir();
 	SetContentRoot(ContentDir);
 
 	// Some standard icon sizes used elsewhere in the editor
@@ -42,47 +39,37 @@ FChaosClothAssetEditorStyle::FChaosClothAssetEditorStyle()
 	const FVector2D ViewportToolbarIconSize = Icon16x16;
 	const FVector2D ToolbarIconSize = Icon20x20;
 
-	FString PropertyNameString = "ChaosClothAssetEditor." + FChaosClothAssetEditorCommands::BeginRemeshToolIdentifier;
-	Set(*PropertyNameString, new FSlateImageBrush(FChaosClothAssetEditorStyle::InContent("Icons/Remesh_40x", ".png"), ToolbarIconSize));
-
-	PropertyNameString = "ChaosClothAssetEditor." + FChaosClothAssetEditorCommands::BeginAttributeEditorToolIdentifier;
-	Set(*PropertyNameString, new FSlateImageBrush(FChaosClothAssetEditorStyle::InContent("Icons/AttributeEditor_40x", ".png"), ToolbarIconSize));
+	FString PropertyNameString = "ChaosClothAssetEditor." + FChaosClothAssetEditorCommands::AddWeightMapNodeIdentifier;
+	Set(*PropertyNameString, new IMAGE_BRUSH_SVG("Icons/PaintMaps", ToolbarIconSize));
 
-	PropertyNameString = "ChaosClothAssetEditor." + FChaosClothAssetEditorCommands::AddWeightMapNodeIdentifier;
-	Set(*PropertyNameString, new FSlateImageBrush(FChaosClothAssetEditorStyle::InContent("Icons/ModelingAttributePaint_x40", ".png"), ToolbarIconSize));
+	PropertyNameString = "ChaosClothAssetEditor." + FChaosClothAssetEditorCommands::AddMeshSelectionNodeIdentifier;
+	Set(*PropertyNameString, new IMAGE_BRUSH_SVG("Icons/MeshSelect", ToolbarIconSize));
 
-	// TODO: Get a permanent icon to replace this placeholder
 	PropertyNameString = "ChaosClothAssetEditor." + FChaosClothAssetEditorCommands::AddTransferSkinWeightsNodeIdentifier;
-	Set(*PropertyNameString, new FSlateImageBrush(FChaosClothAssetEditorStyle::InContent("Icons/ModelingAttributePaint_x40", ".png"), ToolbarIconSize));
+	Set(*PropertyNameString, new IMAGE_BRUSH_SVG("Icons/SkinTransfer", ToolbarIconSize));
 
 	//
 	// Construction Viewport
 	// 
 	
-	// TODO: Get a permanent icon to replace this placeholder
 	PropertyNameString = "ChaosClothAssetEditor." + FChaosClothAssetEditorCommands::ToggleConstructionViewWireframeIdentifier;
-	Set(*PropertyNameString, new CORE_IMAGE_BRUSH_SVG("Starship/EditorViewport/grid", ToolbarIconSize));
+	Set(*PropertyNameString, new IMAGE_BRUSH_SVG("Icons/ClothWireframe_16", ViewportToolbarIconSize));
 
 	//
 	// Preview Viewport
 	// 
 
-	// TODO: Get a permanent icon to replace this placeholder
 	PropertyNameString = "ChaosClothAssetEditor." + FChaosClothAssetEditorCommands::TogglePreviewWireframeIdentifier;
-	Set(*PropertyNameString, new CORE_IMAGE_BRUSH_SVG("Starship/EditorViewport/grid", ToolbarIconSize));
+	Set(*PropertyNameString, new IMAGE_BRUSH_SVG("Icons/ClothWireframe_16", ViewportToolbarIconSize));
 
-	// TODO: Get a permanent icon to replace this placeholder
-	// TODO: Toggle back and forth between Play and Pause icons when the button is clicked?
 	PropertyNameString = "ChaosClothAssetEditor." + FChaosClothAssetEditorCommands::ToggleSimulationSuspendedIdentifier;
-	Set(*PropertyNameString, new CORE_IMAGE_BRUSH_SVG("Starship/MainToolbar/pause", ToolbarIconSize));
+	Set(*PropertyNameString, new IMAGE_BRUSH_SVG("Icons/ClothSimSuspend_16", ViewportToolbarIconSize));
 
-	// TODO: Get a permanent icon to replace this placeholder
 	PropertyNameString = "ChaosClothAssetEditor." + FChaosClothAssetEditorCommands::SoftResetSimulationIdentifier;
-	Set(*PropertyNameString, new CORE_IMAGE_BRUSH_SVG("Starship/Common/Reset", ToolbarIconSize));
+	Set(*PropertyNameString, new IMAGE_BRUSH_SVG("Icons/ResetSoft_16", ViewportToolbarIconSize));
 
-	// TODO: Get a permanent icon to replace this placeholder
 	PropertyNameString = "ChaosClothAssetEditor." + FChaosClothAssetEditorCommands::HardResetSimulationIdentifier;
-	Set(*PropertyNameString, new CORE_IMAGE_BRUSH_SVG("Starship/Common/Apply", ToolbarIconSize));
+	Set(*PropertyNameString, new IMAGE_BRUSH_SVG("Icons/ResetHard_16", ViewportToolbarIconSize));
 
 	FSlateStyleRegistry::RegisterSlateStyle(*this);
 }
