Chaos Cloth Editor - Fixed the skin weight transfer node evaluation when properties are changed from within the tool.

#rnx

diff --git a/Engine/Plugins/Experimental/ChaosClothAssetEditor/Source/ChaosClothAssetEditorTools/Private/ChaosClothAsset/ClothTransferSkinWeightsTool.cpp b/Engine/Plugins/Experimental/ChaosClothAssetEditor/Source/ChaosClothAssetEditorTools/Private/ChaosClothAsset/ClothTransferSkinWeightsTool.cpp
--- a/Engine/Plugins/Experimental/ChaosClothAssetEditor/Source/ChaosClothAssetEditorTools/Private/ChaosClothAsset/ClothTransferSkinWeightsTool.cpp
+++ b/Engine/Plugins/Experimental/ChaosClothAssetEditor/Source/ChaosClothAssetEditorTools/Private/ChaosClothAsset/ClothTransferSkinWeightsTool.cpp
@@ -206,6 +206,7 @@ void UClothTransferSkinWeightsTool::Shutdown(EToolShutdownType ShutdownType)
 	{
 		TransferSkinWeightsNode->SkeletalMesh = ToolProperties->SourceMesh;
 		TransferSkinWeightsNode->Transform = TransformFromProperties();
+		TransferSkinWeightsNode->Invalidate();
 	}
 
 	if (SourceMeshTransformProxy)
