Cloth Editor:
- Add a camera-centered point light in the Construction viewport
- Turn off directional and sky lights in that viewport
- Remove the selection highlighting from that viewport (since everything will be selected all the time anyway)

diff --git a/Engine/Plugins/Experimental/ChaosClothAssetEditor/Source/ChaosClothAssetEditor/Private/ChaosClothAsset/ClothEditorRestSpaceViewportClient.cpp b/Engine/Plugins/Experimental/ChaosClothAssetEditor/Source/ChaosClothAssetEditor/Private/ChaosClothAsset/ClothEditorRestSpaceViewportClient.cpp
--- a/Engine/Plugins/Experimental/ChaosClothAssetEditor/Source/ChaosClothAssetEditor/Private/ChaosClothAsset/ClothEditorRestSpaceViewportClient.cpp
+++ b/Engine/Plugins/Experimental/ChaosClothAssetEditor/Source/ChaosClothAssetEditor/Private/ChaosClothAsset/ClothEditorRestSpaceViewportClient.cpp
@@ -13,6 +13,8 @@
 #include "Components/DynamicMeshComponent.h"
 #include "Framework/Application/SlateApplication.h"
 #include "SceneView.h"
+#include "PreviewScene.h"
+#include "Components/PointLightComponent.h"
 
 namespace UE::Chaos::ClothAsset
 {
@@ -51,8 +53,15 @@ FChaosClothEditorRestSpaceViewportClient::FChaosClothEditorRestSpaceViewportClie
 	BehaviorSet->Add(ZoomBehavior);
 	BehaviorsFor2DMode.Add(ZoomBehavior);
 
-	EngineShowFlags.SetSelectionOutline(true);
+	EngineShowFlags.SetSelectionOutline(false);
 	ModeTools->GetInteractiveToolsContext()->InputRouter->RegisterSource(this);
+
+	CameraPointLight = NewObject<UPointLightComponent>();
+	CameraPointLight->bUseInverseSquaredFalloff = false;
+	CameraPointLight->LightFalloffExponent = 2.0f;
+	CameraPointLight->SetIntensity(3.0f);		// TODO: Hook this up to a slider in the UI?
+	CameraPointLight->SetCastShadows(false);
+	PreviewScene->AddComponent(CameraPointLight, FTransform());
 }
 
 void FChaosClothEditorRestSpaceViewportClient::SetConstructionViewMode(EClothPatternVertexType InViewMode)
@@ -115,6 +124,15 @@ bool FChaosClothEditorRestSpaceViewportClient::ShouldOrbitCamera() const
 	}
 }
 
+void FChaosClothEditorRestSpaceViewportClient::Tick(float DeltaSeconds)
+{
+	FViewportCameraTransform ViewTransform = GetViewTransform();
+	CameraPointLight->SetRelativeLocation(ViewTransform.GetLocation());
+
+	FEditorViewportClient::Tick(DeltaSeconds);
+}
+
+
 bool FChaosClothEditorRestSpaceViewportClient::InputKey(const FInputKeyEventArgs& EventArgs)
 {
 	// See if any tool commands want to handle the key event
diff --git a/Engine/Plugins/Experimental/ChaosClothAssetEditor/Source/ChaosClothAssetEditor/Private/ChaosClothAsset/ClothEditorRestSpaceViewportClient.h b/Engine/Plugins/Experimental/ChaosClothAssetEditor/Source/ChaosClothAssetEditor/Private/ChaosClothAsset/ClothEditorRestSpaceViewportClient.h
--- a/Engine/Plugins/Experimental/ChaosClothAssetEditor/Source/ChaosClothAssetEditor/Private/ChaosClothAsset/ClothEditorRestSpaceViewportClient.h
+++ b/Engine/Plugins/Experimental/ChaosClothAssetEditor/Source/ChaosClothAssetEditor/Private/ChaosClothAsset/ClothEditorRestSpaceViewportClient.h
@@ -10,6 +10,7 @@
 #include "ChaosClothAsset/ClothPatternVertexType.h"
 
 class UInputBehaviorSet;
+class UPointLightComponent;
 
 namespace UE::Chaos::ClothAsset
 {
@@ -39,6 +40,10 @@ public:
 
 private:
 
+	virtual void Tick(float DeltaSeconds) override;
+
+	TObjectPtr<UPointLightComponent> CameraPointLight;
+
 	EClothPatternVertexType ConstructionViewMode = EClothPatternVertexType::Sim2D;
 
 	TObjectPtr<UInputBehaviorSet> BehaviorSet;
diff --git a/Engine/Plugins/Experimental/ChaosClothAssetEditor/Source/ChaosClothAssetEditor/Private/ChaosClothAsset/ClothEditorToolkit.cpp b/Engine/Plugins/Experimental/ChaosClothAssetEditor/Source/ChaosClothAssetEditor/Private/ChaosClothAsset/ClothEditorToolkit.cpp
--- a/Engine/Plugins/Experimental/ChaosClothAssetEditor/Source/ChaosClothAssetEditor/Private/ChaosClothAsset/ClothEditorToolkit.cpp
+++ b/Engine/Plugins/Experimental/ChaosClothAssetEditor/Source/ChaosClothAssetEditor/Private/ChaosClothAsset/ClothEditorToolkit.cpp
@@ -200,9 +200,8 @@ FChaosClothAssetEditorToolkit::FChaosClothAssetEditorToolkit(UAssetEditor* InOwn
 			.ToolkitCommandList(GetToolkitCommands().ToSharedPtr());
 	};
 
-	FPreviewScene::ConstructionValues SceneArgs;
-	ObjectScene = MakeUnique<FPreviewScene>(SceneArgs);
-
+	// Construction view scene
+	ObjectScene = MakeUnique<FPreviewScene>(FPreviewScene::ConstructionValues().SetSkyBrightness(0.0f).SetLightBrightness(0.0f));
 }
 
 FChaosClothAssetEditorToolkit::~FChaosClothAssetEditorToolkit()
