Cloth Editor: add a UI slider to control the brightness/intensity of the camera point light in render mode

diff --git a/Engine/Plugins/Experimental/ChaosClothAssetEditor/Source/ChaosClothAssetEditor/Private/ChaosClothAsset/ClothEditorRestSpaceViewportClient.cpp b/Engine/Plugins/Experimental/ChaosClothAssetEditor/Source/ChaosClothAssetEditor/Private/ChaosClothAsset/ClothEditorRestSpaceViewportClient.cpp
--- a/Engine/Plugins/Experimental/ChaosClothAssetEditor/Source/ChaosClothAssetEditor/Private/ChaosClothAsset/ClothEditorRestSpaceViewportClient.cpp
+++ b/Engine/Plugins/Experimental/ChaosClothAssetEditor/Source/ChaosClothAssetEditor/Private/ChaosClothAsset/ClothEditorRestSpaceViewportClient.cpp
@@ -59,7 +59,7 @@ FChaosClothEditorRestSpaceViewportClient::FChaosClothEditorRestSpaceViewportClie
 	CameraPointLight = NewObject<UPointLightComponent>();
 	CameraPointLight->bUseInverseSquaredFalloff = false;
 	CameraPointLight->LightFalloffExponent = 2.0f;
-	CameraPointLight->SetIntensity(3.0f);		// TODO: Hook this up to a slider in the UI?
+	CameraPointLight->SetIntensity(3.0f);
 	CameraPointLight->SetCastShadows(false);
 	PreviewScene->AddComponent(CameraPointLight, FTransform());
 }
@@ -99,6 +99,12 @@ void FChaosClothEditorRestSpaceViewportClient::SetConstructionViewMode(EClothPat
 	ModeTools->GetInteractiveToolsContext()->InputRouter->RegisterSource(this);
 }
 
+
+EClothPatternVertexType FChaosClothEditorRestSpaceViewportClient::GetConstructionViewMode() const
+{
+	return ConstructionViewMode;
+}
+
 const UInputBehaviorSet* FChaosClothEditorRestSpaceViewportClient::GetInputBehaviors() const
 {
 	return BehaviorSet;
@@ -170,4 +176,16 @@ void FChaosClothEditorRestSpaceViewportClient::SetToolCommandList(TWeakPtr<FUICo
 {
 	ToolCommandList = InToolCommandList;
 }
+
+float FChaosClothEditorRestSpaceViewportClient::GetCameraPointLightIntensity() const
+{
+	return CameraPointLight->Intensity;
+}
+
+void FChaosClothEditorRestSpaceViewportClient::SetCameraPointLightIntensity(float Intensity)
+{
+	CameraPointLight->SetIntensity(Intensity);
+}
+
+
 } // namespace UE::Chaos::ClothAsset
diff --git a/Engine/Plugins/Experimental/ChaosClothAssetEditor/Source/ChaosClothAssetEditor/Private/ChaosClothAsset/ClothEditorRestSpaceViewportClient.h b/Engine/Plugins/Experimental/ChaosClothAssetEditor/Source/ChaosClothAssetEditor/Private/ChaosClothAsset/ClothEditorRestSpaceViewportClient.h
--- a/Engine/Plugins/Experimental/ChaosClothAssetEditor/Source/ChaosClothAssetEditor/Private/ChaosClothAsset/ClothEditorRestSpaceViewportClient.h
+++ b/Engine/Plugins/Experimental/ChaosClothAssetEditor/Source/ChaosClothAssetEditor/Private/ChaosClothAsset/ClothEditorRestSpaceViewportClient.h
@@ -34,10 +34,14 @@ public:
 	virtual bool ShouldOrbitCamera() const override;
 
 	void SetConstructionViewMode(EClothPatternVertexType InViewMode);
+	EClothPatternVertexType GetConstructionViewMode() const;
 
 	void SetEditorViewportWidget(TWeakPtr<SEditorViewport> InEditorViewportWidget);
 	void SetToolCommandList(TWeakPtr<FUICommandList> ToolCommandList);
 
+	float GetCameraPointLightIntensity() const;
+	void SetCameraPointLightIntensity(float Intensity);
+
 private:
 
 	virtual void Tick(float DeltaSeconds) override;
diff --git a/Engine/Plugins/Experimental/ChaosClothAssetEditor/Source/ChaosClothAssetEditor/Private/ChaosClothAsset/ClothEditorToolkit.cpp b/Engine/Plugins/Experimental/ChaosClothAssetEditor/Source/ChaosClothAssetEditor/Private/ChaosClothAsset/ClothEditorToolkit.cpp
--- a/Engine/Plugins/Experimental/ChaosClothAssetEditor/Source/ChaosClothAssetEditor/Private/ChaosClothAsset/ClothEditorToolkit.cpp
+++ b/Engine/Plugins/Experimental/ChaosClothAssetEditor/Source/ChaosClothAssetEditor/Private/ChaosClothAsset/ClothEditorToolkit.cpp
@@ -344,7 +344,7 @@ AssetEditorViewportFactoryFunction FChaosClothAssetEditorToolkit::GetViewportDel
 	AssetEditorViewportFactoryFunction TempViewportDelegate = [this](FAssetEditorViewportConstructionArgs InArgs)
 	{
 		return SAssignNew(RestSpaceViewportWidget, SChaosClothAssetEditorRestSpaceViewport, InArgs)
-			.EditorViewportClient(ViewportClient);
+			.RestSpaceViewportClient(StaticCastSharedPtr<FChaosClothEditorRestSpaceViewportClient>(ViewportClient));
 	};
 
 	return TempViewportDelegate;
diff --git a/Engine/Plugins/Experimental/ChaosClothAssetEditor/Source/ChaosClothAssetEditor/Private/ChaosClothAsset/SClothEditorRestSpaceViewport.cpp b/Engine/Plugins/Experimental/ChaosClothAssetEditor/Source/ChaosClothAssetEditor/Private/ChaosClothAsset/SClothEditorRestSpaceViewport.cpp
--- a/Engine/Plugins/Experimental/ChaosClothAssetEditor/Source/ChaosClothAssetEditor/Private/ChaosClothAsset/SClothEditorRestSpaceViewport.cpp
+++ b/Engine/Plugins/Experimental/ChaosClothAssetEditor/Source/ChaosClothAssetEditor/Private/ChaosClothAsset/SClothEditorRestSpaceViewport.cpp
@@ -1,6 +1,7 @@
 // Copyright Epic Games, Inc. All Rights Reserved.
 
 #include "ChaosClothAsset/SClothEditorRestSpaceViewport.h"
+#include "ChaosClothAsset/ClothEditorRestSpaceViewportClient.h"
 #include "SViewportToolBar.h"
 #include "ChaosClothAsset/SClothEditorRestSpaceViewportToolBar.h"
 #include "ChaosClothAsset/ClothEditorCommands.h"
@@ -11,8 +12,10 @@
 
 void SChaosClothAssetEditorRestSpaceViewport::Construct(const FArguments& InArgs, const FAssetEditorViewportConstructionArgs& InViewportConstructionArgs)
 {
+	RestSpaceViewportClient = InArgs._RestSpaceViewportClient;
+
 	SAssetEditorViewport::FArguments ParentArgs;
-	ParentArgs._EditorViewportClient = InArgs._EditorViewportClient;
+	ParentArgs._EditorViewportClient = InArgs._RestSpaceViewportClient;
 	if (InArgs._ViewportSize.IsSet())
 	{
 		ParentArgs._ViewportSize = InArgs._ViewportSize;
@@ -153,7 +156,8 @@ void SChaosClothAssetEditorRestSpaceViewport::BindCommands()
 TSharedPtr<SWidget> SChaosClothAssetEditorRestSpaceViewport::MakeViewportToolbar()
 {
 	return SNew(SChaosClothAssetEditorRestSpaceViewportToolBar, SharedThis(this))
-		.CommandList(CommandList);
+		.CommandList(CommandList)
+		.RestSpaceViewportClient(RestSpaceViewportClient);
 }
 
 
diff --git a/Engine/Plugins/Experimental/ChaosClothAssetEditor/Source/ChaosClothAssetEditor/Private/ChaosClothAsset/SClothEditorRestSpaceViewport.h b/Engine/Plugins/Experimental/ChaosClothAssetEditor/Source/ChaosClothAssetEditor/Private/ChaosClothAsset/SClothEditorRestSpaceViewport.h
--- a/Engine/Plugins/Experimental/ChaosClothAssetEditor/Source/ChaosClothAssetEditor/Private/ChaosClothAsset/SClothEditorRestSpaceViewport.h
+++ b/Engine/Plugins/Experimental/ChaosClothAssetEditor/Source/ChaosClothAssetEditor/Private/ChaosClothAsset/SClothEditorRestSpaceViewport.h
@@ -5,8 +5,10 @@
 #include "CoreMinimal.h"
 #include "SBaseCharacterFXEditorViewport.h"
 #include "SCommonEditorViewportToolbarBase.h"
+#include "SClothEditorRestSpaceViewportToolBar.h"
 
 class UChaosClothAssetEditorMode;
+namespace UE::Chaos::ClothAsset { class FChaosClothEditorRestSpaceViewportClient; }
 
 class CHAOSCLOTHASSETEDITOR_API SChaosClothAssetEditorRestSpaceViewport : public SBaseCharacterFXEditorViewport, public ICommonEditorViewportToolbarInfoProvider
 {
@@ -14,7 +16,7 @@ public:
 
 	SLATE_BEGIN_ARGS(SChaosClothAssetEditorRestSpaceViewport) {}
 		SLATE_ATTRIBUTE(FVector2D, ViewportSize);
-		SLATE_ARGUMENT(TSharedPtr<FEditorViewportClient>, EditorViewportClient)
+		SLATE_ARGUMENT(TSharedPtr<UE::Chaos::ClothAsset::FChaosClothEditorRestSpaceViewportClient>, RestSpaceViewportClient)
 	SLATE_END_ARGS()
 
 	void Construct(const FArguments& InArgs, const FAssetEditorViewportConstructionArgs& InViewportConstructionArgs);
@@ -34,4 +36,6 @@ private:
 
 	UChaosClothAssetEditorMode* GetEdMode() const;
 
+	TSharedPtr<UE::Chaos::ClothAsset::FChaosClothEditorRestSpaceViewportClient> RestSpaceViewportClient;
+
 };
diff --git a/Engine/Plugins/Experimental/ChaosClothAssetEditor/Source/ChaosClothAssetEditor/Private/ChaosClothAsset/SClothEditorRestSpaceViewportToolBar.cpp b/Engine/Plugins/Experimental/ChaosClothAssetEditor/Source/ChaosClothAssetEditor/Private/ChaosClothAsset/SClothEditorRestSpaceViewportToolBar.cpp
--- a/Engine/Plugins/Experimental/ChaosClothAssetEditor/Source/ChaosClothAssetEditor/Private/ChaosClothAsset/SClothEditorRestSpaceViewportToolBar.cpp
+++ b/Engine/Plugins/Experimental/ChaosClothAssetEditor/Source/ChaosClothAssetEditor/Private/ChaosClothAsset/SClothEditorRestSpaceViewportToolBar.cpp
@@ -2,6 +2,7 @@
 
 #include "ChaosClothAsset/SClothEditorRestSpaceViewportToolBar.h"
 #include "ChaosClothAsset/SClothEditorRestSpaceViewport.h"
+#include "ChaosClothAsset/ClothEditorRestSpaceViewportClient.h"
 #include "ChaosClothAsset/ClothEditorCommands.h"
 #include "EditorViewportCommands.h"
 #include "Styling/AppStyle.h"
@@ -13,11 +14,14 @@
 #include "EditorModeManager.h"
 #include "ChaosClothAsset/ClothEditorMode.h"
 #include "ToolMenus.h"
+#include "Widgets/Input/SSpinBox.h"
 
 #define LOCTEXT_NAMESPACE "SChaosClothAssetEditorRestSpaceViewportToolBar"
 
 void SChaosClothAssetEditorRestSpaceViewportToolBar::Construct(const FArguments& InArgs, TSharedPtr<SChaosClothAssetEditorRestSpaceViewport> InChaosClothAssetEditorViewport)
 {
+	RestSpaceViewportClient = InArgs._RestSpaceViewportClient;
+
 	SCommonEditorViewportToolbarBase::Construct(SCommonEditorViewportToolbarBase::FArguments(), InChaosClothAssetEditorViewport);
 
 	ChaosClothAssetEditorRestSpaceViewportPtr = InChaosClothAssetEditorViewport;
@@ -60,13 +64,85 @@ void SChaosClothAssetEditorRestSpaceViewportToolBar::Construct(const FArguments&
 		];
 }
 
+TSharedRef<SWidget> SChaosClothAssetEditorRestSpaceViewportToolBar::GenerateClothRestSpaceViewportOptionsMenu()
+{
+	GetInfoProvider().OnFloatingButtonClicked();
+	TSharedRef<SEditorViewport> ViewportRef = GetInfoProvider().GetViewportWidget();
+
+	const bool bIsPerspective = GetViewportClient().GetViewportType() == LVT_Perspective;
+	const bool bInShouldCloseWindowAfterMenuSelection = true;
+
+	FMenuBuilder OptionsMenuBuilder(bInShouldCloseWindowAfterMenuSelection, ViewportRef->GetCommandList());
+	{
+		OptionsMenuBuilder.BeginSection("ClothEditorViewportOptions", LOCTEXT("OptionsMenuHeader", "Viewport Options"));
+		{
+			if (bIsPerspective)
+			{
+				OptionsMenuBuilder.AddWidget(GenerateFOVMenu(), LOCTEXT("FOVAngle", "Field of View (H)"));
+			}
+			OptionsMenuBuilder.AddWidget(GenerateLightMenu(), LOCTEXT("LightIntensity", "Render Light Intensity"));
+		}
+		OptionsMenuBuilder.EndSection();
+
+		ExtendOptionsMenu(OptionsMenuBuilder);
+	}
+
+	return OptionsMenuBuilder.MakeWidget();
+}
+
+TSharedRef<SWidget> SChaosClothAssetEditorRestSpaceViewportToolBar::GenerateLightMenu()
+{
+	constexpr float IntensityMin = 0.f;
+	constexpr float IntensityMax = 20.f;
+	
+	return 
+		SNew(SBox)
+		.HAlign(HAlign_Right)
+		[
+			SNew(SBox)
+			.Padding(FMargin(4.0f, 0.0f, 0.0f, 0.0f))
+			.WidthOverride(100.0f)
+			[
+				SNew(SBorder)
+				.BorderImage(FAppStyle::Get().GetBrush("Menu.WidgetBorder"))
+				.Padding(FMargin(1.0f))
+				[
+					SNew(SSpinBox<float>)
+					.Style(&FAppStyle::Get(), "Menu.SpinBox")
+					.Font(FAppStyle::GetFontStyle(TEXT("MenuItem.Font")))
+					.MinValue(IntensityMin)
+					.MaxValue(IntensityMax)
+					.Value(this, &SChaosClothAssetEditorRestSpaceViewportToolBar::GetCameraPointLightIntensity)
+					.OnValueChanged(this, &SChaosClothAssetEditorRestSpaceViewportToolBar::CameraPointLightIntensityChanged)
+					.IsEnabled(this, &SChaosClothAssetEditorRestSpaceViewportToolBar::IsRenderModeEnabled)
+				]
+			]
+		];
+}
+
+
+float SChaosClothAssetEditorRestSpaceViewportToolBar::GetCameraPointLightIntensity() const
+{
+	return RestSpaceViewportClient->GetCameraPointLightIntensity();
+}
+
+void SChaosClothAssetEditorRestSpaceViewportToolBar::CameraPointLightIntensityChanged(float NewValue)
+{
+	RestSpaceViewportClient->SetCameraPointLightIntensity(NewValue);
+}
+
+bool SChaosClothAssetEditorRestSpaceViewportToolBar::IsRenderModeEnabled() const
+{
+	return RestSpaceViewportClient->GetConstructionViewMode() == UE::Chaos::ClothAsset::EClothPatternVertexType::Render;
+}
+
 TSharedRef<SWidget> SChaosClothAssetEditorRestSpaceViewportToolBar::MakeOptionsMenu()
 {
 	return SNew(SEditorViewportToolbarMenu)
 		.ParentToolBar(SharedThis(this))
 		.Cursor(EMouseCursor::Default)
 		.Image("EditorViewportToolBar.OptionsDropdown")
-		.OnGetMenuContent(this, &SChaosClothAssetEditorRestSpaceViewportToolBar::GenerateClothViewportOptionsMenu);
+		.OnGetMenuContent(this, &SChaosClothAssetEditorRestSpaceViewportToolBar::GenerateClothRestSpaceViewportOptionsMenu);
 }
 
 TSharedRef<SWidget> SChaosClothAssetEditorRestSpaceViewportToolBar::MakeDisplayToolBar(const TSharedPtr<FExtender> InExtenders)
diff --git a/Engine/Plugins/Experimental/ChaosClothAssetEditor/Source/ChaosClothAssetEditor/Private/ChaosClothAsset/SClothEditorRestSpaceViewportToolBar.h b/Engine/Plugins/Experimental/ChaosClothAssetEditor/Source/ChaosClothAssetEditor/Private/ChaosClothAsset/SClothEditorRestSpaceViewportToolBar.h
--- a/Engine/Plugins/Experimental/ChaosClothAssetEditor/Source/ChaosClothAssetEditor/Private/ChaosClothAsset/SClothEditorRestSpaceViewportToolBar.h
+++ b/Engine/Plugins/Experimental/ChaosClothAssetEditor/Source/ChaosClothAssetEditor/Private/ChaosClothAsset/SClothEditorRestSpaceViewportToolBar.h
@@ -7,6 +7,7 @@
 class FExtender;
 class FUICommandList;
 class SChaosClothAssetEditorRestSpaceViewport;
+namespace UE::Chaos::ClothAsset { class FChaosClothEditorRestSpaceViewportClient; }
 
 /**
  * Toolbar that shows up at the top of the rest space viewport
@@ -17,11 +18,20 @@ public:
 	SLATE_BEGIN_ARGS(SChaosClothAssetEditorRestSpaceViewportToolBar) {}
 		SLATE_ARGUMENT(TSharedPtr<FUICommandList>, CommandList)
 		SLATE_ARGUMENT(TSharedPtr<FExtender>, Extenders)
+		SLATE_ARGUMENT(TSharedPtr<UE::Chaos::ClothAsset::FChaosClothEditorRestSpaceViewportClient>, RestSpaceViewportClient)
 	SLATE_END_ARGS()
 
 	void Construct(const FArguments& InArgs, TSharedPtr<class SChaosClothAssetEditorRestSpaceViewport> InChaosClothAssetEditorViewport);
 
 private:
+
+	float GetCameraPointLightIntensity() const;
+	void CameraPointLightIntensityChanged(float NewValue);
+	bool IsRenderModeEnabled() const;
+
+	TSharedRef<SWidget> GenerateLightMenu();
+	TSharedRef<SWidget> GenerateClothRestSpaceViewportOptionsMenu();
+
 	TSharedRef<SWidget> MakeOptionsMenu();
 	TSharedRef<SWidget> MakeDisplayToolBar(const TSharedPtr<FExtender> InExtenders);
 	TSharedRef<SWidget> MakeToolBar(const TSharedPtr<FExtender> InExtenders);
@@ -35,5 +45,8 @@ private:
 	/** The viewport that we are in */
 	TWeakPtr<class SChaosClothAssetEditorRestSpaceViewport> ChaosClothAssetEditorRestSpaceViewportPtr;
 
+	/** Client associated with the viewport we are in */
+	TSharedPtr<UE::Chaos::ClothAsset::FChaosClothEditorRestSpaceViewportClient> RestSpaceViewportClient;
+
 	TSharedPtr<FUICommandList> CommandList;
 };
diff --git a/Engine/Plugins/Experimental/ChaosClothAssetEditor/Source/ChaosClothAssetEditor/Private/ChaosClothAsset/SClothEditorViewportToolBarBase.h b/Engine/Plugins/Experimental/ChaosClothAssetEditor/Source/ChaosClothAssetEditor/Private/ChaosClothAsset/SClothEditorViewportToolBarBase.h
--- a/Engine/Plugins/Experimental/ChaosClothAssetEditor/Source/ChaosClothAssetEditor/Private/ChaosClothAsset/SClothEditorViewportToolBarBase.h
+++ b/Engine/Plugins/Experimental/ChaosClothAssetEditor/Source/ChaosClothAssetEditor/Private/ChaosClothAsset/SClothEditorViewportToolBarBase.h
@@ -9,8 +9,6 @@ class CHAOSCLOTHASSETEDITOR_API SChaosClothAssetEditorViewportToolBarBase : publ
 {
 protected:
 	TSharedRef<SWidget> GenerateClothViewportOptionsMenu() const;
-
-private:
 	TSharedRef<SWidget> GenerateFOVMenu() const;
 	float OnGetFOVValue() const;
 };
