Chaos Cloth Asset: Updates to SelectionToWeightMap Node--Separate input selection and map name inputs, support face selection.

diff --git a/Engine/Plugins/Experimental/ChaosClothAssetEditor/Source/ChaosClothAssetDataflowNodes/Private/ChaosClothAsset/SelectionToWeightMapNode.cpp b/Engine/Plugins/Experimental/ChaosClothAssetEditor/Source/ChaosClothAssetDataflowNodes/Private/ChaosClothAsset/SelectionToWeightMapNode.cpp
--- a/Engine/Plugins/Experimental/ChaosClothAssetEditor/Source/ChaosClothAssetDataflowNodes/Private/ChaosClothAsset/SelectionToWeightMapNode.cpp
+++ b/Engine/Plugins/Experimental/ChaosClothAssetEditor/Source/ChaosClothAssetDataflowNodes/Private/ChaosClothAsset/SelectionToWeightMapNode.cpp
@@ -14,8 +14,8 @@ FChaosClothAssetSelectionToWeightMapNode::FChaosClothAssetSelectionToWeightMapNo
 {
 	RegisterInputConnection(&Collection);
 	RegisterOutputConnection(&Collection, &Collection);
-	RegisterInputConnection(&Name);
-	RegisterOutputConnection(&Name, &Name);
+	RegisterInputConnection(&SelectionName);
+	RegisterOutputConnection(&WeightMapName);
 }
 
 
@@ -30,21 +30,21 @@ void FChaosClothAssetSelectionToWeightMapNode::Evaluate(Dataflow::FContext& Cont
 		FCollectionClothFacade ClothFacade(ClothCollection);
 
 		const FCollectionClothSelectionConstFacade SelectionFacade(ClothCollection);
-		const FString InName = GetValue<FString>(Context, &Name);
+		const FString InSelectionName = GetValue<FString>(Context, &SelectionName);
 
 		if (SelectionFacade.IsValid())
 		{
-			const int32 FoundSelectionIndex = SelectionFacade.FindSelection(InName);
+			const int32 FoundSelectionIndex = SelectionFacade.FindSelection(InSelectionName);
 
 			if (FoundSelectionIndex != INDEX_NONE)
 			{
 				const FString& SelectionType = SelectionFacade.GetType()[FoundSelectionIndex];
 
-				if (SelectionType == "SimVertex3D" || SelectionType == "SimVertex2D")
+				if (SelectionType == "SimVertex3D" || SelectionType == "SimVertex2D" || SelectionType == "SimFace")
 				{
-					const FName MapName(InName);
-					ClothFacade.AddWeightMap(MapName);
-					TArrayView<float> OutClothWeights = ClothFacade.GetWeightMap(MapName);
+					const FName InMapName(WeightMapName);
+					ClothFacade.AddWeightMap(InMapName);
+					TArrayView<float> OutClothWeights = ClothFacade.GetWeightMap(InMapName);
 
 					const TSet<int32>& Selection = SelectionFacade.GetIndices()[FoundSelectionIndex];
 
@@ -52,13 +52,11 @@ void FChaosClothAssetSelectionToWeightMapNode::Evaluate(Dataflow::FContext& Cont
 					{
 						for (int32 VertexIndex = 0; VertexIndex < OutClothWeights.Num(); ++VertexIndex)
 						{
-							OutClothWeights[VertexIndex] = Selection.Contains(VertexIndex) ? 1.0f : 0.0f;
+							OutClothWeights[VertexIndex] = Selection.Contains(VertexIndex) ? SelectedValue : UnselectedValue;
 						}
 					}
-					else
+					else if (SelectionType == "SimVertex2D")
 					{
-						check(SelectionType == "SimVertex2D");
-
 						// We are given a selection over the set of 2D vertices, but weight maps only exist for 3D vertices, so
 						// we need a bit of translation
 
@@ -67,24 +65,39 @@ void FChaosClothAssetSelectionToWeightMapNode::Evaluate(Dataflow::FContext& Cont
 						for (int32 Vertex3DIndex = 0; Vertex3DIndex < OutClothWeights.Num(); ++Vertex3DIndex)
 						{
 							// If any corresponding 2D vertex is selected, set the 3D weight map value to one, otherwise it gets zero
-							OutClothWeights[Vertex3DIndex] = 0.0;
+							OutClothWeights[Vertex3DIndex] = UnselectedValue;
 							for (const int32 Vertex2DIndex : Vertex3DTo2D[Vertex3DIndex])
 							{
 								if (Selection.Contains(Vertex2DIndex))
 								{
-									OutClothWeights[Vertex3DIndex] = 1.0;
+									OutClothWeights[Vertex3DIndex] = SelectedValue;
 									break;
 								}
 							}
 						}
 					}
+					else
+					{
+						check(SelectionType == "SimFace");
+						// Fill with unselected value and then set all vertices in the selected faces to the selected value.
+						for (int32 Vertex3DIndex = 0; Vertex3DIndex < OutClothWeights.Num(); ++Vertex3DIndex)
+						{
+							OutClothWeights[Vertex3DIndex] = UnselectedValue;
+						}
+						const TConstArrayView<FIntVector3> SimIndices3D = ClothFacade.GetSimIndices3D();
+						for (const int32 FaceIndex : Selection)
+						{
+							OutClothWeights[SimIndices3D[FaceIndex][0]] = SelectedValue;
+							OutClothWeights[SimIndices3D[FaceIndex][1]] = SelectedValue;
+							OutClothWeights[SimIndices3D[FaceIndex][2]] = SelectedValue;
+						}
+					}
 				}
 				else
 				{
 					FClothDataflowTools::LogAndToastWarning(*this,
 						LOCTEXT("SelectionTypeNotCorrectHeadline", "Selection type is incompatible."),
-						FText::Format(LOCTEXT("SelectionTypeNotCorrectDetails", "Selection with Name \"{0}\" does not have Type \"SimVertex3D\" or \"SimVertex2D\"."),
-							FText::FromString(InName)));
+						FText::Format(LOCTEXT("SelectionTypeNotCorrectDetails", "Selection with Name \"{0}\" does not have Type \"SimVertex3D\", \"SimVertex2D\", or \"SimFace\"."), FText::FromString(InSelectionName)));
 				}
 			}
 			else
@@ -92,12 +105,12 @@ void FChaosClothAssetSelectionToWeightMapNode::Evaluate(Dataflow::FContext& Cont
 				FClothDataflowTools::LogAndToastWarning(*this,
 					LOCTEXT("SelectionNameNotFoundHeadline", "Selection Name was not found."),
 					FText::Format(LOCTEXT("SelectionNameNotFoundDetails", "A Selection with Name \"{0}\" was not found in the Collection."),
-						FText::FromString(InName)));
+						FText::FromString(InSelectionName)));
 			}
 		}
 
 		SetValue(Context, MoveTemp(*ClothCollection), &Collection);
-		SetValue(Context, InName, &Name);
+		SetValue(Context, WeightMapName, &WeightMapName);
 	}
 }
 
diff --git a/Engine/Plugins/Experimental/ChaosClothAssetEditor/Source/ChaosClothAssetDataflowNodes/Public/ChaosClothAsset/SelectionToWeightMapNode.h b/Engine/Plugins/Experimental/ChaosClothAssetEditor/Source/ChaosClothAssetDataflowNodes/Public/ChaosClothAsset/SelectionToWeightMapNode.h
--- a/Engine/Plugins/Experimental/ChaosClothAssetEditor/Source/ChaosClothAssetDataflowNodes/Public/ChaosClothAsset/SelectionToWeightMapNode.h
+++ b/Engine/Plugins/Experimental/ChaosClothAssetEditor/Source/ChaosClothAssetDataflowNodes/Public/ChaosClothAsset/SelectionToWeightMapNode.h
@@ -6,7 +6,7 @@
 #include "GeometryCollection/ManagedArrayCollection.h"
 #include "SelectionToWeightMapNode.generated.h"
 
-/** Convert an integer index selection to a vertex weight map where the map value is one for vertices in the selection set, and zero otherwise */
+/** Convert an integer index selection to a vertex weight map where different map values can be set for selected and unselected vertices. */
 USTRUCT(Meta = (DataflowCloth))
 struct FChaosClothAssetSelectionToWeightMapNode : public FDataflowNode
 {
@@ -15,12 +15,24 @@ struct FChaosClothAssetSelectionToWeightMapNode : public FDataflowNode
 
 public:
 
-	UPROPERTY(Meta = (Dataflowinput, DataflowOutput, DisplayName = "Collection", DataflowPassthrough = "Collection"))
+	UPROPERTY(Meta = (Dataflowinput, DataflowOutput, DataflowPassthrough = "Collection"))
 	FManagedArrayCollection Collection;
 
-	/** The name of the selection to convert and also the name of the weight map attribute that will be added to the collection */
-	UPROPERTY(Meta = (DataflowInput, DataflowOutput, DataflowPassthrough = "Name"))
-	FString Name;
+	/** The name of the selection to convert. */
+	UPROPERTY(Meta = (DataflowInput))
+	FString SelectionName;
+
+	/** The name of the weight map attribute that will be added to the collection */
+	UPROPERTY(EditAnywhere, Category = "Selection To Weight Map", Meta = (DataflowOutput))
+	FString WeightMapName;
+
+	/** The value unselected vertices receive. */
+	UPROPERTY(EditAnywhere, Category = "Selection To Weight Map", Meta = (ClampMin = "0", ClampMax = "1"))
+	float UnselectedValue = 0.f;
+
+	/** The value selected vertices receive. */
+	UPROPERTY(EditAnywhere, Category = "Selection To Weight Map", Meta = (ClampMin = "0", ClampMax = "1"))
+	float SelectedValue = 1.f;
 
 	FChaosClothAssetSelectionToWeightMapNode(const Dataflow::FNodeParameters& InParam, FGuid InGuid = FGuid::NewGuid());
 
