Cloth Editor: Fix automatic connection of new transfer skin weight node

diff --git a/Engine/Plugins/Experimental/ChaosClothAssetEditor/Source/ChaosClothAssetEditor/Private/ChaosClothAsset/ClothEditorMode.cpp b/Engine/Plugins/Experimental/ChaosClothAssetEditor/Source/ChaosClothAssetEditor/Private/ChaosClothAsset/ClothEditorMode.cpp
--- a/Engine/Plugins/Experimental/ChaosClothAssetEditor/Source/ChaosClothAssetEditor/Private/ChaosClothAsset/ClothEditorMode.cpp
+++ b/Engine/Plugins/Experimental/ChaosClothAssetEditor/Source/ChaosClothAssetEditor/Private/ChaosClothAsset/ClothEditorMode.cpp
@@ -211,10 +211,12 @@ void UChaosClothAssetEditorMode::RegisterAddNodeCommand(TSharedPtr<FUICommandInf
 	auto AddNode = [this](const FName& NewNodeType)
 	{
 		const FName ConnectionType = FManagedArrayCollection::StaticType();
+		const FName ConnectionName("Collection");
+
 		UEdGraphNode* const CurrentlySelectedNode = GetSingleSelectedNodeWithOutputType(ConnectionType);
 		checkf(CurrentlySelectedNode, TEXT("No node with FManagedArrayCollection output is currently selected in the Dataflow graph"));
 
-		const UEdGraphNode* const NewNode = CreateAndConnectNewNode(NewNodeType, *CurrentlySelectedNode, ConnectionType);
+		const UEdGraphNode* const NewNode = CreateAndConnectNewNode(NewNodeType, *CurrentlySelectedNode, ConnectionType, ConnectionName);
 		verifyf(NewNode, TEXT("Failed to create a new node: %s"), *NewNodeType.ToString());
 
 		StartToolForSelectedNode(NewNode);
@@ -1138,7 +1140,7 @@ UEdGraphNode* UChaosClothAssetEditorMode::CreateNewNode(const FName& NewNodeType
 }
 
 
-UEdGraphNode* UChaosClothAssetEditorMode::CreateAndConnectNewNode(const FName& NewNodeTypeName, UEdGraphNode& UpstreamNode, const FName& ConnectionTypeName)
+UEdGraphNode* UChaosClothAssetEditorMode::CreateAndConnectNewNode(const FName& NewNodeTypeName, UEdGraphNode& UpstreamNode, const FName& ConnectionTypeName, const FName& NewNodeConnectionName)
 {
 	// First find the specified output of the upstream node, plus any pins it's connected to
 
@@ -1173,7 +1175,7 @@ UEdGraphNode* UChaosClothAssetEditorMode::CreateAndConnectNewNode(const FName& N
 		UEdGraphPin* NewNodeInputPin = nullptr;
 		for (const FDataflowInput* const NewNodeInput : NewDataflowNode->GetInputs())
 		{
-			if (NewNodeInput->GetType() == ConnectionTypeName)
+			if (NewNodeInput->GetType() == ConnectionTypeName && NewNodeInput->GetName() == NewNodeConnectionName)
 			{
 				NewNodeInputPin = NewDataflowEdNode->FindPin(*NewNodeInput->GetName().ToString(), EGPD_Input);
 			}
@@ -1182,7 +1184,7 @@ UEdGraphNode* UChaosClothAssetEditorMode::CreateAndConnectNewNode(const FName& N
 		UEdGraphPin* NewNodeOutputPin = nullptr;
 		for (const FDataflowOutput* const NewNodeOutput : NewDataflowNode->GetOutputs())
 		{
-			if (NewNodeOutput->GetType() == ConnectionTypeName)
+			if (NewNodeOutput->GetType() == ConnectionTypeName && NewNodeOutput->GetName() == NewNodeConnectionName)
 			{
 				NewNodeOutputPin = NewDataflowEdNode->FindPin(*NewNodeOutput->GetName().ToString(), EGPD_Output);
 				break;
diff --git a/Engine/Plugins/Experimental/ChaosClothAssetEditor/Source/ChaosClothAssetEditor/Private/ChaosClothAsset/ClothEditorMode.h b/Engine/Plugins/Experimental/ChaosClothAssetEditor/Source/ChaosClothAssetEditor/Private/ChaosClothAsset/ClothEditorMode.h
--- a/Engine/Plugins/Experimental/ChaosClothAssetEditor/Source/ChaosClothAssetEditor/Private/ChaosClothAsset/ClothEditorMode.h
+++ b/Engine/Plugins/Experimental/ChaosClothAssetEditor/Source/ChaosClothAssetEditor/Private/ChaosClothAsset/ClothEditorMode.h
@@ -181,9 +181,10 @@ private:
 	* @param NewNodeTypeName The type of node to create, by name
 	* @param UpstreamNode Node to connect the new node to
 	* @param ConnectionTypeName The type of output of the upstream node to connect our new node to
+	* @param NewNodeConnectionName The name of the input/output connection on our new node that will be connected
 	* @return The newly-created node
 	*/
-	UEdGraphNode* CreateAndConnectNewNode(const FName& NewNodeTypeName,	UEdGraphNode& UpstreamNode,	const FName& ConnectionTypeName);
+	UEdGraphNode* CreateAndConnectNewNode(const FName& NewNodeTypeName,	UEdGraphNode& UpstreamNode,	const FName& ConnectionTypeName, const FName& NewNodeConnectionName);
 
 
 	void InitializeContextObject();
