Chaos Cloth - Fixed multi-material cloth asset build.

#rnx

diff --git a/Engine/Plugins/Experimental/ChaosClothAssetEditor/Source/ChaosClothAssetDataflowNodes/Private/ChaosClothAsset/TerminalNode.cpp b/Engine/Plugins/Experimental/ChaosClothAssetEditor/Source/ChaosClothAssetDataflowNodes/Private/ChaosClothAsset/TerminalNode.cpp
--- a/Engine/Plugins/Experimental/ChaosClothAssetEditor/Source/ChaosClothAssetDataflowNodes/Private/ChaosClothAsset/TerminalNode.cpp
+++ b/Engine/Plugins/Experimental/ChaosClothAssetEditor/Source/ChaosClothAssetDataflowNodes/Private/ChaosClothAsset/TerminalNode.cpp
@@ -79,17 +79,13 @@ void FChaosClothAssetTerminalNode::SetAssetValue(TObjectPtr<UObject> Asset, Data
 				{
 					const FString& RenderMaterialPathName = LodRenderMaterialPathName[LodMaterialIndex];
 
-					auto Predicate = [&RenderMaterialPathName](const FSkeletalMaterial& SkeletalMaterial)
+					if (UMaterialInterface* const Material = LoadObject<UMaterialInterface>(ClothAsset, *RenderMaterialPathName, nullptr, LOAD_None, nullptr))
 					{
-						return SkeletalMaterial.MaterialInterface && SkeletalMaterial.MaterialInterface->GetPathName() == RenderMaterialPathName;
-					};
-
-					if (!Materials.FindByPredicate(Predicate))
+						Materials.Emplace(Material, true, false, Material->GetFName());
+					}
+					else
 					{
-						if (UMaterialInterface* const Material = LoadObject<UMaterialInterface>(ClothAsset, *RenderMaterialPathName, nullptr, LOAD_None, nullptr))
-						{
-							Materials.Emplace(Material, true, false, Material->GetFName());
-						}
+						Materials.Emplace();
 					}
 				}
 
diff --git a/Engine/Plugins/Experimental/ChaosClothAssetEditor/Source/ChaosClothAssetTools/Private/ChaosClothAsset/ClothAssetBuilderEditor.cpp b/Engine/Plugins/Experimental/ChaosClothAssetEditor/Source/ChaosClothAssetTools/Private/ChaosClothAsset/ClothAssetBuilderEditor.cpp
--- a/Engine/Plugins/Experimental/ChaosClothAssetEditor/Source/ChaosClothAssetTools/Private/ChaosClothAsset/ClothAssetBuilderEditor.cpp
+++ b/Engine/Plugins/Experimental/ChaosClothAssetEditor/Source/ChaosClothAssetTools/Private/ChaosClothAsset/ClothAssetBuilderEditor.cpp
@@ -40,24 +40,6 @@ void UClothAssetBuilderEditor::BuildLod(FSkeletalMeshLODModel& LODModel, const U
 
 	const FCollectionClothConstFacade ClothFacade(ClothCollection);
 
-	// Cloth Render Patterns == Skeletal Mesh Sections
-	const int32 NumSections = ClothFacade.GetNumRenderPatterns();
-	const int32 NumLodMaterials = NumSections;
-	const TConstArrayView<FString> LodRenderMaterialPathName = ClothFacade.GetRenderMaterialPathName();
-
-	TArray<int32> LodMaterialMap;
-	LodMaterialMap.AddUninitialized(NumLodMaterials);
-
-	for (int32 LodMaterialIndex = 0; LodMaterialIndex < NumLodMaterials; ++LodMaterialIndex)
-	{
-		const FString& RenderMaterialPathName = LodRenderMaterialPathName[LodMaterialIndex];
-		const int32 SkeletalMaterialIndex = Materials.IndexOfByPredicate([&RenderMaterialPathName](const FSkeletalMaterial& SkeletalMaterial)
-			{
-				return SkeletalMaterial.MaterialInterface && SkeletalMaterial.MaterialInterface->GetPathName() == RenderMaterialPathName;
-			});
-		LodMaterialMap[LodMaterialIndex] = SkeletalMaterialIndex;
-	}
-
 	// Keep track of the active bone indices for this LOD model
 	TSet<FBoneIndexType> ActiveBoneIndices;
 	ActiveBoneIndices.Reserve(ClothAsset.RefSkeleton.GetNum());
@@ -91,6 +73,7 @@ void UClothAssetBuilderEditor::BuildLod(FSkeletalMeshLODModel& LODModel, const U
 	LODModel.MeshToImportVertexMap.Reserve(NumRenderVertices);
 
 	// Populate this LOD's sections and the LOD index buffer
+	const int32 NumSections = ClothFacade.GetNumRenderPatterns();  // Cloth Render Patterns == Skeletal Mesh Sections
 	LODModel.Sections.SetNum(NumSections);
 
 	int32 BaseIndex = 0;
