Chaos Cloth - Add name check on AddWeightMap node (missing file from CL 26767237).

#rnx

diff --git a/Engine/Plugins/Experimental/ChaosClothAssetEditor/Source/ChaosClothAssetDataflowNodes/Private/ChaosClothAsset/AddWeightMapNode.cpp b/Engine/Plugins/Experimental/ChaosClothAssetEditor/Source/ChaosClothAssetDataflowNodes/Private/ChaosClothAsset/AddWeightMapNode.cpp
--- a/Engine/Plugins/Experimental/ChaosClothAssetEditor/Source/ChaosClothAssetDataflowNodes/Private/ChaosClothAsset/AddWeightMapNode.cpp
+++ b/Engine/Plugins/Experimental/ChaosClothAssetEditor/Source/ChaosClothAssetDataflowNodes/Private/ChaosClothAsset/AddWeightMapNode.cpp
@@ -4,6 +4,7 @@
 #include "ChaosClothAsset/CollectionClothFacade.h"
 #include "ChaosClothAsset/ClothDataflowTools.h"
 #include "ChaosClothAsset/ClothGeometryTools.h"
+#include "ChaosClothAsset/WeightedValue.h"
 #include "Dataflow/DataflowInputOutput.h"
 
 #include UE_INLINE_GENERATED_CPP_BY_NAME(AddWeightMapNode)
@@ -24,6 +25,9 @@ void FChaosClothAssetAddWeightMapNode::Evaluate(Dataflow::FContext& Context, con
 
 	if (Out->IsA<FManagedArrayCollection>(&Collection))
 	{
+		// Make the name a valid attribute name, and replace the value in the UI
+		FWeightMapTools::MakeWeightMapName(Name);
+
 		// Evaluate in collection
 		FManagedArrayCollection InCollection = GetValue<FManagedArrayCollection>(Context, &Collection);
 		const TSharedRef<FManagedArrayCollection> ClothCollection = MakeShared<FManagedArrayCollection>(MoveTemp(InCollection));
