Chaos Cloth:
- Added UV transform node.
- Added virtual garment import test asset.

#rnx

diff --git a/Engine/Plugins/Experimental/ChaosClothAssetEditor/Source/ChaosClothAssetDataflowNodes/Private/ChaosClothAsset/ChaosClothAssetDataflowNodesModule.cpp b/Engine/Plugins/Experimental/ChaosClothAssetEditor/Source/ChaosClothAssetDataflowNodes/Private/ChaosClothAsset/ChaosClothAssetDataflowNodesModule.cpp
--- a/Engine/Plugins/Experimental/ChaosClothAssetEditor/Source/ChaosClothAssetDataflowNodes/Private/ChaosClothAsset/ChaosClothAssetDataflowNodesModule.cpp
+++ b/Engine/Plugins/Experimental/ChaosClothAssetEditor/Source/ChaosClothAssetDataflowNodes/Private/ChaosClothAsset/ChaosClothAssetDataflowNodesModule.cpp
@@ -41,6 +41,7 @@
 #include "ChaosClothAsset/StaticMeshImportNode.h"
 #include "ChaosClothAsset/TerminalNode.h"
 #include "ChaosClothAsset/TransferSkinWeightsNode.h"
+#include "ChaosClothAsset/TransformUVsNode.h"
 #include "ChaosClothAsset/WeightedValueCustomization.h"
 #include "Dataflow/DataflowNodeColorsRegistry.h"
 #include "Dataflow/DataflowNodeFactory.h"
@@ -91,6 +92,7 @@ namespace UE::Chaos::ClothAsset::Private
 		DATAFLOW_NODE_REGISTER_CREATION_FACTORY(FChaosClothAssetStaticMeshImportNode);
 		DATAFLOW_NODE_REGISTER_CREATION_FACTORY(FChaosClothAssetTerminalNode);
 		DATAFLOW_NODE_REGISTER_CREATION_FACTORY(FChaosClothAssetTransferSkinWeightsNode);
+		DATAFLOW_NODE_REGISTER_CREATION_FACTORY(FChaosClothAssetTransformUVsNode);
 
 		DATAFLOW_NODE_REGISTER_CREATION_FACTORY_NODE_COLORS_BY_CATEGORY("Cloth", FColorScheme::NodeHeader, FColorScheme::NodeBody);
 	}
diff --git a/Engine/Plugins/Experimental/ChaosClothAssetEditor/Source/ChaosClothAssetDataflowNodes/Private/ChaosClothAsset/TransformUVsNode.cpp b/Engine/Plugins/Experimental/ChaosClothAssetEditor/Source/ChaosClothAssetDataflowNodes/Private/ChaosClothAsset/TransformUVsNode.cpp
new file mode 100644
--- /dev/null
+++ b/Engine/Plugins/Experimental/ChaosClothAssetEditor/Source/ChaosClothAssetDataflowNodes/Private/ChaosClothAsset/TransformUVsNode.cpp
@@ -0,0 +1,75 @@
+// Copyright Epic Games, Inc. All Rights Reserved.
+
+#include "ChaosClothAsset/TransformUVsNode.h"
+#include "ChaosClothAsset/CollectionClothFacade.h"
+#include "Dataflow/DataflowInputOutput.h"
+#include "Materials/Material.h"
+
+#include UE_INLINE_GENERATED_CPP_BY_NAME(TransformUVsNode)
+
+#define LOCTEXT_NAMESPACE "ChaosClothAssetTransformUVsNode"
+
+FChaosClothAssetTransformUVsNode::FChaosClothAssetTransformUVsNode(const Dataflow::FNodeParameters& InParam, FGuid InGuid)
+	: FDataflowNode(InParam, InGuid)
+{
+	RegisterInputConnection(&Collection);
+	RegisterOutputConnection(&Collection, &Collection);
+}
+
+void FChaosClothAssetTransformUVsNode::Evaluate(Dataflow::FContext& Context, const FDataflowOutput* Out) const
+{
+	if (Out->IsA<FManagedArrayCollection>(&Collection))
+	{
+		using namespace UE::Chaos::ClothAsset;
+
+		// Evaluate in collection
+		FManagedArrayCollection InCollection = GetValue<FManagedArrayCollection>(Context, &Collection);
+		const TSharedRef<FManagedArrayCollection> ClothCollection = MakeShared<FManagedArrayCollection>(MoveTemp(InCollection));
+
+		// Always check for a valid cloth collection/facade to avoid processing non cloth collections
+		FCollectionClothFacade ClothFacade(ClothCollection);
+		if (ClothFacade.IsValid())
+		{
+			auto Transform = [this](TArrayView<TArray<FVector2f>>&& RenderUVsView)
+				{
+					for (TArray<FVector2f>& RenderUVs : RenderUVsView)
+					{
+						for (int32 Index = 0; Index < RenderUVs.Num(); ++Index)
+						{
+							if (UVChannel == INDEX_NONE || UVChannel == Index)
+							{
+								// Scale
+								RenderUVs[Index] = RenderUVs[Index] * Scale;
+								// Rotate
+								float SinAngle, CosAngle;
+								FMath::SinCos(&SinAngle, &CosAngle, FMath::DegreesToRadians(Rotation));
+								RenderUVs[Index] = FVector2f(
+									RenderUVs[Index].X * CosAngle - RenderUVs[Index].Y * SinAngle,
+									RenderUVs[Index].Y * CosAngle + RenderUVs[Index].X * SinAngle);  // Clockwise rotation
+								// Translate
+								RenderUVs[Index] += FVector2f(Translation.X, -Translation.Y);  // UVs are imported as 1-Y
+							}
+						}
+					}
+				};
+
+			if (Pattern == INDEX_NONE)
+			{
+				Transform(ClothFacade.GetRenderUVs());
+			}
+			else
+			{
+				const int32 NumPatterns = ClothFacade.GetNumRenderPatterns();
+				if (Pattern >= 0 && Pattern < NumPatterns)
+				{
+					FCollectionClothRenderPatternFacade ClothPatternFacade = ClothFacade.GetRenderPattern(Pattern);
+					Transform(ClothPatternFacade.GetRenderUVs());
+				}
+			}
+		}
+
+		SetValue(Context, MoveTemp(*ClothCollection), &Collection);
+	}
+}
+
+#undef LOCTEXT_NAMESPACE
diff --git a/Engine/Plugins/Experimental/ChaosClothAssetEditor/Source/ChaosClothAssetDataflowNodes/Public/ChaosClothAsset/TransformUVsNode.h b/Engine/Plugins/Experimental/ChaosClothAssetEditor/Source/ChaosClothAssetDataflowNodes/Public/ChaosClothAsset/TransformUVsNode.h
new file mode 100644
--- /dev/null
+++ b/Engine/Plugins/Experimental/ChaosClothAssetEditor/Source/ChaosClothAssetDataflowNodes/Public/ChaosClothAsset/TransformUVsNode.h
@@ -0,0 +1,43 @@
+// Copyright Epic Games, Inc. All Rights Reserved.
+
+#pragma once 
+
+#include "Dataflow/DataflowNode.h"
+#include "GeometryCollection/ManagedArrayCollection.h"
+#include "TransformUVsNode.generated.h"
+
+USTRUCT(Meta = (DataflowCloth))
+struct FChaosClothAssetTransformUVsNode : public FDataflowNode
+{
+	GENERATED_USTRUCT_BODY()
+	DATAFLOW_NODE_DEFINE_INTERNAL(FChaosClothAssetTransformUVsNode, "TransformUVs", "Cloth", "Cloth Simulation Transform UV")
+
+public:
+
+	UPROPERTY(Meta = (Dataflowinput, DataflowOutput, DataflowPassthrough = "Collection"))
+	FManagedArrayCollection Collection;
+
+	/** Transform scale. */
+	UPROPERTY(EditAnywhere, Category = "Transform UVs", Meta = (AllowPreserveRatio))
+	FVector2f Scale = { 1.f, 1.f };
+
+	/** Transform rotation angle in degrees. */
+	UPROPERTY(EditAnywhere, Category = "Transform UVs", Meta = (UIMin = -360, UIMax = 360, ClampMin = -360, ClampMax = 360))
+	float Rotation = 0.f;
+
+	/** Transform translation. */
+	UPROPERTY(EditAnywhere, Category = "Transform UVs")
+	FVector2f Translation = { 0.f, 0.f };
+
+	/** Pattern to transform. All patterns will be used when set to -1. */
+	UPROPERTY(EditAnywhere, Category = "Transform UVs", Meta = (UIMax = 10, ClampMin = -1))
+	int32 Pattern = INDEX_NONE;
+
+	/** UV channel to transform. All UV channels will be used when set to -1. */
+	UPROPERTY(EditAnywhere, Category = "Transform UVs", Meta = (UIMax = 5, ClampMin = -1))
+	int32 UVChannel = INDEX_NONE;
+
+	FChaosClothAssetTransformUVsNode(const Dataflow::FNodeParameters& InParam, FGuid InGuid = FGuid::NewGuid());
+
+	virtual void Evaluate(Dataflow::FContext& Context, const FDataflowOutput* Out) const override;
+};
