Cloth Asset: Fixed an issue where the SelectionToWeightMap node doesn't correctly handle not having a valid input collection.

#rnx

diff --git a/Engine/Plugins/Experimental/ChaosClothAssetEditor/Source/ChaosClothAssetDataflowNodes/Private/ChaosClothAsset/SelectionToWeightMapNode.cpp b/Engine/Plugins/Experimental/ChaosClothAssetEditor/Source/ChaosClothAssetDataflowNodes/Private/ChaosClothAsset/SelectionToWeightMapNode.cpp
--- a/Engine/Plugins/Experimental/ChaosClothAssetEditor/Source/ChaosClothAssetDataflowNodes/Private/ChaosClothAsset/SelectionToWeightMapNode.cpp
+++ b/Engine/Plugins/Experimental/ChaosClothAssetEditor/Source/ChaosClothAssetDataflowNodes/Private/ChaosClothAsset/SelectionToWeightMapNode.cpp
@@ -32,7 +32,7 @@ void FChaosClothAssetSelectionToWeightMapNode::Evaluate(Dataflow::FContext& Cont
 		const FCollectionClothSelectionConstFacade SelectionFacade(ClothCollection);
 		const FString InSelectionName = GetValue<FString>(Context, &SelectionName);
 
-		if (SelectionFacade.IsValid())
+		if (SelectionFacade.IsValid() && ClothFacade.IsValid())
 		{
 			const int32 FoundSelectionIndex = SelectionFacade.FindSelection(InSelectionName);
 
