Fixed a bug where the cloth asset was not bound to root bone upon import.

#rnx

diff --git a/Engine/Plugins/Experimental/ChaosClothAssetEditor/Source/ChaosClothAssetDataflowNodes/Private/ChaosClothAsset/SkeletalMeshImportNode.cpp b/Engine/Plugins/Experimental/ChaosClothAssetEditor/Source/ChaosClothAssetDataflowNodes/Private/ChaosClothAsset/SkeletalMeshImportNode.cpp
--- a/Engine/Plugins/Experimental/ChaosClothAssetEditor/Source/ChaosClothAssetDataflowNodes/Private/ChaosClothAsset/SkeletalMeshImportNode.cpp
+++ b/Engine/Plugins/Experimental/ChaosClothAssetEditor/Source/ChaosClothAssetDataflowNodes/Private/ChaosClothAsset/SkeletalMeshImportNode.cpp
@@ -2,6 +2,7 @@
 
 #include "ChaosClothAsset/SkeletalMeshImportNode.h"
 #include "ChaosClothAsset/ClothDataflowTools.h"
+#include "ChaosClothAsset/ClothGeometryTools.h"
 #include "ChaosClothAsset/CollectionClothFacade.h"
 #include "Animation/Skeleton.h"
 #include "Dataflow/DataflowInputOutput.h"
@@ -81,10 +82,19 @@ void FChaosClothAssetSkeletalMeshImportNode::Evaluate(Dataflow::FContext& Contex
 			{
 				ClothFacade.SetPhysicsAssetPathName(PhysicsAsset->GetPathName());
 			}
+
 			if (const USkeleton* Skeleton = SkeletalMesh->GetSkeleton())
 			{
 				ClothFacade.SetSkeletonAssetPathName(Skeleton->GetPathName());
 			}
+			else
+			{
+				static const TCHAR* const DefaultSkeletonPathName = TEXT("/Engine/EditorMeshes/SkeletalMesh/DefaultSkeletalMesh_Skeleton.DefaultSkeletalMesh_Skeleton");
+				ClothFacade.SetSkeletonAssetPathName(DefaultSkeletonPathName);
+			}
+
+			// Must be bound to root bone by default
+			FClothGeometryTools::BindMeshToRootBone(ClothCollection, bImportSimMesh, bImportRenderMesh);
 		}
 		SetValue(Context, MoveTemp(*ClothCollection), &Collection);
 	}
diff --git a/Engine/Plugins/Experimental/ChaosClothAssetEditor/Source/ChaosClothAssetDataflowNodes/Private/ChaosClothAsset/StaticMeshImportNode.cpp b/Engine/Plugins/Experimental/ChaosClothAssetEditor/Source/ChaosClothAssetDataflowNodes/Private/ChaosClothAsset/StaticMeshImportNode.cpp
--- a/Engine/Plugins/Experimental/ChaosClothAssetEditor/Source/ChaosClothAssetDataflowNodes/Private/ChaosClothAsset/StaticMeshImportNode.cpp
+++ b/Engine/Plugins/Experimental/ChaosClothAssetEditor/Source/ChaosClothAssetDataflowNodes/Private/ChaosClothAsset/StaticMeshImportNode.cpp
@@ -170,6 +170,9 @@ void FChaosClothAssetStaticMeshImportNode::Evaluate(Dataflow::FContext& Context,
 
 				// Set a default skeleton
 				ClothFacade.SetSkeletonAssetPathName(DefaultSkeletonPathName);
+				
+				// Bind to root bone by default
+				FClothGeometryTools::BindMeshToRootBone(ClothCollection, bImportAsSimMesh, bImportAsRenderMesh);
 			}
 			else
 			{
