Chaos Cloth Asset:
- Fixed the StaticMeshImport node property names to match the same syntax as the SkeletalMeshImport node and prevent future interoperability issues.
- Update test assets with the new names.

#rnx

diff --git a/Engine/Plugins/Experimental/ChaosClothAssetEditor/Source/ChaosClothAssetDataflowNodes/Private/ChaosClothAsset/StaticMeshImportNode.cpp b/Engine/Plugins/Experimental/ChaosClothAssetEditor/Source/ChaosClothAssetDataflowNodes/Private/ChaosClothAsset/StaticMeshImportNode.cpp
--- a/Engine/Plugins/Experimental/ChaosClothAssetEditor/Source/ChaosClothAssetDataflowNodes/Private/ChaosClothAsset/StaticMeshImportNode.cpp
+++ b/Engine/Plugins/Experimental/ChaosClothAssetEditor/Source/ChaosClothAssetDataflowNodes/Private/ChaosClothAsset/StaticMeshImportNode.cpp
@@ -133,15 +133,15 @@ void FChaosClothAssetStaticMeshImportNode::Evaluate(Dataflow::FContext& Context,
 		FCollectionClothFacade ClothFacade(ClothCollection);
 		ClothFacade.DefineSchema();
 
-		if (StaticMesh && (bImportAsSimMesh || bImportAsRenderMesh))
+		if (StaticMesh && (bImportSimMesh || bImportRenderMesh))
 		{
 			const int32 NumLods = StaticMesh->GetNumSourceModels();
-			if (LodIndex < NumLods)
+			if (LODIndex < NumLods)
 			{
-				const FMeshDescription* const MeshDescription = StaticMesh->GetMeshDescription(LodIndex);
+				const FMeshDescription* const MeshDescription = StaticMesh->GetMeshDescription(LODIndex);
 				check(MeshDescription);
 
-				if (bImportAsSimMesh)
+				if (bImportSimMesh)
 				{
 					FMeshDescriptionToDynamicMesh Converter;
 					Converter.bPrintDebugMessages = false;
@@ -153,10 +153,10 @@ void FChaosClothAssetStaticMeshImportNode::Evaluate(Dataflow::FContext& Context,
 					constexpr bool bAppend = false;
 					FClothGeometryTools::BuildSimMeshFromDynamicMesh(ClothCollection, DynamicMesh, UVChannel, UVScale, bAppend);
 				}
-				if (bImportAsRenderMesh)
+				if (bImportRenderMesh)
 				{
 					// Add render data into a single pattern for now
-					if (!InitializeDataFromMeshDescription(MeshDescription, StaticMesh->GetSourceModel(LodIndex).BuildSettings, StaticMesh->GetStaticMaterials(), ClothCollection))
+					if (!InitializeDataFromMeshDescription(MeshDescription, StaticMesh->GetSourceModel(LODIndex).BuildSettings, StaticMesh->GetStaticMaterials(), ClothCollection))
 					{
 						FClothDataflowTools::LogAndToastWarning(*this,
 							LOCTEXT("InvalidRenderMeshHeadline", "Invalid render mesh."),
@@ -167,17 +167,17 @@ void FChaosClothAssetStaticMeshImportNode::Evaluate(Dataflow::FContext& Context,
 				}
 				
 				// Bind to root bone by default
-				FClothGeometryTools::BindMeshToRootBone(ClothCollection, bImportAsSimMesh, bImportAsRenderMesh);
+				FClothGeometryTools::BindMeshToRootBone(ClothCollection, bImportSimMesh, bImportRenderMesh);
 			}
 			else
 			{
 				FClothDataflowTools::LogAndToastWarning(*this,
-					LOCTEXT("InvalidLodIndexHeadline", "Invalid LOD index."),
-					FText::Format(LOCTEXT("InvalidLodIndexDetails",
+					LOCTEXT("InvalidLODIndexHeadline", "Invalid LOD index."),
+					FText::Format(LOCTEXT("InvalidLODIndexDetails",
 						"{0} is not a valid LOD index for static mesh {1}.\n"
 						"This static mesh has {2} LOD(s)."),
 						FText::FromString(StaticMesh->GetName()),
-						LodIndex,
+						LODIndex,
 						NumLods));
 			}
 		}
diff --git a/Engine/Plugins/Experimental/ChaosClothAssetEditor/Source/ChaosClothAssetDataflowNodes/Public/ChaosClothAsset/SkeletalMeshImportNode.h b/Engine/Plugins/Experimental/ChaosClothAssetEditor/Source/ChaosClothAssetDataflowNodes/Public/ChaosClothAsset/SkeletalMeshImportNode.h
--- a/Engine/Plugins/Experimental/ChaosClothAssetEditor/Source/ChaosClothAssetDataflowNodes/Public/ChaosClothAsset/SkeletalMeshImportNode.h
+++ b/Engine/Plugins/Experimental/ChaosClothAssetEditor/Source/ChaosClothAssetDataflowNodes/Public/ChaosClothAsset/SkeletalMeshImportNode.h
@@ -24,7 +24,7 @@ public:
 	TObjectPtr<const USkeletalMesh> SkeletalMesh;
 
 	/** The skeletal mesh LOD to import. */
-	UPROPERTY(VisibleAnywhere, Category = "Skeletal Mesh Import", Meta = (ClampMin = "0"))
+	UPROPERTY(VisibleAnywhere, Category = "Skeletal Mesh Import", Meta = (ClampMin = "0", DisplayName = "LOD Index"))
 	int32 LODIndex = 0;
 
 	/** The skeletal mesh LOD section to import. */
diff --git a/Engine/Plugins/Experimental/ChaosClothAssetEditor/Source/ChaosClothAssetDataflowNodes/Public/ChaosClothAsset/StaticMeshImportNode.h b/Engine/Plugins/Experimental/ChaosClothAssetEditor/Source/ChaosClothAssetDataflowNodes/Public/ChaosClothAsset/StaticMeshImportNode.h
--- a/Engine/Plugins/Experimental/ChaosClothAssetEditor/Source/ChaosClothAssetDataflowNodes/Public/ChaosClothAsset/StaticMeshImportNode.h
+++ b/Engine/Plugins/Experimental/ChaosClothAssetEditor/Source/ChaosClothAssetDataflowNodes/Public/ChaosClothAsset/StaticMeshImportNode.h
@@ -24,23 +24,23 @@ public:
 	TObjectPtr<const UStaticMesh> StaticMesh;
 
 	/* Which static mesh Lod to import.*/
-	UPROPERTY(EditAnywhere, Category = "Static Mesh Import", Meta = (ClampMin = "0"))
-	int32 LodIndex = 0;
+	UPROPERTY(EditAnywhere, Category = "Static Mesh Import", Meta = (ClampMin = "0", DisplayName = "LOD Index"))
+	int32 LODIndex = 0;
 
 	/* Import static mesh data as a simulation mesh data.*/
 	UPROPERTY(EditAnywhere, Category = "Static Mesh Import")
-	bool bImportAsSimMesh = true;
+	bool bImportSimMesh = true;
 
 	/* Import static mesh data as render mesh data.*/
 	UPROPERTY(EditAnywhere, Category = "Static Mesh Import")
-	bool bImportAsRenderMesh = true;
+	bool bImportRenderMesh = true;
 
 	/* UV Channel used to populate Sim Mesh positions */
-	UPROPERTY(EditAnywhere, Category = "Static Mesh Import", meta = (ClampMin = "-1", EditCondition = "bImportAsSimMesh"))
+	UPROPERTY(EditAnywhere, Category = "Static Mesh Import", meta = (ClampMin = "-1", EditCondition = "bImportSimMesh"))
 	int32 UVChannel = 0;
 
 	/* Apply this scale to the UVs when populating Sim Mesh positions. */
-	UPROPERTY(EditAnywhere, Category = "Static Mesh Import", Meta = (AllowPreserveRatio, EditCondition = "bImportAsSimMesh && (UVChannel >= 0)"))
+	UPROPERTY(EditAnywhere, Category = "Static Mesh Import", Meta = (AllowPreserveRatio, EditCondition = "bImportSimMesh && (UVChannel >= 0)"))
 	FVector2f UVScale = { 1.f, 1.f };
 
 	FChaosClothAssetStaticMeshImportNode(const Dataflow::FNodeParameters& InParam, FGuid InGuid = FGuid::NewGuid());
