Fixed a bug where we were missing layered mesh support setting in the UI so it can be turned off if needed.

#rnx

diff --git a/Engine/Plugins/Experimental/ChaosClothAssetEditor/Source/ChaosClothAssetDataflowNodes/Private/ChaosClothAsset/TransferSkinWeightsNode.cpp b/Engine/Plugins/Experimental/ChaosClothAssetEditor/Source/ChaosClothAssetDataflowNodes/Private/ChaosClothAsset/TransferSkinWeightsNode.cpp
--- a/Engine/Plugins/Experimental/ChaosClothAssetEditor/Source/ChaosClothAssetDataflowNodes/Private/ChaosClothAsset/TransferSkinWeightsNode.cpp
+++ b/Engine/Plugins/Experimental/ChaosClothAssetEditor/Source/ChaosClothAssetDataflowNodes/Private/ChaosClothAsset/TransferSkinWeightsNode.cpp
@@ -253,6 +253,7 @@ namespace UE::Chaos::ClothAsset::Private
 		const FReferenceSkeleton& TargetRefSkeleton,
 		const double NormalThreshold,
 		const double RadiusPercentage,
+		const bool LayeredMeshSupport,
 		const int32 NumSmoothingIterations,
 		const double SmoothingStrength,
 		bool bUseParallel,
@@ -279,7 +280,7 @@ namespace UE::Chaos::ClothAsset::Private
 		TransferBoneWeights.SearchRadius = RadiusPercentage * WeldedSimMesh.GetBounds().DiagonalLength();
 		TransferBoneWeights.NumSmoothingIterations = NumSmoothingIterations;
 		TransferBoneWeights.SmoothingStrength = SmoothingStrength;
-		TransferBoneWeights.LayeredMeshSupport = true; // multilayerd clothing
+		TransferBoneWeights.LayeredMeshSupport = LayeredMeshSupport; // multilayerd clothing
 		TransferBoneWeights.ForceInpaintWeightMapName = InpaintMaskWeightMapName;
 
 		if (!ensure(TransferBoneWeights.Validate() == EOperationValidationResult::Ok))
@@ -304,6 +305,7 @@ namespace UE::Chaos::ClothAsset::Private
 		const FReferenceSkeleton& TargetRefSkeleton,
 		const double NormalThreshold,
 		const double RadiusPercentage,
+		const bool LayeredMeshSupport,
 		const int32 NumSmoothingIterations,
 		const double SmoothingStrength,
 		bool bUseParallel,
@@ -318,7 +320,7 @@ namespace UE::Chaos::ClothAsset::Private
 		// Convert cloth sim mesh LOD to the welded dynamic sim mesh and transfer weights.
 		//
 		FDynamicMesh3 WeldedSimMesh;
-		if (!TransferInpaintWeightsToSim(TargetRefSkeleton, NormalThreshold, RadiusPercentage, NumSmoothingIterations, SmoothingStrength, bUseParallel, ClothCollection, InpaintMaskWeightMapName, TransferBoneWeights, WeldedSimMesh))
+		if (!TransferInpaintWeightsToSim(TargetRefSkeleton, NormalThreshold, RadiusPercentage, LayeredMeshSupport, NumSmoothingIterations, SmoothingStrength, bUseParallel, ClothCollection, InpaintMaskWeightMapName, TransferBoneWeights, WeldedSimMesh))
 		{
 			return false;
 		}
@@ -335,7 +337,7 @@ namespace UE::Chaos::ClothAsset::Private
 		if (FCollectionClothFacade(UserTransferCollection).IsValid())
 		{
 			FDynamicMesh3 TransferWeightSimMesh;
-			if (TransferInpaintWeightsToSim(TargetRefSkeleton, NormalThreshold, RadiusPercentage, NumSmoothingIterations, SmoothingStrength, bUseParallel, UserTransferCollection, InpaintMaskWeightMapName, TransferBoneWeights, TransferWeightSimMesh))
+			if (TransferInpaintWeightsToSim(TargetRefSkeleton, NormalThreshold, RadiusPercentage, LayeredMeshSupport, NumSmoothingIterations, SmoothingStrength, bUseParallel, UserTransferCollection, InpaintMaskWeightMapName, TransferBoneWeights, TransferWeightSimMesh))
 			{
 				WeldedSimMesh = MoveTemp(TransferWeightSimMesh); 
 			}
@@ -492,7 +494,7 @@ void FChaosClothAssetTransferSkinWeightsNode::Evaluate(Dataflow::FContext& Conte
 				FManagedArrayCollection TransferWeightsManagedCollection = GetValue<FManagedArrayCollection>(Context, &TransferWeightsCollection);
 				const TSharedRef<FManagedArrayCollection> ClothTransferWeightCollection = MakeShared<FManagedArrayCollection>(MoveTemp(TransferWeightsManagedCollection));
 				const FString& InpainMaskWeightMapName = GetValue<FString>(Context, &InpaintMask.WeightMap);
-				bTransferResult = TransferInpaintWeights(TargetRefSkeleton, NormalThreshold, RadiusPercentage, NumSmoothingIterations, SmoothingStrength, bUseParallel, ClothCollection, ClothTransferWeightCollection, FName(InpainMaskWeightMapName), TransferBoneWeights);
+				bTransferResult = TransferInpaintWeights(TargetRefSkeleton, NormalThreshold, RadiusPercentage, LayeredMeshSupport, NumSmoothingIterations, SmoothingStrength, bUseParallel, ClothCollection, ClothTransferWeightCollection, FName(InpainMaskWeightMapName), TransferBoneWeights);
 			}
 			else if (TransferMethod == EChaosClothAssetTransferSkinWeightsMethod::ClosestPointOnSurface)
 			{
diff --git a/Engine/Plugins/Experimental/ChaosClothAssetEditor/Source/ChaosClothAssetDataflowNodes/Public/ChaosClothAsset/TransferSkinWeightsNode.h b/Engine/Plugins/Experimental/ChaosClothAssetEditor/Source/ChaosClothAssetDataflowNodes/Public/ChaosClothAsset/TransferSkinWeightsNode.h
--- a/Engine/Plugins/Experimental/ChaosClothAssetEditor/Source/ChaosClothAssetDataflowNodes/Public/ChaosClothAsset/TransferSkinWeightsNode.h
+++ b/Engine/Plugins/Experimental/ChaosClothAssetEditor/Source/ChaosClothAssetDataflowNodes/Public/ChaosClothAsset/TransferSkinWeightsNode.h
@@ -66,6 +66,14 @@ public:
 	UPROPERTY(EditAnywhere, Category = "Transfer Skin Weights|Transfer Method", Meta = (UIMin = -1, UIMax = 180, ClampMin = -1, ClampMax = 180, EditCondition="TransferMethod==EChaosClothAssetTransferSkinWeightsMethod::InpaintWeights"))
 	double NormalThreshold = 30;
 
+	/** 
+	 * If true, when the closest point doesn't pass the normal threshold test, will try again with a flipped normal. 
+	 * This helps with layered meshes where the "inner" and "outer" layers are close to each other but whose normals 
+	 * are pointing in the opposite directions.
+	 */
+	UPROPERTY(EditAnywhere, Category = "Transfer Skin Weights|Transfer Method", Meta = (DisplayName = "Layered Mesh Support", EditCondition="TransferMethod==EChaosClothAssetTransferSkinWeightsMethod::InpaintWeights"))
+	bool LayeredMeshSupport = true;
+
 	/** The number of smoothing iterations applied to the vertices whose weights were automatically computed. */
 	UPROPERTY(EditAnywhere, Category = "Transfer Skin Weights|Transfer Method", Meta = (UIMin = 0, UIMax = 100, ClampMin = 0, ClampMax = 100, DisplayName = "Smoothing Iterations", EditCondition="TransferMethod==EChaosClothAssetTransferSkinWeightsMethod::InpaintWeights"))
 	int32 NumSmoothingIterations = 10;
