Chaos Cloth Asset Editor: log a warning when the preview SKM is incompatible with the cloth asset because it will not update all required bones. This will check all SKM LODs.

It can be difficult to debug what's going wrong when the SKM is incompatible, especially if this is due to bones being excluded at certain LODs. Ideally, we will come up with a solution that will allow us to make these assets compatible, but for now, just warn that they aren't.

I moved some static methods from the SkeletalMeshComponent to the SkinnedMeshComponent so they could be used by the ClothComponent.

diff --git a/Engine/Plugins/Experimental/ChaosClothAssetEditor/Source/ChaosClothAssetEditor/Private/ChaosClothAsset/ClothEditor.cpp b/Engine/Plugins/Experimental/ChaosClothAssetEditor/Source/ChaosClothAssetEditor/Private/ChaosClothAsset/ClothEditor.cpp
--- a/Engine/Plugins/Experimental/ChaosClothAssetEditor/Source/ChaosClothAssetEditor/Private/ChaosClothAsset/ClothEditor.cpp
+++ b/Engine/Plugins/Experimental/ChaosClothAssetEditor/Source/ChaosClothAssetEditor/Private/ChaosClothAsset/ClothEditor.cpp
@@ -5,6 +5,8 @@
 
 #include UE_INLINE_GENERATED_CPP_BY_NAME(ClothEditor)
 
+DEFINE_LOG_CATEGORY(LogChaosClothAssetEditor);
+
 TSharedPtr<FBaseAssetToolkit> UChaosClothAssetEditor::CreateToolkit()
 {
 	using namespace UE::Chaos::ClothAsset;
diff --git a/Engine/Plugins/Experimental/ChaosClothAssetEditor/Source/ChaosClothAssetEditor/Private/ChaosClothAsset/ClothEditor.h b/Engine/Plugins/Experimental/ChaosClothAssetEditor/Source/ChaosClothAssetEditor/Private/ChaosClothAsset/ClothEditor.h
--- a/Engine/Plugins/Experimental/ChaosClothAssetEditor/Source/ChaosClothAssetEditor/Private/ChaosClothAsset/ClothEditor.h
+++ b/Engine/Plugins/Experimental/ChaosClothAssetEditor/Source/ChaosClothAssetEditor/Private/ChaosClothAsset/ClothEditor.h
@@ -27,3 +27,5 @@ public:
 
 	virtual TSharedPtr<FBaseAssetToolkit> CreateToolkit() override;
 };
+
+DECLARE_LOG_CATEGORY_EXTERN(LogChaosClothAssetEditor, Log, All);
\ No newline at end of file
diff --git a/Engine/Plugins/Experimental/ChaosClothAssetEditor/Source/ChaosClothAssetEditor/Private/ChaosClothAsset/ClothEditorPreviewScene.cpp b/Engine/Plugins/Experimental/ChaosClothAssetEditor/Source/ChaosClothAssetEditor/Private/ChaosClothAsset/ClothEditorPreviewScene.cpp
--- a/Engine/Plugins/Experimental/ChaosClothAssetEditor/Source/ChaosClothAssetEditor/Private/ChaosClothAsset/ClothEditorPreviewScene.cpp
+++ b/Engine/Plugins/Experimental/ChaosClothAssetEditor/Source/ChaosClothAssetEditor/Private/ChaosClothAsset/ClothEditorPreviewScene.cpp
@@ -2,6 +2,7 @@
 
 #include "ChaosClothAsset/ClothEditorPreviewScene.h"
 #include "ChaosClothAsset/ClothComponent.h"
+#include "ChaosClothAsset/ClothEditor.h"
 #include "Components/SkeletalMeshComponent.h"
 #include "Animation/SkeletalMeshActor.h"
 #include "AssetEditorModeManager.h"
@@ -11,11 +12,110 @@
 #include "ChaosClothAsset/ClothAsset.h"
 #include "SkinnedAssetCompiler.h"
 #include "Misc/TransactionObjectEvent.h"
+#include "PhysicsEngine/PhysicsAsset.h"
+#include "Rendering/SkeletalMeshRenderData.h"
 #include "Transforms/TransformGizmoDataBinder.h"
+#include "Framework/Notifications/NotificationManager.h"
+#include "Widgets/Notifications/SNotificationList.h"
 
 
 #define LOCTEXT_NAMESPACE "UChaosClothEditorPreviewScene"
 
+namespace UE::Chaos::ClothAsset::Private
+{
+	void ValidateClothComponentAttachmentBones(const UChaosClothComponent& ClothComponent)
+	{
+		const UChaosClothAsset* const ClothAsset = ClothComponent.GetClothAsset();
+		const FSkeletalMeshRenderData* const ClothRenderData = ClothComponent.GetSkeletalMeshRenderData();
+		const USkeletalMeshComponent* const SkeletalMeshComponent = Cast<USkeletalMeshComponent>(ClothComponent.LeaderPoseComponent.Get());
+		if (ClothAsset && ClothRenderData && SkeletalMeshComponent)
+		{
+			bool bAnyMissingBones = false;
+
+			const TArray<int32>& LeaderBoneMap = ClothComponent.GetLeaderBoneMap();
+			const FReferenceSkeleton& RefSkeleton = ClothAsset->GetRefSkeleton();
+			auto CheckRequiredBones = [&LeaderBoneMap](const TArray<FBoneIndexType>& RequiredFollowerBones,
+				const TArray<FBoneIndexType>& AvailableLeaderBones, TArray<FBoneIndexType>& OutMissingBones)
+			{
+				OutMissingBones.Reset();
+				if (RequiredFollowerBones.IsEmpty())
+				{
+					return;
+				}
+				TSet<FBoneIndexType> AvailableLeaderBonesSet(AvailableLeaderBones);
+				for (const FBoneIndexType RequiredBone : RequiredFollowerBones)
+				{
+					if (!LeaderBoneMap.IsValidIndex(RequiredBone) || LeaderBoneMap[RequiredBone] == INDEX_NONE ||
+						!AvailableLeaderBonesSet.Contains(LeaderBoneMap[RequiredBone]))
+					{
+						OutMissingBones.Add(RequiredBone);
+					}
+				}
+			};
+
+			TArray<FBoneIndexType> PhysAssetBones;
+			if (const UPhysicsAsset* ClothPhysAsset = ClothComponent.GetPhysicsAsset())
+			{
+				USkinnedMeshComponent::GetPhysicsRequiredBones(ClothAsset, ClothPhysAsset, PhysAssetBones);
+			}
+
+			TArray<FBoneIndexType> MissingBones;
+			for (int32 LODIndex = SkeletalMeshComponent->ComputeMinLOD(); LODIndex < SkeletalMeshComponent->GetNumLODs(); ++LODIndex)
+			{
+				// Check all Leader SKM LODs since the Cloth should try to follow the SKM LOD.
+				// These should be the bones that will actually be calculated by the SKM (see USkeletalMeshComponent::RecalcRequiredBones)
+				TArray<FBoneIndexType> SKMRequiredBones;
+				TArray<FBoneIndexType> SKMFillComponentSpaceTransformsRequiredBones;
+				SkeletalMeshComponent->ComputeRequiredBones(SKMRequiredBones, SKMFillComponentSpaceTransformsRequiredBones, LODIndex, /*bIgnorePhysicsAsset*/ false);
+
+				CheckRequiredBones(PhysAssetBones, SKMFillComponentSpaceTransformsRequiredBones, MissingBones);
+				if (MissingBones.Num())
+				{
+					FText MissingBonesMsg = FText::FormatOrdered(LOCTEXT("MissingPhysicsBones", "SkeletalMesh \"{0}\" (LOD {1}) will not update the following bones required by \"{2}\"'s PhysicsAsset ({3}): "),
+						FText::FromString(SkeletalMeshComponent->GetSkinnedAsset()->GetName()),
+						FText::AsNumber(LODIndex), FText::FromString(ClothAsset->GetName()), FText::FromString(ClothComponent.GetPhysicsAsset()->GetName()));
+					for (const FBoneIndexType MissingBone : MissingBones)
+					{
+						MissingBonesMsg = FText::Format(LOCTEXT("MissingPhysicsBoneList", "{0} {1}"), MissingBonesMsg, 
+							FText::FromName(RefSkeleton.GetBoneName(MissingBone)));
+					}
+					UE_LOG(LogChaosClothAssetEditor, Warning, TEXT("%s"), *MissingBonesMsg.ToString());
+					bAnyMissingBones = true;
+				}
+
+				// These should be the bones that are required by the follower cloth component
+				const int32 ClothLODIndex = FMath::Min(ClothComponent.GetNumLODs() - 1, LODIndex);
+				if (ensure(ClothRenderData->LODRenderData.IsValidIndex(ClothLODIndex)))
+				{
+					CheckRequiredBones(ClothRenderData->LODRenderData[ClothLODIndex].ActiveBoneIndices,
+						SKMFillComponentSpaceTransformsRequiredBones, MissingBones);
+
+					if (MissingBones.Num())
+					{
+						FText MissingBonesMsg = FText::FormatOrdered(LOCTEXT("MissingClothBones", "SkeletalMesh \"{0}\" (LOD {1}) will not update the following bones required by \"{2}\" (LOD {3}): "),
+							FText::FromString(SkeletalMeshComponent->GetSkinnedAsset()->GetName()),
+							FText::AsNumber(LODIndex), FText::FromString(ClothAsset->GetName()), FText::AsNumber(ClothLODIndex));
+						for (const FBoneIndexType MissingBone : MissingBones)
+						{
+							MissingBonesMsg = FText::Format(LOCTEXT("MissingClothBoneList", "{0} {1}"), MissingBonesMsg, FText::FromName(RefSkeleton.GetBoneName(MissingBone)));
+						}
+						UE_LOG(LogChaosClothAssetEditor, Warning, TEXT("%s"), *MissingBonesMsg.ToString());
+						bAnyMissingBones = true;
+					}
+				}
+			}
+
+			if (bAnyMissingBones)
+			{
+				FNotificationInfo NotificationInfo(FText::Format(LOCTEXT("MissingBonesNotification", "Cloth asset {0} is not compatible with the preview skeletal mesh {1} (missing bones). See log for more details."),
+					FText::FromString(ClothAsset->GetName()), FText::FromString(SkeletalMeshComponent->GetSkinnedAsset()->GetName())));
+				NotificationInfo.ExpireDuration = 5.0f;
+				FSlateNotificationManager::Get().AddNotification(NotificationInfo);
+			}
+		}
+	}
+}
+
 void UChaosClothPreviewSceneDescription::SetPreviewScene(UE::Chaos::ClothAsset::FChaosClothPreviewScene* InPreviewScene)
 {
 	PreviewScene = InPreviewScene;
@@ -137,6 +237,7 @@ void FChaosClothPreviewScene::UpdateClothComponentAttachment()
 			const FComponentReregisterContext Context(ClothComponent);
 		}
 	}
+	UE::Chaos::ClothAsset::Private::ValidateClothComponentAttachmentBones(*ClothComponent);
 }
 
 void FChaosClothPreviewScene::SceneDescriptionPropertyChanged(const FName& PropertyName)
diff --git a/Engine/Source/Runtime/Engine/Classes/Components/SkeletalMeshComponent.h b/Engine/Source/Runtime/Engine/Classes/Components/SkeletalMeshComponent.h
--- a/Engine/Source/Runtime/Engine/Classes/Components/SkeletalMeshComponent.h
+++ b/Engine/Source/Runtime/Engine/Classes/Components/SkeletalMeshComponent.h
@@ -1658,9 +1658,6 @@ public:
 	// Get the required virtual bones from the SkeletalMesh Reference Skeleton
 	static ENGINE_API void GetRequiredVirtualBones(const USkeletalMesh* SkeletalMesh, TArray<FBoneIndexType>& OutRequiredBones);
 
-	// Get the bones used by the physics asset (if any), as these have to be updated even if LODed out
-	static ENGINE_API void GetPhysicsRequiredBones(const USkeletalMesh* SkeletalMesh, const UPhysicsAsset* PhysicsAsset, TArray<FBoneIndexType>& OutRequiredBones);
-
 	// Removes the bones expicitly hidden or hiden by parent
 	static ENGINE_API void ExcludeHiddenBones(const USkeletalMeshComponent* SkeletalMeshComponent, const USkeletalMesh* SkeletalMesh, TArray<FBoneIndexType>& OutRequiredBones);
 
@@ -1668,19 +1665,9 @@ public:
 	UE_DEPRECATED(5.3, "This method has been deprecated. Please use UMirrorDataTable for mirroring support.")
 	static ENGINE_API void GetMirroringRequiredBones(const USkeletalMesh* SkeletalMesh, TArray<FBoneIndexType>& OutRequiredBones);
 
-	// Get the bones that are used by sockets
-	// OutRequiredBones : the bones used by sockets that are set as always animate
-	// NeededBonesForFillComponentSpaceTransforms : the rest of bones used by sockets (that are not set to always animate)
-	static ENGINE_API void GetSocketRequiredBones(const USkeletalMesh* SkeletalMesh, TArray<FBoneIndexType>& OutRequiredBones, TArray<FBoneIndexType>& NeededBonesForFillComponentSpaceTransforms);
-
 	// Get the bones required for shadow shapes
 	static ENGINE_API void GetShadowShapeRequiredBones(const USkeletalMeshComponent* SkeletalMeshComponent, TArray<FBoneIndexType>& OutRequiredBones);
 
-	/** Takes sorted array Base and then adds any elements from sorted array Insert which is missing from it, preserving order.
-	 * this assumes both arrays are sorted and contain unique bone indices.
-	 */
-	static ENGINE_API void MergeInBoneIndexArrays(TArray<FBoneIndexType>& BaseArray, const TArray<FBoneIndexType>& InsertArray);
-
 	/**
 	* Recalculates the AnimCurveUids array in RequiredBone of this SkeletalMeshComponent based on current required bone set
 	* Is called when Skeleton->IsRequiredCurvesUpToDate() = false
diff --git a/Engine/Source/Runtime/Engine/Classes/Components/SkinnedMeshComponent.h b/Engine/Source/Runtime/Engine/Classes/Components/SkinnedMeshComponent.h
--- a/Engine/Source/Runtime/Engine/Classes/Components/SkinnedMeshComponent.h
+++ b/Engine/Source/Runtime/Engine/Classes/Components/SkinnedMeshComponent.h
@@ -1859,6 +1859,26 @@ public:
 
 	friend class FRenderStateRecreator;
 	friend class FSkeletalMeshStreamOut;
+
+	//
+	// Animation required bones
+	//
+
+	/** Get the bones used by the physics asset (if any), as these have to be updated even if LODed out. */
+	static ENGINE_API void GetPhysicsRequiredBones(const USkinnedAsset* SkinnedAsset, const UPhysicsAsset* PhysicsAsset, TArray<FBoneIndexType>& OutRequiredBones);
+
+	/**
+	 * Get the bones that are used by sockets.
+	 * @param OutRequiredBones - the bones used by sockets that are set as always animate
+	 * @param NeededBonesForFillComponentSpaceTransforms - the rest of bones used by sockets (that are not set to always animate)
+	 */ 
+	static ENGINE_API void GetSocketRequiredBones(const USkinnedAsset* SkinnedAsset, TArray<FBoneIndexType>& OutRequiredBones, TArray<FBoneIndexType>& NeededBonesForFillComponentSpaceTransforms);
+
+	/** 
+	 * Takes sorted array Base and then adds any elements from sorted array Insert which is missing from it, preserving order.
+	 * This assumes both arrays are sorted and contain unique bone indices.
+	 */
+	static ENGINE_API void MergeInBoneIndexArrays(TArray<FBoneIndexType>& BaseArray, const TArray<FBoneIndexType>& InsertArray);
 };
 
 class FRenderStateRecreator
diff --git a/Engine/Source/Runtime/Engine/Private/Components/SkeletalMeshComponent.cpp b/Engine/Source/Runtime/Engine/Private/Components/SkeletalMeshComponent.cpp
--- a/Engine/Source/Runtime/Engine/Private/Components/SkeletalMeshComponent.cpp
+++ b/Engine/Source/Runtime/Engine/Private/Components/SkeletalMeshComponent.cpp
@@ -1683,60 +1683,6 @@ static void IntersectBoneIndexArrays(TArray<FBoneIndexType>& Output, const TArra
 	}
 }
 
-/** Takes sorted array Base and then adds any elements from sorted array Insert which is missing from it, preserving order.
- * this assumes both arrays are sorted and contain unique bone indices. */
-/*static*/ void USkeletalMeshComponent::MergeInBoneIndexArrays(TArray<FBoneIndexType>& BaseArray, const TArray<FBoneIndexType>& InsertArray)
-{
-	// Then we merge them into the array of required bones.
-	int32 BaseBonePos = 0;
-	int32 InsertBonePos = 0;
-
-	// Iterate over each of the bones we need.
-	while( InsertBonePos < InsertArray.Num() )
-	{
-		// Find index of physics bone
-		FBoneIndexType InsertBoneIndex = InsertArray[InsertBonePos];
-
-		// If at end of BaseArray array - just append.
-		if( BaseBonePos == BaseArray.Num() )
-		{
-			BaseArray.Add(InsertBoneIndex);
-			BaseBonePos++;
-			InsertBonePos++;
-		}
-		// If in the middle of BaseArray, merge together.
-		else
-		{
-			// Check that the BaseArray array is strictly increasing, otherwise merge code does not work.
-			check( BaseBonePos == 0 || BaseArray[BaseBonePos-1] < BaseArray[BaseBonePos] );
-
-			// Get next required bone index.
-			FBoneIndexType BaseBoneIndex = BaseArray[BaseBonePos];
-
-			// We have a bone in BaseArray not required by Insert. Thats ok - skip.
-			if( BaseBoneIndex < InsertBoneIndex )
-			{
-				BaseBonePos++;
-			}
-			// Bone required by Insert is in 
-			else if( BaseBoneIndex == InsertBoneIndex )
-			{
-				BaseBonePos++;
-				InsertBonePos++;
-			}
-			// Bone required by Insert is missing - insert it now.
-			else // BaseBoneIndex > InsertBoneIndex
-			{
-				BaseArray.InsertUninitialized(BaseBonePos);
-				BaseArray[BaseBonePos] = InsertBoneIndex;
-
-				BaseBonePos++;
-				InsertBonePos++;
-			}
-		}
-	}
-}
-
 // this is optimized version of updating only curves
 // if you call RecalcRequiredBones, curve should be refreshed
 void USkeletalMeshComponent::RecalcRequiredCurves()
@@ -1846,35 +1792,6 @@ void USkeletalMeshComponent::ComputeRequiredBones(TArray<FBoneIndexType>& OutReq
 	MergeInBoneIndexArrays(OutRequiredBones, SkeletalMesh->GetRefSkeleton().GetRequiredVirtualBones());
 }
 
-/*static*/ void USkeletalMeshComponent::GetPhysicsRequiredBones(const USkeletalMesh* SkeletalMesh, const UPhysicsAsset* PhysicsAsset, TArray<FBoneIndexType>& OutRequiredBones)
-{
-	check(SkeletalMesh != nullptr);
-	check(PhysicsAsset != nullptr);
-
-	// If we have a PhysicsAsset, we also need to make sure that all the bones used by it are always updated, as its used
-	// by line checks etc. We might also want to kick in the physics, which means having valid bone transforms.
-	TArray<FBoneIndexType> PhysAssetBones;
-	PhysAssetBones.Reserve(PhysicsAsset->SkeletalBodySetups.Num());
-	for (int32 i = 0; i < PhysicsAsset->SkeletalBodySetups.Num(); i++)
-	{
-		if (!ensure(PhysicsAsset->SkeletalBodySetups[i]))
-		{
-			continue;
-		}
-		int32 PhysBoneIndex = SkeletalMesh->GetRefSkeleton().FindBoneIndex(PhysicsAsset->SkeletalBodySetups[i]->BoneName);
-		if (PhysBoneIndex != INDEX_NONE)
-		{
-			PhysAssetBones.Add(PhysBoneIndex);
-		}
-	}
-
-	// Then sort array of required bones in hierarchy order
-	PhysAssetBones.Sort();
-
-	// Make sure all of these are in RequiredBones.
-	MergeInBoneIndexArrays(OutRequiredBones, PhysAssetBones);
-}
-
 /*static*/ void USkeletalMeshComponent::ExcludeHiddenBones(const USkeletalMeshComponent* SkeletalMeshComponent, const USkeletalMesh* SkeletalMesh, TArray<FBoneIndexType>& OutRequiredBones)
 {
 	check(SkeletalMeshComponent != nullptr);
@@ -1916,36 +1833,6 @@ void USkeletalMeshComponent::ComputeRequiredBones(TArray<FBoneIndexType>& OutReq
 {
 }
 
-/*static*/ void USkeletalMeshComponent::GetSocketRequiredBones(const USkeletalMesh* SkeletalMesh, TArray<FBoneIndexType>& OutRequiredBones, TArray<FBoneIndexType>& NeededBonesForFillComponentSpaceTransforms)
-{
-	check(SkeletalMesh != nullptr);
-
-	TArray<FBoneIndexType> ForceAnimatedSocketBones;
-	TArray<USkeletalMeshSocket*> ActiveSocketList = SkeletalMesh->GetActiveSocketList();
-	ForceAnimatedSocketBones.Reserve(ActiveSocketList.Num());
-	for (const USkeletalMeshSocket* Socket : ActiveSocketList)
-	{
-		int32 BoneIndex = SkeletalMesh->GetRefSkeleton().FindBoneIndex(Socket->BoneName);
-		if (BoneIndex != INDEX_NONE)
-		{
-			if (Socket->bForceAlwaysAnimated)
-			{
-				ForceAnimatedSocketBones.AddUnique(BoneIndex);
-			}
-			else
-			{
-				NeededBonesForFillComponentSpaceTransforms.AddUnique(BoneIndex);
-			}
-		}
-	}
-
-	// Then sort array of required bones in hierarchy order
-	ForceAnimatedSocketBones.Sort();
-
-	// Make sure all of these are in OutRequiredBones.
-	MergeInBoneIndexArrays(OutRequiredBones, ForceAnimatedSocketBones);
-}
-
 /*static*/ void USkeletalMeshComponent::GetShadowShapeRequiredBones(const USkeletalMeshComponent* SkeletalMeshComponent, TArray<FBoneIndexType>& OutRequiredBones)
 {
 	if (FSkeletalMeshSceneProxy* SkeletalMeshProxy = (FSkeletalMeshSceneProxy*)SkeletalMeshComponent->SceneProxy)
diff --git a/Engine/Source/Runtime/Engine/Private/Components/SkinnedMeshComponent.cpp b/Engine/Source/Runtime/Engine/Private/Components/SkinnedMeshComponent.cpp
--- a/Engine/Source/Runtime/Engine/Private/Components/SkinnedMeshComponent.cpp
+++ b/Engine/Source/Runtime/Engine/Private/Components/SkinnedMeshComponent.cpp
@@ -4705,6 +4705,119 @@ void USkinnedMeshComponent::HandleFeatureLevelChanged(ERHIFeatureLevel::Type InF
 }
 #endif
 
+/** Takes sorted array Base and then adds any elements from sorted array Insert which is missing from it, preserving order.
+ * this assumes both arrays are sorted and contain unique bone indices. */
+/*static*/ void USkinnedMeshComponent::MergeInBoneIndexArrays(TArray<FBoneIndexType>& BaseArray, const TArray<FBoneIndexType>& InsertArray)
+{
+	// Then we merge them into the array of required bones.
+	int32 BaseBonePos = 0;
+	int32 InsertBonePos = 0;
+
+	// Iterate over each of the bones we need.
+	while (InsertBonePos < InsertArray.Num())
+	{
+		// Find index of physics bone
+		FBoneIndexType InsertBoneIndex = InsertArray[InsertBonePos];
+
+		// If at end of BaseArray array - just append.
+		if (BaseBonePos == BaseArray.Num())
+		{
+			BaseArray.Add(InsertBoneIndex);
+			BaseBonePos++;
+			InsertBonePos++;
+		}
+		// If in the middle of BaseArray, merge together.
+		else
+		{
+			// Check that the BaseArray array is strictly increasing, otherwise merge code does not work.
+			check(BaseBonePos == 0 || BaseArray[BaseBonePos - 1] < BaseArray[BaseBonePos]);
+
+			// Get next required bone index.
+			FBoneIndexType BaseBoneIndex = BaseArray[BaseBonePos];
+
+			// We have a bone in BaseArray not required by Insert. Thats ok - skip.
+			if (BaseBoneIndex < InsertBoneIndex)
+			{
+				BaseBonePos++;
+			}
+			// Bone required by Insert is in 
+			else if (BaseBoneIndex == InsertBoneIndex)
+			{
+				BaseBonePos++;
+				InsertBonePos++;
+			}
+			// Bone required by Insert is missing - insert it now.
+			else // BaseBoneIndex > InsertBoneIndex
+			{
+				BaseArray.InsertUninitialized(BaseBonePos);
+				BaseArray[BaseBonePos] = InsertBoneIndex;
+
+				BaseBonePos++;
+				InsertBonePos++;
+			}
+		}
+	}
+}
+
+void USkinnedMeshComponent::GetPhysicsRequiredBones(const USkinnedAsset* SkinnedAsset, const UPhysicsAsset* PhysicsAsset, TArray<FBoneIndexType>& OutRequiredBones)
+{
+	check(SkinnedAsset != nullptr);
+	check(PhysicsAsset != nullptr);
+
+	// If we have a PhysicsAsset, we also need to make sure that all the bones used by it are always updated, as its used
+	// by line checks etc. We might also want to kick in the physics, which means having valid bone transforms.
+	TArray<FBoneIndexType> PhysAssetBones;
+	PhysAssetBones.Reserve(PhysicsAsset->SkeletalBodySetups.Num());
+	for (int32 i = 0; i < PhysicsAsset->SkeletalBodySetups.Num(); ++i)
+	{
+		if (!ensure(PhysicsAsset->SkeletalBodySetups[i]))
+		{
+			continue;
+		}
+		const int32 PhysBoneIndex = SkinnedAsset->GetRefSkeleton().FindBoneIndex(PhysicsAsset->SkeletalBodySetups[i]->BoneName);
+		if (PhysBoneIndex != INDEX_NONE)
+		{
+			PhysAssetBones.Add(PhysBoneIndex);
+		}
+	}
+
+	// Then sort array of required bones in hierarchy order
+	PhysAssetBones.Sort();
+
+	// Make sure all of these are in RequiredBones.
+	MergeInBoneIndexArrays(OutRequiredBones, PhysAssetBones);
+}
+
+void USkinnedMeshComponent::GetSocketRequiredBones(const USkinnedAsset* SkinnedAsset, TArray<FBoneIndexType>& OutRequiredBones, TArray<FBoneIndexType>& NeededBonesForFillComponentSpaceTransforms)
+{
+	check(SkinnedAsset != nullptr);
+
+	TArray<FBoneIndexType> ForceAnimatedSocketBones;
+	const TArray<USkeletalMeshSocket*> ActiveSocketList = SkinnedAsset->GetActiveSocketList();
+	ForceAnimatedSocketBones.Reserve(ActiveSocketList.Num());
+	for (const USkeletalMeshSocket* Socket : ActiveSocketList)
+	{
+		const int32 BoneIndex = SkinnedAsset->GetRefSkeleton().FindBoneIndex(Socket->BoneName);
+		if (BoneIndex != INDEX_NONE)
+		{
+			if (Socket->bForceAlwaysAnimated)
+			{
+				ForceAnimatedSocketBones.AddUnique(BoneIndex);
+			}
+			else
+			{
+				NeededBonesForFillComponentSpaceTransforms.AddUnique(BoneIndex);
+			}
+		}
+	}
+
+	// Then sort array of required bones in hierarchy order
+	ForceAnimatedSocketBones.Sort();
+
+	// Make sure all of these are in OutRequiredBones.
+	MergeInBoneIndexArrays(OutRequiredBones, ForceAnimatedSocketBones);
+}
+
 void FAnimUpdateRateParameters::SetTrailMode(float DeltaTime, uint8 UpdateRateShift, int32 NewUpdateRate, int32 NewEvaluationRate, bool bNewInterpSkippedFrames)
 {
 	OptimizeMode = TrailMode;
