Update how RemoveOptionPin works for Dataflow nodes. These are currently only used by a couple of cloth nodes.

The problem was that we weren't updating the FDataflowNode's connection data that the connection was broken. This would eventually cause a crash because the connection data was deleted, but there were still pointers hanging around with them.This was fixed by calling BreakAllPinLinks with bNotifyNodes = true inside UDataflowEdNode::RemoveOptionPin and then deleting the pin AFTER breaking the link. I did the latter by breaking the RemovePin method into two methods, a const method GetPinToRemove which just returns which pin (if any) to remove, and a non-const method OnPinRemoved which will clean up the DataflowNode's connection.

Note that I did not update the way AddOptionPin works because getting the data required to create the input is easiest in one call to the derived node, rather than splitting. Also there is no existing callback for OnPinAdded at the ed graph level.

diff --git a/Engine/Plugins/Experimental/ChaosClothAssetEditor/Source/ChaosClothAssetDataflowNodes/Private/ChaosClothAsset/MergeClothCollectionsNode.cpp b/Engine/Plugins/Experimental/ChaosClothAssetEditor/Source/ChaosClothAssetDataflowNodes/Private/ChaosClothAsset/MergeClothCollectionsNode.cpp
--- a/Engine/Plugins/Experimental/ChaosClothAssetEditor/Source/ChaosClothAssetDataflowNodes/Private/ChaosClothAsset/MergeClothCollectionsNode.cpp
+++ b/Engine/Plugins/Experimental/ChaosClothAssetEditor/Source/ChaosClothAssetDataflowNodes/Private/ChaosClothAsset/MergeClothCollectionsNode.cpp
@@ -99,27 +99,68 @@ Dataflow::FPin FChaosClothAssetMergeClothCollectionsNode::AddPin()
 	return Super::AddPin();
 }
 
-Dataflow::FPin FChaosClothAssetMergeClothCollectionsNode::RemovePin()
+Dataflow::FPin FChaosClothAssetMergeClothCollectionsNode::GetPinToRemove() const
 {
-	auto RemoveInput = [this](const FManagedArrayCollection* InCollection) -> Dataflow::FPin
+	auto PinToRemove = [this](const FManagedArrayCollection* Collection) -> Dataflow::FPin
 	{
-		const FDataflowInput* const Input = FindInput(InCollection);
+		const FDataflowInput* const Input = FindInput(Collection);
 		check(Input);
-		Dataflow::FPin Pin = { Dataflow::FPin::EDirection::INPUT, Input->GetType(), Input->GetName() };
-		UnregisterInputConnection(InCollection);  // This will delete the input, so set the pin before that
-		return Pin;
+		return { Dataflow::FPin::EDirection::INPUT, Input->GetType(), Input->GetName() };
 	};
 
 	switch (NumInputs - 1)
 	{
-	case 1: --NumInputs; return RemoveInput(&Collection1);
-	case 2: --NumInputs; return RemoveInput(&Collection2);
-	case 3: --NumInputs; return RemoveInput(&Collection3);
-	case 4: --NumInputs; return RemoveInput(&Collection4);
-	case 5: --NumInputs; return RemoveInput(&Collection5);
+	case 1: return PinToRemove(&Collection1);
+	case 2: return PinToRemove(&Collection2);
+	case 3: return PinToRemove(&Collection3);
+	case 4: return PinToRemove(&Collection4);
+	case 5: return PinToRemove(&Collection5);
 	default: break;
 	}
-	return Super::AddPin();
+	return Super::GetPinToRemove();
+}
+
+void FChaosClothAssetMergeClothCollectionsNode::OnPinRemoved(const Dataflow::FPin& Pin)
+{
+	auto CheckPinRemoved = [this, &Pin](const FManagedArrayCollection* Collection)
+	{
+		check(Pin.Direction == Dataflow::FPin::EDirection::INPUT);
+#if DO_CHECK
+		const FDataflowInput* const Input = FindInput(Collection);
+		check(Input);
+		check(Input->GetName() == Pin.Name);
+		check(Input->GetType() == Pin.Type);
+#endif
+	};
+
+	switch (NumInputs - 1)
+	{
+	case 1:
+		CheckPinRemoved(&Collection1);
+		--NumInputs;
+		break;
+	case 2:
+		CheckPinRemoved(&Collection2);
+		--NumInputs;
+		break;
+	case 3:
+		CheckPinRemoved(&Collection3);
+		--NumInputs;
+		break;
+	case 4:
+		CheckPinRemoved(&Collection4);
+		--NumInputs;
+		break;
+	case 5:
+		CheckPinRemoved(&Collection5);
+		--NumInputs;
+		break;
+	default:
+		checkNoEntry();
+		break;
+	}
+
+	return Super::OnPinRemoved(Pin);
 }
 
 TArray<const FManagedArrayCollection*> FChaosClothAssetMergeClothCollectionsNode::GetCollections() const
diff --git a/Engine/Plugins/Experimental/ChaosClothAssetEditor/Source/ChaosClothAssetDataflowNodes/Private/ChaosClothAsset/TerminalNode.cpp b/Engine/Plugins/Experimental/ChaosClothAssetEditor/Source/ChaosClothAssetDataflowNodes/Private/ChaosClothAsset/TerminalNode.cpp
--- a/Engine/Plugins/Experimental/ChaosClothAssetEditor/Source/ChaosClothAssetDataflowNodes/Private/ChaosClothAsset/TerminalNode.cpp
+++ b/Engine/Plugins/Experimental/ChaosClothAssetEditor/Source/ChaosClothAssetDataflowNodes/Private/ChaosClothAsset/TerminalNode.cpp
@@ -172,27 +172,68 @@ Dataflow::FPin FChaosClothAssetTerminalNode::AddPin()
 	return Super::AddPin();
 }
 
-Dataflow::FPin FChaosClothAssetTerminalNode::RemovePin()
+Dataflow::FPin FChaosClothAssetTerminalNode::GetPinToRemove() const
 {
-	auto RemoveInput = [this](const FManagedArrayCollection* Collection) -> Dataflow::FPin
+	auto PinToRemove = [this](const FManagedArrayCollection* Collection) -> Dataflow::FPin
 		{
 			const FDataflowInput* const Input = FindInput(Collection);
 			check(Input);
-			Dataflow::FPin Pin = { Dataflow::FPin::EDirection::INPUT, Input->GetType(), Input->GetName() };
-			UnregisterInputConnection(Collection);  // This will delete the input, so set the pin before that
-			return Pin;
+			return { Dataflow::FPin::EDirection::INPUT, Input->GetType(), Input->GetName() };
 		};
 
 	switch (NumLods - 1)
 	{
-	case 1: --NumLods; return RemoveInput(&CollectionLod1);
-	case 2: --NumLods; return RemoveInput(&CollectionLod2);
-	case 3: --NumLods; return RemoveInput(&CollectionLod3);
-	case 4: --NumLods; return RemoveInput(&CollectionLod4);
-	case 5: --NumLods; return RemoveInput(&CollectionLod5);
+	case 1: return PinToRemove(&CollectionLod1);
+	case 2: return PinToRemove(&CollectionLod2);
+	case 3: return PinToRemove(&CollectionLod3);
+	case 4: return PinToRemove(&CollectionLod4);
+	case 5: return PinToRemove(&CollectionLod5);
 	default: break;
 	}
-	return Super::AddPin();
+	return Super::GetPinToRemove();
+}
+
+void FChaosClothAssetTerminalNode::OnPinRemoved(const Dataflow::FPin& Pin)
+{
+	auto CheckPinRemoved = [this, &Pin](const FManagedArrayCollection* Collection)
+	{
+		check(Pin.Direction == Dataflow::FPin::EDirection::INPUT);
+#if DO_CHECK
+		const FDataflowInput* const Input = FindInput(Collection);
+		check(Input);
+		check(Input->GetName() == Pin.Name);
+		check(Input->GetType() == Pin.Type);
+#endif
+	};
+
+	switch (NumLods - 1)
+	{
+	case 1:
+		CheckPinRemoved(&CollectionLod1);
+		--NumLods; 
+		break;
+	case 2:
+		CheckPinRemoved(&CollectionLod2);
+		--NumLods;
+		break;
+	case 3:
+		CheckPinRemoved(&CollectionLod3);
+		--NumLods;
+		break;
+	case 4:
+		CheckPinRemoved(&CollectionLod4);
+		--NumLods;
+		break;
+	case 5:
+		CheckPinRemoved(&CollectionLod5);
+		--NumLods;
+		break;
+	default: 
+		checkNoEntry();
+		break;
+	}
+
+	return Super::OnPinRemoved(Pin);
 }
 
 TArray<const FManagedArrayCollection*> FChaosClothAssetTerminalNode::GetCollectionLods() const
diff --git a/Engine/Plugins/Experimental/ChaosClothAssetEditor/Source/ChaosClothAssetDataflowNodes/Public/ChaosClothAsset/MergeClothCollectionsNode.h b/Engine/Plugins/Experimental/ChaosClothAssetEditor/Source/ChaosClothAssetDataflowNodes/Public/ChaosClothAsset/MergeClothCollectionsNode.h
--- a/Engine/Plugins/Experimental/ChaosClothAssetEditor/Source/ChaosClothAssetDataflowNodes/Public/ChaosClothAsset/MergeClothCollectionsNode.h
+++ b/Engine/Plugins/Experimental/ChaosClothAssetEditor/Source/ChaosClothAssetDataflowNodes/Public/ChaosClothAsset/MergeClothCollectionsNode.h
@@ -46,7 +46,8 @@ private:
 	virtual Dataflow::FPin AddPin() override;
 	virtual bool CanAddPin() const override { return NumInputs < MaxInputs; }
 	virtual bool CanRemovePin() const override { return NumInputs > 1; }
-	virtual Dataflow::FPin RemovePin() override;
+	virtual Dataflow::FPin GetPinToRemove() const override;
+	virtual void OnPinRemoved(const Dataflow::FPin& Pin) override;
 	virtual void Serialize(FArchive& Ar) override;
 	//~ End FDataflowNode interface
 
diff --git a/Engine/Plugins/Experimental/ChaosClothAssetEditor/Source/ChaosClothAssetDataflowNodes/Public/ChaosClothAsset/TerminalNode.h b/Engine/Plugins/Experimental/ChaosClothAssetEditor/Source/ChaosClothAssetDataflowNodes/Public/ChaosClothAsset/TerminalNode.h
--- a/Engine/Plugins/Experimental/ChaosClothAssetEditor/Source/ChaosClothAssetDataflowNodes/Public/ChaosClothAsset/TerminalNode.h
+++ b/Engine/Plugins/Experimental/ChaosClothAssetEditor/Source/ChaosClothAssetDataflowNodes/Public/ChaosClothAsset/TerminalNode.h
@@ -48,7 +48,8 @@ private:
 	virtual Dataflow::FPin AddPin() override;
 	virtual bool CanAddPin() const override { return NumLods < MaxLods; }
 	virtual bool CanRemovePin() const override { return NumLods > 1; }
-	virtual Dataflow::FPin RemovePin() override;
+	virtual Dataflow::FPin GetPinToRemove() const override;
+	virtual void OnPinRemoved(const Dataflow::FPin& Pin) override;
 	virtual void Serialize(FArchive& Ar) override;
 	//~ End FDataflowNode interface
 
diff --git a/Engine/Source/Runtime/Experimental/Dataflow/Core/Private/Dataflow/DataflowNode.cpp b/Engine/Source/Runtime/Experimental/Dataflow/Core/Private/Dataflow/DataflowNode.cpp
--- a/Engine/Source/Runtime/Experimental/Dataflow/Core/Private/Dataflow/DataflowNode.cpp
+++ b/Engine/Source/Runtime/Experimental/Dataflow/Core/Private/Dataflow/DataflowNode.cpp
@@ -222,6 +222,42 @@ TArray<Dataflow::FPin> FDataflowNode::GetPins() const
 	return RetVal;
 }
 
+void FDataflowNode::UnregisterPinConnection(const Dataflow::FPin& Pin)
+{
+	if (Pin.Direction == Dataflow::FPin::EDirection::INPUT)
+	{
+		for (TMap< int32, FDataflowInput*>::TIterator Iter = Inputs.CreateIterator(); Iter; ++Iter)
+		{
+			FDataflowInput* Con = Iter.Value();
+			if (Con->GetName().IsEqual(Pin.Name) && Con->GetType().IsEqual(Pin.Type))
+			{
+				Iter.RemoveCurrent();
+				delete Con;
+
+				// Invalidate graph as this input might have had connections
+				Invalidate();
+				break;
+			}
+		}
+	}
+	else if (Pin.Direction == Dataflow::FPin::EDirection::OUTPUT)
+	{
+		for (TMap<int32, FDataflowOutput*>::TIterator Iter = Outputs.CreateIterator(); Iter; ++Iter)
+		{
+			FDataflowOutput* Con = Iter.Value();
+			if (Con->GetName().IsEqual(Pin.Name) && Con->GetType().IsEqual(Pin.Type))
+			{
+				Iter.RemoveCurrent();
+				delete Con;
+
+				// Invalidate graph as this input might have had connections
+				Invalidate();
+				break;
+			}
+		}
+	}
+}
+
 void FDataflowNode::Invalidate(const Dataflow::FTimestamp& InModifiedTimestamp)
 {
 	if (LastModifiedTimestamp < InModifiedTimestamp)
diff --git a/Engine/Source/Runtime/Experimental/Dataflow/Core/Public/Dataflow/DataflowNode.h b/Engine/Source/Runtime/Experimental/Dataflow/Core/Public/Dataflow/DataflowNode.h
--- a/Engine/Source/Runtime/Experimental/Dataflow/Core/Public/Dataflow/DataflowNode.h
+++ b/Engine/Source/Runtime/Experimental/Dataflow/Core/Public/Dataflow/DataflowNode.h
@@ -93,9 +93,16 @@ struct FDataflowNode
 	virtual Dataflow::FPin AddPin() { return { Dataflow::FPin::EDirection::NONE, NAME_None, NAME_None }; }
 	/** Override this function to add the AddOptionPin functionality to the node's context menu. */
 	virtual bool CanAddPin() const { return false; }
-	/** Override this function to add the RemoveOPtionPin functionality to the node's context menu. */
-	virtual Dataflow::FPin RemovePin() { return { Dataflow::FPin::EDirection::NONE, NAME_None, NAME_None }; }
-	/** Override this function to add the RemoveOPtionPin functionality to the node's context menu. */
+	/** Override this function to add the RemoveOptionPin functionality to the node's context menu. */
+	virtual Dataflow::FPin GetPinToRemove() const { return { Dataflow::FPin::EDirection::NONE, NAME_None, NAME_None }; }
+	UE_DEPRECATED(5.4, "Use GetPinToRemove and OnPinRemoved instead.")
+	virtual Dataflow::FPin RemovePin() { return GetPinToRemove(); }
+	/** 
+	 * Override this to update any bookkeeping when a pin is being removed.
+	 * This will be called before the pin is unregistered as an input.
+	 */
+	virtual void OnPinRemoved(const Dataflow::FPin& Pin) {}
+	/** Override this function to add the RemoveOptionPin functionality to the node's context menu. */
 	virtual bool CanRemovePin() const { return false; }
 
 	DATAFLOWCORE_API virtual void AddInput(FDataflowInput* InPtr);
@@ -148,6 +155,8 @@ struct FDataflowNode
 		const FName& PassthroughName = NAME_None);
 	/** Unregister the input connection if one exists matching this property, and then invalidate the graph. */
 	DATAFLOWCORE_API void UnregisterInputConnection(const void* Property, const FName& PropertyName = NAME_None);
+	/** Unregister the connection if one exists matching this pin, then invalidate the graph. */
+	DATAFLOWCORE_API void UnregisterPinConnection(const Dataflow::FPin& Pin);
 
 	//
 	// Evaluation
diff --git a/Engine/Source/Runtime/Experimental/Dataflow/Engine/Private/Dataflow/DataflowEdNode.cpp b/Engine/Source/Runtime/Experimental/Dataflow/Engine/Private/Dataflow/DataflowEdNode.cpp
--- a/Engine/Source/Runtime/Experimental/Dataflow/Engine/Private/Dataflow/DataflowEdNode.cpp
+++ b/Engine/Source/Runtime/Experimental/Dataflow/Engine/Private/Dataflow/DataflowEdNode.cpp
@@ -20,7 +20,6 @@
 
 DEFINE_LOG_CATEGORY_STATIC(DATAFLOWNODE_LOG, Error, All);
 
-
 UDataflowEdNode::UDataflowEdNode(const FObjectInitializer& ObjectInitializer)
 	: Super(ObjectInitializer)
 {
@@ -201,13 +200,14 @@ void UDataflowEdNode::RemoveOptionPin()
 
 		if (const TSharedPtr<FDataflowNode> DataflowNode = DataflowGraph->FindBaseNode(DataflowNodeGuid))
 		{
-			const Dataflow::FPin Pin = DataflowNode->RemovePin();
+			const Dataflow::FPin Pin = DataflowNode->GetPinToRemove();
 			switch (Pin.Direction)
 			{
 			case Dataflow::FPin::EDirection::INPUT:
 				if (UEdGraphPin* const EdPin = FindPin(Pin.Name, EEdGraphPinDirection::EGPD_Input))
 				{
-					EdPin->BreakAllPinLinks();
+					constexpr bool bNotifyNodes = true;
+					EdPin->BreakAllPinLinks(bNotifyNodes);
 					RemovePin(EdPin);
 					ReconstructNode();
 				}
@@ -215,7 +215,8 @@ void UDataflowEdNode::RemoveOptionPin()
 			case Dataflow::FPin::EDirection::OUTPUT:
 				if (UEdGraphPin* const EdPin = FindPin(Pin.Name, EEdGraphPinDirection::EGPD_Output))
 				{
-					EdPin->BreakAllPinLinks();
+					constexpr bool bNotifyNodes = true;
+					EdPin->BreakAllPinLinks(bNotifyNodes);
 					RemovePin(EdPin);
 					ReconstructNode();
 				}
@@ -448,6 +449,34 @@ void UDataflowEdNode::AutowireNewNode(UEdGraphPin* FromPin)
 	}
 }
 
+void UDataflowEdNode::OnPinRemoved(UEdGraphPin* InRemovedPin)
+{
+	if (DataflowGraph && DataflowNodeGuid.IsValid())
+	{
+		if (const TSharedPtr<FDataflowNode> DataflowNode = DataflowGraph->FindBaseNode(DataflowNodeGuid))
+		{
+			if (InRemovedPin->Direction == EEdGraphPinDirection::EGPD_Input)
+			{
+				if (FDataflowInput* Con = DataflowNode->FindInput(FName(InRemovedPin->GetName())))
+				{
+					const Dataflow::FPin Pin = { Dataflow::FPin::EDirection::INPUT, Con->GetType(), Con->GetName() };
+					DataflowNode->OnPinRemoved(Pin);
+					DataflowNode->UnregisterPinConnection(Pin);
+				}
+			}
+			else if (InRemovedPin->Direction == EEdGraphPinDirection::EGPD_Output)
+			{
+				if (FDataflowOutput* Con = DataflowNode->FindOutput(FName(InRemovedPin->GetName())))
+				{
+					const Dataflow::FPin Pin = { Dataflow::FPin::EDirection::OUTPUT, Con->GetType(), Con->GetName() };
+					DataflowNode->OnPinRemoved(Pin);
+					DataflowNode->UnregisterPinConnection(Pin);
+				}
+			}
+		}
+	}
+}
+
 #endif
 
 
diff --git a/Engine/Source/Runtime/Experimental/Dataflow/Engine/Public/Dataflow/DataflowEdNode.h b/Engine/Source/Runtime/Experimental/Dataflow/Engine/Public/Dataflow/DataflowEdNode.h
--- a/Engine/Source/Runtime/Experimental/Dataflow/Engine/Public/Dataflow/DataflowEdNode.h
+++ b/Engine/Source/Runtime/Experimental/Dataflow/Engine/Public/Dataflow/DataflowEdNode.h
@@ -46,6 +46,7 @@ public:
 	DATAFLOWENGINE_API virtual void GetPinHoverText(const UEdGraphPin& Pin, FString& HoverTextOut) const override;
 	DATAFLOWENGINE_API virtual FText GetPinDisplayName(const UEdGraphPin* Pin) const override;
 	DATAFLOWENGINE_API virtual void AutowireNewNode(UEdGraphPin* FromPin) override;
+	DATAFLOWENGINE_API virtual void OnPinRemoved(UEdGraphPin* InRemovedPin) override;
 #endif // WITH_EDITOR
 	// End of UEdGraphNode interface
 
