Fix CIS bug from 27281487

diff --git a/Engine/Plugins/Experimental/ChaosClothAssetEditor/Source/ChaosClothAssetDataflowNodes/Private/ChaosClothAsset/MergeClothCollectionsNode.cpp b/Engine/Plugins/Experimental/ChaosClothAssetEditor/Source/ChaosClothAssetDataflowNodes/Private/ChaosClothAsset/MergeClothCollectionsNode.cpp
--- a/Engine/Plugins/Experimental/ChaosClothAssetEditor/Source/ChaosClothAssetDataflowNodes/Private/ChaosClothAsset/MergeClothCollectionsNode.cpp
+++ b/Engine/Plugins/Experimental/ChaosClothAssetEditor/Source/ChaosClothAssetDataflowNodes/Private/ChaosClothAsset/MergeClothCollectionsNode.cpp
@@ -101,9 +101,9 @@ Dataflow::FPin FChaosClothAssetMergeClothCollectionsNode::AddPin()
 
 Dataflow::FPin FChaosClothAssetMergeClothCollectionsNode::GetPinToRemove() const
 {
-	auto PinToRemove = [this](const FManagedArrayCollection* Collection) -> Dataflow::FPin
+	auto PinToRemove = [this](const FManagedArrayCollection* InCollection) -> Dataflow::FPin
 	{
-		const FDataflowInput* const Input = FindInput(Collection);
+		const FDataflowInput* const Input = FindInput(InCollection);
 		check(Input);
 		return { Dataflow::FPin::EDirection::INPUT, Input->GetType(), Input->GetName() };
 	};
@@ -122,11 +122,11 @@ Dataflow::FPin FChaosClothAssetMergeClothCollectionsNode::GetPinToRemove() const
 
 void FChaosClothAssetMergeClothCollectionsNode::OnPinRemoved(const Dataflow::FPin& Pin)
 {
-	auto CheckPinRemoved = [this, &Pin](const FManagedArrayCollection* Collection)
+	auto CheckPinRemoved = [this, &Pin](const FManagedArrayCollection* InCollection)
 	{
 		check(Pin.Direction == Dataflow::FPin::EDirection::INPUT);
 #if DO_CHECK
-		const FDataflowInput* const Input = FindInput(Collection);
+		const FDataflowInput* const Input = FindInput(InCollection);
 		check(Input);
 		check(Input->GetName() == Pin.Name);
 		check(Input->GetType() == Pin.Type);
