Cloth Editor: Update terminal node to only update properties rather than rebuild the asset if only property values have changed. This allows you to live tweak simulation properties.

The terminal node determines if any data forcing a rebuild has changed by hashing all of the data. The same runtime hash used by TSet/TMap (GetTypeHash) is used to hash the data. The hashing is editor-only and non-persistent.

diff --git a/Engine/Plugins/Experimental/ChaosClothAsset/Source/ChaosClothAsset/Private/ChaosClothAsset/CollectionClothFacade.cpp b/Engine/Plugins/Experimental/ChaosClothAsset/Source/ChaosClothAsset/Private/ChaosClothAsset/CollectionClothFacade.cpp
--- a/Engine/Plugins/Experimental/ChaosClothAsset/Source/ChaosClothAsset/Private/ChaosClothAsset/CollectionClothFacade.cpp
+++ b/Engine/Plugins/Experimental/ChaosClothAsset/Source/ChaosClothAsset/Private/ChaosClothAsset/CollectionClothFacade.cpp
@@ -23,6 +23,98 @@ namespace UE::Chaos::ClothAsset
 		return ClothCollection->IsValid() && ClothCollection->GetNumElements(FClothCollection::LodsGroup) == 1;
 	}
 
+	bool FCollectionClothConstFacade::HasValidData() const
+	{
+		return IsValid() &&
+			GetNumSimPatterns() &&
+			GetNumRenderPatterns() &&
+			GetNumSimVertices2D() &&
+			GetNumSimVertices3D() &&
+			GetNumSimFaces() &&
+			GetNumRenderVertices() &&
+			GetNumRenderFaces();
+	}
+
+	uint32 FCollectionClothConstFacade::CalculateTypeHash(bool bIncludeWeightMaps, uint32 PreviousHash) const
+	{
+		check(IsValid());
+		uint32 ResultHash = PreviousHash;
+
+		//~ LOD (single per collection) Group
+		ResultHash = HashCombineFast(ResultHash, GetTypeHash(GetSkeletalMeshPathName()));
+		ResultHash = HashCombineFast(ResultHash, GetTypeHash(GetSkeletalMeshPathName()));
+
+		//~ Seam Stitches Group
+		ResultHash = HashCombineFast(ResultHash, GetTypeHash(ClothCollection->GetSeamStitchStart()));
+		ResultHash = HashCombineFast(ResultHash, GetTypeHash(ClothCollection->GetSeamStitchEnd()));
+
+		//~ Seam Stitches Group
+		ResultHash = HashCombineFast(ResultHash, GetTypeHash(ClothCollection->GetSeamStitch2DEndIndices()));
+		ResultHash = HashCombineFast(ResultHash, GetTypeHash(ClothCollection->GetSeamStitch3DIndex()));
+
+		//~ Sim Patterns Group
+		ResultHash = HashCombineFast(ResultHash, GetTypeHash(ClothCollection->GetSimVertices2DStart()));
+		ResultHash = HashCombineFast(ResultHash, GetTypeHash(ClothCollection->GetSimVertices2DEnd()));
+		ResultHash = HashCombineFast(ResultHash, GetTypeHash(ClothCollection->GetSimFacesStart()));
+		ResultHash = HashCombineFast(ResultHash, GetTypeHash(ClothCollection->GetSimFacesEnd()));
+
+		//~ Render Patterns Group
+		ResultHash = HashCombineFast(ResultHash, GetTypeHash(ClothCollection->GetRenderVerticesStart()));
+		ResultHash = HashCombineFast(ResultHash, GetTypeHash(ClothCollection->GetRenderVerticesEnd()));
+		ResultHash = HashCombineFast(ResultHash, GetTypeHash(ClothCollection->GetRenderFacesStart()));
+		ResultHash = HashCombineFast(ResultHash, GetTypeHash(ClothCollection->GetRenderFacesEnd()));
+		ResultHash = HashCombineFast(ResultHash, GetTypeHash(ClothCollection->GetRenderMaterialPathName()));
+
+		//~ Sim Faces Group
+		ResultHash = HashCombineFast(ResultHash, GetTypeHash(ClothCollection->GetSimIndices2D()));
+		ResultHash = HashCombineFast(ResultHash, GetTypeHash(ClothCollection->GetSimIndices3D()));
+
+		//~ Sim Vertices 2D Group
+		ResultHash = HashCombineFast(ResultHash, GetTypeHash(ClothCollection->GetSimPosition2D()));
+		ResultHash = HashCombineFast(ResultHash, GetTypeHash(ClothCollection->GetSimVertex3DLookup()));
+		
+		//~ Sim Vertices 3D Group
+		ResultHash = HashCombineFast(ResultHash, GetTypeHash(ClothCollection->GetSimPosition3D()));
+		ResultHash = HashCombineFast(ResultHash, GetTypeHash(ClothCollection->GetSimNormal()));
+		ResultHash = HashCombineFast(ResultHash, GetTypeHash(ClothCollection->GetSimBoneIndices()));
+		ResultHash = HashCombineFast(ResultHash, GetTypeHash(ClothCollection->GetSimBoneWeights()));
+		ResultHash = HashCombineFast(ResultHash, GetTypeHash(ClothCollection->GetTetherKinematicIndex()));
+		ResultHash = HashCombineFast(ResultHash, GetTypeHash(ClothCollection->GetTetherReferenceLength()));
+		ResultHash = HashCombineFast(ResultHash, GetTypeHash(ClothCollection->GetSimVertex2DLookup()));
+		ResultHash = HashCombineFast(ResultHash, GetTypeHash(ClothCollection->GetSeamStitchLookup()));
+
+		//~ Render Faces Group
+		ResultHash = HashCombineFast(ResultHash, GetTypeHash(ClothCollection->GetRenderIndices()));
+
+		//~ Render Vertices Group
+		ResultHash = HashCombineFast(ResultHash, GetTypeHash(ClothCollection->GetRenderPosition()));
+		ResultHash = HashCombineFast(ResultHash, GetTypeHash(ClothCollection->GetRenderNormal()));
+		ResultHash = HashCombineFast(ResultHash, GetTypeHash(ClothCollection->GetRenderTangentU()));
+		ResultHash = HashCombineFast(ResultHash, GetTypeHash(ClothCollection->GetRenderTangentV()));
+		ResultHash = HashCombineFast(ResultHash, GetTypeHash(ClothCollection->GetRenderUVs()));
+		ResultHash = HashCombineFast(ResultHash, GetTypeHash(ClothCollection->GetRenderColor()));
+		ResultHash = HashCombineFast(ResultHash, GetTypeHash(ClothCollection->GetRenderBoneIndices()));
+		ResultHash = HashCombineFast(ResultHash, GetTypeHash(ClothCollection->GetRenderBoneWeights()));
+
+		if (bIncludeWeightMaps)
+		{
+			ResultHash = CalculateWeightMapTypeHash(ResultHash);
+		}
+		return ResultHash;
+	}
+
+	uint32 FCollectionClothConstFacade::CalculateWeightMapTypeHash(uint32 PreviousHash) const
+	{
+		uint32 ResultHash = PreviousHash;
+		const TArray<FName> WeightMapNames = GetWeightMapNames();
+		for (const FName& WeightMapName : WeightMapNames)
+		{
+			ResultHash = HashCombineFast(ResultHash, GetTypeHash(WeightMapName));
+			ResultHash = HashCombineFast(ResultHash, GetTypeHash(ClothCollection->GetUserDefinedAttribute<float>(WeightMapName, FClothCollection::SimVertices3DGroup)));
+		}
+		return ResultHash;
+	}
+
 	const FString& FCollectionClothConstFacade::GetPhysicsAssetPathName() const
 	{
 		static const FString EmptyString;
diff --git a/Engine/Plugins/Experimental/ChaosClothAsset/Source/ChaosClothAsset/Public/ChaosClothAsset/CollectionClothFacade.h b/Engine/Plugins/Experimental/ChaosClothAsset/Source/ChaosClothAsset/Public/ChaosClothAsset/CollectionClothFacade.h
--- a/Engine/Plugins/Experimental/ChaosClothAsset/Source/ChaosClothAsset/Public/ChaosClothAsset/CollectionClothFacade.h
+++ b/Engine/Plugins/Experimental/ChaosClothAsset/Source/ChaosClothAsset/Public/ChaosClothAsset/CollectionClothFacade.h
@@ -35,6 +35,15 @@ namespace UE::Chaos::ClothAsset
 		/** Return whether the facade is defined on the collection. */
 		bool IsValid() const;
 
+		/**
+		 * Return whether the facade has a non-empty sim and render mesh data.
+		 */
+		bool HasValidData() const;
+
+		uint32 CalculateTypeHash(bool bIncludeWeightMaps, uint32 PreviousHash = 0) const;
+
+		uint32 CalculateWeightMapTypeHash(uint32 PreviousHash = 0) const;
+
 		//~ LOD (single per collection) Group
 		/** Return the physics asset path names used for this collection. */
 		const FString& GetPhysicsAssetPathName() const;
diff --git a/Engine/Plugins/Experimental/ChaosClothAsset/Source/ChaosClothAssetEngine/Private/ChaosClothAsset/ClothComponent.cpp b/Engine/Plugins/Experimental/ChaosClothAsset/Source/ChaosClothAssetEngine/Private/ChaosClothAsset/ClothComponent.cpp
--- a/Engine/Plugins/Experimental/ChaosClothAsset/Source/ChaosClothAssetEngine/Private/ChaosClothAsset/ClothComponent.cpp
+++ b/Engine/Plugins/Experimental/ChaosClothAsset/Source/ChaosClothAssetEngine/Private/ChaosClothAsset/ClothComponent.cpp
@@ -97,6 +97,27 @@ void UChaosClothComponent::ResetConfigProperties()
 	}
 }
 
+#if WITH_EDITOR
+void UChaosClothComponent::UpdateConfigProperties()
+{
+	if (IsRegistered())
+	{
+		if (GetClothAsset())
+		{
+			const TArray<TSharedRef<FManagedArrayCollection>>& ClothCollections = GetClothAsset()->GetClothCollections();
+			if (ClothCollections.Num() == PropertyCollections.Num())
+			{
+				check(CollectionPropertyFacades.Num() == ClothCollections.Num());
+				for (int32 LodIndex = 0; LodIndex < ClothCollections.Num(); ++LodIndex)
+				{
+					CollectionPropertyFacades[LodIndex]->UpdateProperties(ClothCollections[LodIndex].ToSharedPtr());
+				}
+			}
+		}
+	}
+}
+#endif
+
 void UChaosClothComponent::RecreateClothSimulationProxy()
 {
 	if (IsRegistered())
diff --git a/Engine/Plugins/Experimental/ChaosClothAsset/Source/ChaosClothAssetEngine/Public/ChaosClothAsset/ClothComponent.h b/Engine/Plugins/Experimental/ChaosClothAsset/Source/ChaosClothAssetEngine/Public/ChaosClothAsset/ClothComponent.h
--- a/Engine/Plugins/Experimental/ChaosClothAsset/Source/ChaosClothAssetEngine/Public/ChaosClothAsset/ClothComponent.h
+++ b/Engine/Plugins/Experimental/ChaosClothAsset/Source/ChaosClothAssetEngine/Public/ChaosClothAsset/ClothComponent.h
@@ -102,6 +102,11 @@ public:
 
 	const UE::Chaos::ClothAsset::FClothSimulationProxy* GetClothSimulationProxy() const { return ClothSimulationProxy.Get(); }
 
+#if WITH_EDITOR
+	/** Update config properties from the asset. Will only update existing values.*/
+	void UpdateConfigProperties();
+#endif
+
 protected:
 	//~ Begin UObject Interface
 	virtual void PostLoad() override;
diff --git a/Engine/Plugins/Experimental/ChaosClothAssetEditor/Source/ChaosClothAssetDataflowNodes/Private/ChaosClothAsset/TerminalNode.cpp b/Engine/Plugins/Experimental/ChaosClothAssetEditor/Source/ChaosClothAssetDataflowNodes/Private/ChaosClothAsset/TerminalNode.cpp
--- a/Engine/Plugins/Experimental/ChaosClothAssetEditor/Source/ChaosClothAssetDataflowNodes/Private/ChaosClothAsset/TerminalNode.cpp
+++ b/Engine/Plugins/Experimental/ChaosClothAssetEditor/Source/ChaosClothAssetDataflowNodes/Private/ChaosClothAsset/TerminalNode.cpp
@@ -17,6 +17,49 @@
 
 #define LOCTEXT_NAMESPACE "ChaosClothAssetTerminalNode"
 
+namespace UE::Chaos::ClothAsset::Private
+{
+	uint32 CalculateClothChecksum(const TArray<TSharedRef<FManagedArrayCollection>>& InClothCollections)
+	{
+		uint32 Checksum = 0;
+		for (const TSharedRef<FManagedArrayCollection>& ClothCollection : InClothCollections)
+		{
+			constexpr bool bIncludeWeightMapsFalse = false;
+			FCollectionClothConstFacade Cloth(ClothCollection);
+			if (Cloth.HasValidData())
+			{
+				Checksum = Cloth.CalculateTypeHash(bIncludeWeightMapsFalse, Checksum);
+			}
+		}
+		return Checksum;
+	}
+
+	bool PropertyKeysMatch(const TArray<TSharedRef<FManagedArrayCollection>>& Collections0, const TArray<TSharedRef<FManagedArrayCollection>>& Collections1)
+	{
+		if (Collections0.Num() != Collections1.Num())
+		{
+			return false;
+		}
+		for (int32 LODIndex = 0; LODIndex < Collections0.Num(); ++LODIndex)
+		{
+			::Chaos::Softs::FCollectionPropertyConstFacade Property0(Collections0[LODIndex].ToSharedPtr());
+			::Chaos::Softs::FCollectionPropertyConstFacade Property1(Collections1[LODIndex].ToSharedPtr());
+			if (Property0.Num() != Property1.Num())
+			{
+				return false;
+			}
+			for (int32 PropertyIndex = 0; PropertyIndex < Property0.Num(); ++PropertyIndex)
+			{
+				if (Property0.GetKey(PropertyIndex) != Property1.GetKey(PropertyIndex))
+				{
+					return false;
+				}
+			}
+		}
+		return true;
+	}
+}
+
 FChaosClothAssetTerminalNode::FChaosClothAssetTerminalNode(const Dataflow::FNodeParameters& InParam, FGuid InGuid)
 	: FDataflowTerminalNode(InParam, InGuid)
 {
@@ -28,11 +71,30 @@ void FChaosClothAssetTerminalNode::SetAssetValue(TObjectPtr<UObject> Asset, Data
 	if (UChaosClothAsset* ClothAsset = Cast<UChaosClothAsset>(Asset.Get()))
 	{
 		using namespace UE::Chaos::ClothAsset;
-		const TArray<const FManagedArrayCollection*> CollectionLods = GetCollectionLods();
 
-		// Reset the asset's collection
+		TArray<TSharedRef<FManagedArrayCollection>> InClothCollections = GetCleanedCollectionLodValues(Context);
 		TArray<TSharedRef<FManagedArrayCollection>>& ClothCollections = ClothAsset->GetClothCollections();
-		ClothCollections.Reset(CollectionLods.Num());
+
+		const uint32 PreviousChecksum = ClothColllectionChecksum;
+		const bool bPreviousChecksumsValid = bClothCollectionChecksumValid;
+		ClothColllectionChecksum = Private::CalculateClothChecksum(InClothCollections);
+		bClothCollectionChecksumValid = InClothCollections.Num() > 0;
+
+		if (bPreviousChecksumsValid && PreviousChecksum == ClothColllectionChecksum && Private::PropertyKeysMatch(InClothCollections, ClothCollections))
+		{
+			// Cloth and property keys match. Just update property values.
+			check(InClothCollections.Num() == ClothCollections.Num());
+			check(ClothCollections.Num() > 0);
+			for (int32 LODIndex = 0; LODIndex < InClothCollections.Num(); ++LODIndex)
+			{
+				Chaos::Softs::FCollectionPropertyFacade Properties(ClothCollections[LODIndex]);
+				Properties.UpdateProperties(InClothCollections[LODIndex].ToSharedPtr());
+			}
+			return;
+		}
+
+		// Reset the asset's collection
+		ClothCollections.Reset(InClothCollections.Num());
 
 		// Reset the asset's material list
 		ClothAsset->GetMaterials().Reset();
@@ -40,91 +102,50 @@ void FChaosClothAssetTerminalNode::SetAssetValue(TObjectPtr<UObject> Asset, Data
 		// Iterate through the LODs
 		FString PhysicsAssetPathName;
 
-		int32 LastValidLodIndex = INDEX_NONE;
-		for (int32 LodIndex = 0; LodIndex < CollectionLods.Num(); ++LodIndex)
+		for (int32 LodIndex = 0; LodIndex < InClothCollections.Num(); ++LodIndex)
 		{
 			// New LOD
 			TSharedRef<FManagedArrayCollection>& ClothCollection = ClothCollections.Emplace_GetRef(MakeShared<FManagedArrayCollection>());
 			UE::Chaos::ClothAsset::FCollectionClothFacade ClothFacade(ClothCollection);
 			ClothFacade.DefineSchema();
 
-			// Retrieve input LOD
-			const FManagedArrayCollection& InCollectionLod = GetValue<FManagedArrayCollection>(Context, CollectionLods[LodIndex]);
-			const TSharedRef<const FManagedArrayCollection> InClothCollection = MakeShared<const FManagedArrayCollection>(InCollectionLod);
-
-			const FCollectionClothConstFacade InClothFacade(InClothCollection);
-				// Check LOD validity
-			if (InClothFacade.GetNumSimPatterns() &&
-				InClothFacade.GetNumRenderPatterns() &&
-				InClothFacade.GetNumSimVertices2D() &&
-				InClothFacade.GetNumSimVertices3D() &&
-				InClothFacade.GetNumSimFaces() &&
-				InClothFacade.GetNumRenderVertices() &&
-				InClothFacade.GetNumRenderFaces())
-			{
-				LastValidLodIndex = LodIndex;
-
-				// Copy input LOD to current output LOD
-				ClothFacade.Initialize(InClothFacade);
-				FClothGeometryTools::CleanupAndCompactMesh(ClothCollection);
-
-				// Add this LOD's materials to the asset
-				const int32 NumLodMaterials = ClothFacade.GetNumRenderPatterns();
+			const FCollectionClothConstFacade InClothFacade(InClothCollections[LodIndex]);
+			check(InClothFacade.HasValidData());
 
-				TArray<FSkeletalMaterial>& Materials = ClothAsset->GetMaterials();
-				Materials.Reserve(Materials.Num() + NumLodMaterials);
+			// Copy input LOD to current output LOD
+			ClothFacade.Initialize(InClothFacade);
 
-				const TConstArrayView<FString> LodRenderMaterialPathName = ClothFacade.GetRenderMaterialPathName();
-				for (int32 LodMaterialIndex = 0; LodMaterialIndex < NumLodMaterials; ++LodMaterialIndex)
-				{
-					const FString& RenderMaterialPathName = LodRenderMaterialPathName[LodMaterialIndex];
-
-					if (UMaterialInterface* const Material = LoadObject<UMaterialInterface>(ClothAsset, *RenderMaterialPathName, nullptr, LOAD_None, nullptr))
-					{
-						Materials.Emplace(Material, true, false, Material->GetFName());
-					}
-					else
-					{
-						Materials.Emplace();
-					}
-				}
+			// Add this LOD's materials to the asset
+			const int32 NumLodMaterials = ClothFacade.GetNumRenderPatterns();
 
-				// Set properties
-				constexpr bool bUpdateExistingProperties = false;
-				Chaos::Softs::FCollectionPropertyMutableFacade(ClothCollection).Append(InClothCollection.ToSharedPtr(), bUpdateExistingProperties);
+			TArray<FSkeletalMaterial>& Materials = ClothAsset->GetMaterials();
+			Materials.Reserve(Materials.Num() + NumLodMaterials);
 
-				// Set physics asset only with LOD 0 at the moment
-				if (LodIndex == 0)
-				{
-					using namespace ::Chaos::Softs;
-					PhysicsAssetPathName = InClothFacade.GetPhysicsAssetPathName();
-				}
-			}
-			if (LodIndex != LastValidLodIndex)
+			const TConstArrayView<FString> LodRenderMaterialPathName = ClothFacade.GetRenderMaterialPathName();
+			for (int32 LodMaterialIndex = 0; LodMaterialIndex < NumLodMaterials; ++LodMaterialIndex)
 			{
-				if (LastValidLodIndex >= 0)
+				const FString& RenderMaterialPathName = LodRenderMaterialPathName[LodMaterialIndex];
+
+				if (UMaterialInterface* const Material = LoadObject<UMaterialInterface>(ClothAsset, *RenderMaterialPathName, nullptr, LOAD_None, nullptr))
 				{
-					ClothFacade.Initialize(FCollectionClothConstFacade(ClothCollections[LastValidLodIndex]));
-
-					FClothDataflowTools::LogAndToastWarning(*this,
-						LOCTEXT("InvalidInputLodNHeadline", "Invalid input LOD."),
-						FText::Format(
-							LOCTEXT("InvalidInputLodNDetails",
-								"Invalid or empty input LOD for LOD {0}.\n"
-								"Using the previous valid LOD {1} instead."),
-							LodIndex,
-							LastValidLodIndex));
+					Materials.Emplace(Material, true, false, Material->GetFName());
 				}
 				else
 				{
-					FClothDataflowTools::LogAndToastWarning(*this,
-						LOCTEXT("InvalidInputLod0Headline", "Invalid input LOD 0."),
-						LOCTEXT("InvalidInputLod0Details",
-							"Invalid or empty input LOD for LOD 0.\n"
-							"LOD 0 cannot be empty in order to construct a valid Cloth Asset."));
-					break; // Empty cloth asset
+					Materials.Emplace();
 				}
 			}
+
+			// Set properties
+			constexpr bool bUpdateExistingProperties = false;
+			Chaos::Softs::FCollectionPropertyMutableFacade(ClothCollection).Append(InClothCollections[LodIndex].ToSharedPtr(), bUpdateExistingProperties);
+
+			// Set physics asset only with LOD 0 at the moment
+			if (LodIndex == 0)
+			{
+				using namespace ::Chaos::Softs;
+				PhysicsAssetPathName = InClothFacade.GetPhysicsAssetPathName();
+			}
 		}
 
 		// Make sure that whatever happens there is always at least one empty LOD to avoid crashing the render data
@@ -257,6 +278,54 @@ TArray<const FManagedArrayCollection*> FChaosClothAssetTerminalNode::GetCollecti
 	return CollectionLods;
 }
 
+TArray<TSharedRef<FManagedArrayCollection>> FChaosClothAssetTerminalNode::GetCleanedCollectionLodValues(Dataflow::FContext& Context) const
+{
+	using namespace UE::Chaos::ClothAsset;
+
+	const TArray<const FManagedArrayCollection*> CollectionLods = GetCollectionLods();
+	TArray<TSharedRef<FManagedArrayCollection>> CollectionLodValues;
+	CollectionLodValues.Reserve(CollectionLods.Num());
+
+	int32 LastValidLodIndex = INDEX_NONE;	
+	for (int32 LodIndex = 0; LodIndex < CollectionLods.Num(); ++LodIndex)
+	{
+		TSharedRef<FManagedArrayCollection>& CollectionLodValue = CollectionLodValues.Emplace_GetRef(MakeShared<FManagedArrayCollection>(GetValue<FManagedArrayCollection>(Context, CollectionLods[LodIndex])));
+
+		FCollectionClothFacade ClothFacade(CollectionLodValue);
+		if (ClothFacade.HasValidData())
+		{
+			FClothGeometryTools::CleanupAndCompactMesh(CollectionLodValue);
+			LastValidLodIndex = LodIndex;
+		}
+		else
+		{
+			if (LastValidLodIndex >= 0)
+			{
+				ClothFacade.DefineSchema();
+				ClothFacade.Initialize(FCollectionClothConstFacade(CollectionLodValues[LastValidLodIndex]));
+				FClothDataflowTools::LogAndToastWarning(*this,
+					LOCTEXT("InvalidInputLodNHeadline", "Invalid input LOD."),
+					FText::Format(
+						LOCTEXT("InvalidInputLodNDetails",
+							"Invalid or empty input LOD for LOD {0}.\n"
+							"Using the previous valid LOD {1} instead."),
+						LodIndex,
+						LastValidLodIndex));
+			}
+			else
+			{
+				FClothDataflowTools::LogAndToastWarning(*this,
+					LOCTEXT("InvalidInputLod0Headline", "Invalid input LOD 0."),
+					LOCTEXT("InvalidInputLod0Details",
+						"Invalid or empty input LOD for LOD 0.\n"
+						"LOD 0 cannot be empty in order to construct a valid Cloth Asset."));
+				break;
+			}
+		}
+	}
+	return CollectionLodValues;
+}
+
 void FChaosClothAssetTerminalNode::Serialize(FArchive& Ar)
 {
 	// because we add pins we need to make sure we restore them when loading
diff --git a/Engine/Plugins/Experimental/ChaosClothAssetEditor/Source/ChaosClothAssetDataflowNodes/Public/ChaosClothAsset/TerminalNode.h b/Engine/Plugins/Experimental/ChaosClothAssetEditor/Source/ChaosClothAssetDataflowNodes/Public/ChaosClothAsset/TerminalNode.h
--- a/Engine/Plugins/Experimental/ChaosClothAssetEditor/Source/ChaosClothAssetDataflowNodes/Public/ChaosClothAsset/TerminalNode.h
+++ b/Engine/Plugins/Experimental/ChaosClothAssetEditor/Source/ChaosClothAssetDataflowNodes/Public/ChaosClothAsset/TerminalNode.h
@@ -54,7 +54,12 @@ private:
 	//~ End FDataflowNode interface
 
 	TArray<const FManagedArrayCollection*> GetCollectionLods() const;
+	TArray<TSharedRef<FManagedArrayCollection>> GetCleanedCollectionLodValues(Dataflow::FContext& Context) const;
 
 	UPROPERTY()
 	mutable TArray<FChaosClothAssetLodTransitionDataCache> LODTransitionDataCache;
+
+	// This is for runtime only--used to determine if only properties need to be updated.
+	mutable bool bClothCollectionChecksumValid = false;
+	mutable uint32 ClothColllectionChecksum = 0;
 };
diff --git a/Engine/Plugins/Experimental/ChaosClothAssetEditor/Source/ChaosClothAssetEditor/Private/ChaosClothAsset/ClothEditorPreviewScene.cpp b/Engine/Plugins/Experimental/ChaosClothAssetEditor/Source/ChaosClothAssetEditor/Private/ChaosClothAsset/ClothEditorPreviewScene.cpp
--- a/Engine/Plugins/Experimental/ChaosClothAssetEditor/Source/ChaosClothAssetEditor/Private/ChaosClothAsset/ClothEditorPreviewScene.cpp
+++ b/Engine/Plugins/Experimental/ChaosClothAssetEditor/Source/ChaosClothAssetEditor/Private/ChaosClothAsset/ClothEditorPreviewScene.cpp
@@ -319,7 +319,16 @@ void FChaosClothPreviewScene::SetClothAsset(UChaosClothAsset* Asset)
 	check(SceneActor);
 	check(ClothComponent);
 
-	ClothComponent->SetClothAsset(Asset);
+	if (ClothComponent->GetClothAsset() == Asset)
+	{
+		// Update the config properties on the component from the asset.
+		ClothComponent->UpdateConfigProperties();
+	}
+	else
+	{
+		ClothComponent->SetClothAsset(Asset);
+	}
+
 	UpdateClothComponentAttachment();
 
 	// Wait for asset to load and update the component bounds
diff --git a/Engine/Source/Runtime/Experimental/Chaos/Private/Chaos/CollectionPropertyFacade.cpp b/Engine/Source/Runtime/Experimental/Chaos/Private/Chaos/CollectionPropertyFacade.cpp
--- a/Engine/Source/Runtime/Experimental/Chaos/Private/Chaos/CollectionPropertyFacade.cpp
+++ b/Engine/Source/Runtime/Experimental/Chaos/Private/Chaos/CollectionPropertyFacade.cpp
@@ -146,6 +146,32 @@ namespace Chaos::Softs
 		}
 	}
 
+	void FCollectionPropertyFacade::UpdateProperties(const TSharedPtr<const FManagedArrayCollection>& InManagedArrayCollection)
+	{
+		FCollectionPropertyConstFacade InPropertyFacade(InManagedArrayCollection);
+		if (InPropertyFacade.IsValid())
+		{
+			const int32 NumInKeys = InPropertyFacade.Num();
+			for (int32 InKeyIndex = 0; InKeyIndex < NumInKeys; ++InKeyIndex)
+			{
+				const FString& PropertyName = InPropertyFacade.GetKey(InKeyIndex);
+				const int32 PropertyIndex = GetKeyIndex(PropertyName);
+				if (PropertyIndex == INDEX_NONE)
+				{
+					continue;
+				}
+				PRAGMA_DISABLE_DEPRECATION_WARNINGS
+				// TODO: GetFlags needs to return an ECollectionPropertyFlags, not an uint8, but the uint8 getter needs to be deprecated first
+				SetFlags(PropertyIndex, (ECollectionPropertyFlags)InPropertyFacade.GetFlags(InKeyIndex));
+				PRAGMA_ENABLE_DEPRECATION_WARNINGS
+				// Setting as FVector3f since that is the underlying type
+				SetLowValue(PropertyIndex, InPropertyFacade.GetLowValue<FVector3f>(InKeyIndex));
+				SetHighValue(PropertyIndex, InPropertyFacade.GetHighValue<FVector3f>(InKeyIndex));
+				SetStringValue(PropertyIndex, InPropertyFacade.GetStringValue(InKeyIndex));
+			}
+		}
+	}
+
 	FCollectionPropertyMutableFacade::FCollectionPropertyMutableFacade(const TSharedPtr<FManagedArrayCollection>& InManagedArrayCollection)
 		: FCollectionPropertyFacade(InManagedArrayCollection, NoInit)
 	{
@@ -257,6 +283,10 @@ namespace Chaos::Softs
 			// Nothing to do
 			return;
 		}
+		if (UpdateFlags == ECollectionPropertyUpdateFlags::UpdateExistingProperties)
+		{
+			return UpdateProperties(InManagedArrayCollection);
+		}
 
 		const bool bAppendNewProperties = (UpdateFlags & ECollectionPropertyUpdateFlags::AppendNewProperties) != ECollectionPropertyUpdateFlags::None;
 		const bool bUpdateExistingProperties = (UpdateFlags & ECollectionPropertyUpdateFlags::UpdateExistingProperties) != ECollectionPropertyUpdateFlags::None;
diff --git a/Engine/Source/Runtime/Experimental/Chaos/Public/Chaos/CollectionPropertyFacade.h b/Engine/Source/Runtime/Experimental/Chaos/Public/Chaos/CollectionPropertyFacade.h
--- a/Engine/Source/Runtime/Experimental/Chaos/Public/Chaos/CollectionPropertyFacade.h
+++ b/Engine/Source/Runtime/Experimental/Chaos/Public/Chaos/CollectionPropertyFacade.h
@@ -345,6 +345,8 @@ PRAGMA_ENABLE_DEPRECATION_WARNINGS
 
 		CHAOS_API void ClearDirtyFlags();
 
+		CHAOS_API void UpdateProperties(const TSharedPtr<const FManagedArrayCollection>& InManagedArrayCollection);
+
 	protected:
 		// No init constructor for FCollectionPropertyMutableFacade
 		CHAOS_API FCollectionPropertyFacade(const TSharedPtr<FManagedArrayCollection>& InManagedArrayCollection, ENoInit);
@@ -374,6 +376,7 @@ PRAGMA_ENABLE_DEPRECATION_WARNINGS
 		inline void SetValue(int32 KeyIndex, const TArrayView<T>& ValueArray, const T& Value);
 
 		CHAOS_API void EnableFlags(int32 KeyIndex, ECollectionPropertyFlags Flags, bool bEnable);
+
 	};
 
 	enum class ECollectionPropertyUpdateFlags : uint8
diff --git a/Engine/Source/Runtime/Experimental/Chaos/Public/GeometryCollection/ManagedArray.h b/Engine/Source/Runtime/Experimental/Chaos/Public/GeometryCollection/ManagedArray.h
--- a/Engine/Source/Runtime/Experimental/Chaos/Public/GeometryCollection/ManagedArray.h
+++ b/Engine/Source/Runtime/Experimental/Chaos/Public/GeometryCollection/ManagedArray.h
@@ -11,6 +11,7 @@
 #include "Chaos/ParticleHandle.h"
 #include "Chaos/BVHParticles.h"
 #include "Math/Vector.h"
+#include "Templates/TypeHash.h"
 
 struct FManagedArrayCollection;
 DEFINE_LOG_CATEGORY_STATIC(UManagedArrayLogging, NoLogging, All);
@@ -580,6 +581,14 @@ public:
 		}
 	}
 
+	/**
+	* This hash is using HashCombineFast and should not be serialized!
+	*/
+	FORCEINLINE uint32 GetTypeHash() const
+	{
+		return GetArrayHash(Array.GetData(), Array.Num());
+	}
+
 	typedef typename TArray<InElementType>::RangedForIteratorType		RangedForIteratorType;
 	typedef typename TArray<InElementType>::RangedForConstIteratorType	RangedForConstIteratorType;
 
@@ -639,6 +648,12 @@ private:
 
 };
 
+template<class T>
+inline uint32 GetTypeHash(const TManagedArrayBase<T>& ManagedArray)
+{
+	return ManagedArray.GetTypeHash();
+}
+
 template <typename T>
 void InitHelper(TArray<T>& Array, const TManagedArrayBase<T>& NewTypedArray, int32 Size)
 {
