Chaos Cloth Asset - Added the no border style to the group edit combo box of the selection node UI customization.

#rnx

diff --git a/Engine/Plugins/Experimental/ChaosClothAssetEditor/Source/ChaosClothAssetDataflowNodes/Private/ChaosClothAsset/SelectionGroupCustomization.cpp b/Engine/Plugins/Experimental/ChaosClothAssetEditor/Source/ChaosClothAssetDataflowNodes/Private/ChaosClothAsset/SelectionGroupCustomization.cpp
--- a/Engine/Plugins/Experimental/ChaosClothAssetEditor/Source/ChaosClothAssetDataflowNodes/Private/ChaosClothAsset/SelectionGroupCustomization.cpp
+++ b/Engine/Plugins/Experimental/ChaosClothAssetEditor/Source/ChaosClothAssetDataflowNodes/Private/ChaosClothAsset/SelectionGroupCustomization.cpp
@@ -40,9 +40,10 @@ namespace UE::Chaos::ClothAsset
 			[
 				SNew(SHorizontalBox)
 				+ SHorizontalBox::Slot()
-				.MaxWidth(160.f)
+				.MaxWidth(145.f)
 				[
 					SAssignNew(ComboButton, SComboButton)
+					.ButtonStyle(FAppStyle::Get(), "NoBorder")
 					.ContentPadding(0)
 					.OnGetMenuContent(this, &FSelectionGroupCustomization::OnGetMenuContent)
 					.ButtonContent()
