CIS - Fix compile warnings in module ChaosClothAssetDataflowNodes.

#rnx

diff --git a/Engine/Plugins/Experimental/ChaosClothAssetEditor/Source/ChaosClothAssetDataflowNodes/Public/ChaosClothAsset/SelectionNode.h b/Engine/Plugins/Experimental/ChaosClothAssetEditor/Source/ChaosClothAssetDataflowNodes/Public/ChaosClothAsset/SelectionNode.h
--- a/Engine/Plugins/Experimental/ChaosClothAssetEditor/Source/ChaosClothAssetDataflowNodes/Public/ChaosClothAsset/SelectionNode.h
+++ b/Engine/Plugins/Experimental/ChaosClothAssetEditor/Source/ChaosClothAssetDataflowNodes/Public/ChaosClothAsset/SelectionNode.h
@@ -42,6 +42,8 @@ struct FChaosClothAssetNodeSelectionGroup
 	FString Name;
 };
 
+PRAGMA_DISABLE_DEPRECATION_WARNINGS  // For EChaosClothAssetSelectionType
+
 /** Integer index set selection node. */
 USTRUCT(Meta = (DataflowCloth))
 struct FChaosClothAssetSelectionNode : public FDataflowNode
@@ -60,10 +62,8 @@ public:
 
 	/** The type of element the selection refers to */
 	UE_DEPRECATED(5.4, "Use Group instead")
-PRAGMA_DISABLE_DEPRECATION_WARNINGS
 	UPROPERTY()
 	EChaosClothAssetSelectionType Type_DEPRECATED = EChaosClothAssetSelectionType::Deprecated;
-PRAGMA_ENABLE_DEPRECATION_WARNINGS
 
 	/** The type of element the selection refers to */
 	UPROPERTY(EditAnywhere, Category = "Selection")
@@ -88,3 +88,5 @@ private:
 	TArray<FName> CachedCollectionGroupNames;
 };
 
+PRAGMA_ENABLE_DEPRECATION_WARNINGS  // For EChaosClothAssetSelectionType
+
