Cloth Editor Toolkit: deactivate ClothEditorMode in response to OnClose(), not OnRequestClose(), since editor shutdown can be cancelled after OnRequestClose() is called

diff --git a/Engine/Plugins/Experimental/ChaosClothAssetEditor/Source/ChaosClothAssetEditor/Private/ChaosClothAsset/ClothEditorToolkit.cpp b/Engine/Plugins/Experimental/ChaosClothAssetEditor/Source/ChaosClothAssetEditor/Private/ChaosClothAsset/ClothEditorToolkit.cpp
--- a/Engine/Plugins/Experimental/ChaosClothAssetEditor/Source/ChaosClothAssetEditor/Private/ChaosClothAsset/ClothEditorToolkit.cpp
+++ b/Engine/Plugins/Experimental/ChaosClothAssetEditor/Source/ChaosClothAssetEditor/Private/ChaosClothAsset/ClothEditorToolkit.cpp
@@ -395,13 +395,18 @@ bool FChaosClothAssetEditorToolkit::OnRequestClose(EAssetEditorCloseReason InClo
 		return true;
 	}
 
+	return FAssetEditorToolkit::OnRequestClose(InCloseReason);
+}
+
+void FChaosClothAssetEditorToolkit::OnClose()
+{
 	// Give any active modes a chance to shutdown while the toolkit host is still alive
 	// This is super important to do, otherwise currently opened tabs won't be marked as "closed".
 	// This results in tabs not being properly recycled upon reopening the editor and tab
 	// duplication for each opening event.
 	GetEditorModeManager().ActivateDefaultMode();
 
-	return FAssetEditorToolkit::OnRequestClose(InCloseReason);
+	FBaseCharacterFXEditorToolkit::OnClose();
 }
 
 void FChaosClothAssetEditorToolkit::PostInitAssetEditor()
@@ -645,6 +650,8 @@ void FChaosClothAssetEditorToolkit::UnregisterTabSpawners(const TSharedRef<FTabM
 	InTabManager->UnregisterTabSpawner(ClothPreviewTabID);
 	InTabManager->UnregisterTabSpawner(ViewportTabID);
 	InTabManager->UnregisterTabSpawner(DetailsTabID);
+	InTabManager->UnregisterTabSpawner(OutlinerTabID);
+	InTabManager->UnregisterTabSpawner(PreviewSceneDetailsTabID);
 	InTabManager->UnregisterTabSpawner(GraphCanvasTabId);
 	InTabManager->UnregisterTabSpawner(NodeDetailsTabId);
 }
diff --git a/Engine/Plugins/Experimental/ChaosClothAssetEditor/Source/ChaosClothAssetEditor/Private/ChaosClothAsset/ClothEditorToolkit.h b/Engine/Plugins/Experimental/ChaosClothAssetEditor/Source/ChaosClothAssetEditor/Private/ChaosClothAsset/ClothEditorToolkit.h
--- a/Engine/Plugins/Experimental/ChaosClothAssetEditor/Source/ChaosClothAssetEditor/Private/ChaosClothAsset/ClothEditorToolkit.h
+++ b/Engine/Plugins/Experimental/ChaosClothAssetEditor/Source/ChaosClothAssetEditor/Private/ChaosClothAsset/ClothEditorToolkit.h
@@ -82,6 +82,8 @@ private:
 	virtual void AddViewportOverlayWidget(TSharedRef<SWidget> InViewportOverlayWidget) override;
 	virtual void RemoveViewportOverlayWidget(TSharedRef<SWidget> InViewportOverlayWidget) override;
 	virtual bool OnRequestClose(EAssetEditorCloseReason InCloseReason) override;
+	virtual void OnClose() override;
+
 	virtual void PostInitAssetEditor() override;
 	virtual void InitToolMenuContext(FToolMenuContext& MenuContext) override;
 	virtual void GetSaveableObjects(TArray<UObject*>& OutObjects) const override;
