Chaos Cloth Asset - Fixed crash in terminal node when the imported asset has no geometry.

#rnx

diff --git a/Engine/Plugins/Experimental/ChaosClothAssetEditor/Source/ChaosClothAssetDataflowNodes/Private/ChaosClothAsset/TerminalNode.cpp b/Engine/Plugins/Experimental/ChaosClothAssetEditor/Source/ChaosClothAssetDataflowNodes/Private/ChaosClothAsset/TerminalNode.cpp
--- a/Engine/Plugins/Experimental/ChaosClothAssetEditor/Source/ChaosClothAssetDataflowNodes/Private/ChaosClothAsset/TerminalNode.cpp
+++ b/Engine/Plugins/Experimental/ChaosClothAssetEditor/Source/ChaosClothAssetDataflowNodes/Private/ChaosClothAsset/TerminalNode.cpp
@@ -127,7 +127,15 @@ void FChaosClothAssetTerminalNode::SetAssetValue(TObjectPtr<UObject> Asset, Data
 			ClothFacade.DefineSchema();
 
 			const FCollectionClothConstFacade InClothFacade(InClothCollections[LodIndex]);
-			check(InClothFacade.HasValidData());
+			if (!InClothFacade.HasValidData())
+			{
+				FClothDataflowTools::LogAndToastWarning(*this, LOCTEXT("InvalidLODHeadline", "Invalid LOD."),
+					FText::Format(
+						LOCTEXT("InvalidLODDetails", "LOD {0} has no valid data."),
+						LodIndex));
+
+				continue;
+			}
 
 			// Copy input LOD to current output LOD
 			ClothFacade.Initialize(InClothFacade);
