Cloth weight map paint tool: change from triangle painting to vertex painting

diff --git a/Engine/Plugins/Experimental/ChaosClothAssetEditor/Source/ChaosClothAssetEditorTools/Private/ChaosClothAsset/ClothWeightMapPaintBrushOps.cpp b/Engine/Plugins/Experimental/ChaosClothAssetEditor/Source/ChaosClothAssetEditorTools/Private/ChaosClothAsset/ClothWeightMapPaintBrushOps.cpp
--- a/Engine/Plugins/Experimental/ChaosClothAssetEditor/Source/ChaosClothAssetEditorTools/Private/ChaosClothAsset/ClothWeightMapPaintBrushOps.cpp
+++ b/Engine/Plugins/Experimental/ChaosClothAssetEditor/Source/ChaosClothAssetEditorTools/Private/ChaosClothAsset/ClothWeightMapPaintBrushOps.cpp
@@ -14,16 +14,24 @@ void FWeightMapPaintBrushOp::ApplyStampByVertices(
 	const TArray<int32>& Vertices,
 	TArray<double>& NewAttributesOut)
 {
-	const FVector3d& StampPos = Stamp.LocalFrame.Origin;
-
 	UWeightMapPaintBrushOpProps* const Props = GetPropertySetAs<UWeightMapPaintBrushOpProps>();
 	const double TargetValue = (double)Props->GetAttribute();
 
 	const int32 NumVertices = Vertices.Num();
 	ParallelFor(NumVertices, [&](int32 k)
 	{
-		const int32 VertIdx = Vertices[k];
-		const FVector3d OrigPos = Mesh->GetVertex(VertIdx);
+		if (bApplyRadiusLimit)
+		{
+			const FVector3d& StampPos = Stamp.LocalFrame.Origin;
+			const int32 VertIdx = Vertices[k];
+			const FVector3d VertexPos = Mesh->GetVertex(VertIdx);
+			const double DistanceSquared = (VertexPos - StampPos).SquaredLength();
+			if (DistanceSquared >= Stamp.Radius * Stamp.Radius)
+			{
+				return;
+			}
+		}
+
 		const double ExistingValue = NewAttributesOut[k];
 		const double ValueDiff = TargetValue - ExistingValue;
 		NewAttributesOut[k] = FMath::Clamp(ExistingValue + Stamp.Power * ValueDiff, 0.0, 1.0);
@@ -31,6 +39,42 @@ void FWeightMapPaintBrushOp::ApplyStampByVertices(
 }
 
 
+//
+// Erase Brush
+//
+
+void FWeightMapEraseBrushOp::ApplyStampByVertices(
+	const FDynamicMesh3* Mesh,
+	const FSculptBrushStamp& Stamp,
+	const TArray<int32>& Vertices,
+	TArray<double>& NewAttributesOut)
+{
+	UWeightMapEraseBrushOpProps* Props = GetPropertySetAs<UWeightMapEraseBrushOpProps>();
+	const double EraseAttribute = (double)Props->GetAttribute();
+
+	// TODO: Add something here to get the old value so we can subtract (clamped) the AttributeValue from it.
+	// TODO: Handle the stamp's properties for fall off, etc..
+
+	check(NewAttributesOut.Num() == Vertices.Num());
+
+	for (int32 k = 0; k < Vertices.Num(); ++k)
+	{
+		if (bApplyRadiusLimit)
+		{
+			const FVector3d& StampPos = Stamp.LocalFrame.Origin;
+			const int32 VertIdx = Vertices[k];
+			const FVector3d VertexPos = Mesh->GetVertex(VertIdx);
+			const double DistanceSquared = (VertexPos - StampPos).SquaredLength();
+			if (DistanceSquared >= Stamp.Radius * Stamp.Radius)
+			{
+				continue;
+			}
+		}
+		NewAttributesOut[k] = EraseAttribute;
+	}
+}
+
+
 //
 // Smooth Brush
 //
@@ -86,8 +130,18 @@ void FWeightMapSmoothBrushOp::ApplyStampByVertices(
 	// Blend vertex value with its average one-ring value
 	for (int32 BufferIndex = 0; BufferIndex < NumVertices; ++BufferIndex)
 	{
-		const int32 VertexIndex = Vertices[BufferIndex];
 		const double Diff = OneRingAverages[BufferIndex] - VertexWeightValues[BufferIndex];
+		if (bApplyRadiusLimit)
+		{
+			const FVector3d& StampPos = Stamp.LocalFrame.Origin;
+			const int32 VertexIndex = Vertices[BufferIndex];
+			const FVector3d VertexPos = Mesh->GetVertex(VertexIndex);
+			const double DistanceSquared = (VertexPos - StampPos).SquaredLength();
+			if (DistanceSquared >= Stamp.Radius * Stamp.Radius)
+			{
+				continue;
+			}
+		}
 		VertexWeightValues[BufferIndex] += Stamp.Power * Diff;
 	}
 }
diff --git a/Engine/Plugins/Experimental/ChaosClothAssetEditor/Source/ChaosClothAssetEditorTools/Private/ChaosClothAsset/ClothWeightMapPaintBrushOps.h b/Engine/Plugins/Experimental/ChaosClothAssetEditor/Source/ChaosClothAssetEditorTools/Private/ChaosClothAsset/ClothWeightMapPaintBrushOps.h
--- a/Engine/Plugins/Experimental/ChaosClothAssetEditor/Source/ChaosClothAssetEditorTools/Private/ChaosClothAsset/ClothWeightMapPaintBrushOps.h
+++ b/Engine/Plugins/Experimental/ChaosClothAssetEditor/Source/ChaosClothAssetEditorTools/Private/ChaosClothAsset/ClothWeightMapPaintBrushOps.h
@@ -29,6 +29,9 @@ public:
 		const TArray<int32>& Vertices,
 		TArray<double>& NewAttributesOut) = 0;
 
+	/** If bApplyRadiusLimit is enabled, only vertices within the stamp radius are considered */
+	bool bApplyRadiusLimit = true;
+
 };
 
 
@@ -42,11 +45,14 @@ class CHAOSCLOTHASSETEDITORTOOLS_API UWeightMapEraseBrushOpProps : public UMeshS
 {
 	GENERATED_BODY()
 public:
+	// TODO: This AttributeValue is not currenly connected to the UI (jira UE-198413)
 	/** Attribute to set as Erased value */
 	UPROPERTY(EditAnywhere, Category = EraseBrush, meta = (DisplayName = "Erase Attribute", UIMin = 0))
 	double AttributeValue = 0;
 
 	virtual double GetAttribute() { return AttributeValue; }
+
+	virtual float GetFalloff() override { return 0.0f; }
 };
 
 
@@ -59,23 +65,8 @@ public:
 		const FDynamicMesh3* Mesh,
 		const FSculptBrushStamp& Stamp,
 		const TArray<int32>& Vertices,
-		TArray<double>& NewAttributesOut) override
-	{
-		const FVector3d& StampPos = Stamp.LocalFrame.Origin;
-
-		UWeightMapEraseBrushOpProps* Props = GetPropertySetAs<UWeightMapEraseBrushOpProps>();
-		double EraseAttribute = (double)Props->GetAttribute();
-
-		// TODO: Add something here to get the old value so we can subtract (clamped) the AttributeValue from it.
-
-		// TODO: Handle the stamp's properties for fall off, etc..
+		TArray<double>& NewAttributesOut) override;
 
-		int32 NumVertices = Vertices.Num();
-		for (int32 k = 0; k < NumVertices; ++k)
-		{
-			NewAttributesOut[k] = 0.0;
-		}
-	}
 };
 
 
@@ -102,6 +93,8 @@ public:
 
 	virtual float GetStrength() override { return Strength; }
 	virtual void SetStrength(float NewStrength) override { Strength = NewStrength;  }
+
+	virtual float GetFalloff() override { return 0.0f; }
 };
 
 
@@ -139,6 +132,7 @@ public:
 	virtual float GetStrength() override { return Strength; }
 	virtual void SetStrength(float NewStrength) override { Strength = NewStrength; }
 
+	virtual float GetFalloff() override { return 0.0f; }
 };
 
 
diff --git a/Engine/Plugins/Experimental/ChaosClothAssetEditor/Source/ChaosClothAssetEditorTools/Private/ChaosClothAsset/ClothWeightMapPaintTool.cpp b/Engine/Plugins/Experimental/ChaosClothAssetEditor/Source/ChaosClothAssetEditorTools/Private/ChaosClothAsset/ClothWeightMapPaintTool.cpp
--- a/Engine/Plugins/Experimental/ChaosClothAssetEditor/Source/ChaosClothAssetEditorTools/Private/ChaosClothAsset/ClothWeightMapPaintTool.cpp
+++ b/Engine/Plugins/Experimental/ChaosClothAssetEditor/Source/ChaosClothAssetEditorTools/Private/ChaosClothAsset/ClothWeightMapPaintTool.cpp
@@ -144,6 +144,8 @@ void UClothEditorWeightMapPaintTool::Setup()
 
 	FDynamicMesh3* Mesh = GetSculptMesh();
 	Mesh->EnableVertexColors(FVector3f::One());
+	Mesh->Attributes()->EnablePrimaryColors();
+	Mesh->Attributes()->PrimaryColors()->CreateFromPredicate([](int ParentVID, int TriIDA, int TriIDB) {return true; }, 0.f);
 	FAxisAlignedBox3d Bounds = Mesh->GetBounds(true);
 
 	TFuture<void> PrecomputeFuture = Async(UE::Chaos::ClothAsset::Private::WeightPaintToolAsyncExecTarget, [&]()
@@ -251,6 +253,7 @@ void UClothEditorWeightMapPaintTool::Setup()
 		EraseBrushOpProperties);
 
 	AddToolPropertySource(UMeshSculptToolBase::ViewProperties);
+	SetToolPropertySourceEnabled(UMeshSculptToolBase::ViewProperties, true);
 
 	AddToolPropertySource(UMeshSculptToolBase::GizmoProperties);
 	SetToolPropertySourceEnabled(UMeshSculptToolBase::GizmoProperties, false);
@@ -282,29 +285,11 @@ void UClothEditorWeightMapPaintTool::Setup()
 		ProcessFunc(*GetSculptMesh());
 	});
 
-	// force colors update... ?
-	DynamicMeshComponent->SetTriangleColorFunction([this](const FDynamicMesh3* Mesh, int TriangleID)
-	{
-		if (ActiveWeightMap)
-		{
-			FIndex3i Vertices = Mesh->GetTriangle(TriangleID);
-			FVector3f WeightPerVertex;
-			ActiveWeightMap->GetValue(Vertices[0], &WeightPerVertex[0]);
-			ActiveWeightMap->GetValue(Vertices[1], &WeightPerVertex[1]);
-			ActiveWeightMap->GetValue(Vertices[2], &WeightPerVertex[2]);
-			return GetColorForWeightValue(WeightPerVertex[0] / 3.0 + WeightPerVertex[1] / 3.0 + WeightPerVertex[2] / 3.0);
-		}
-		else
-		{
-			return LinearColors::Black3b();
-		}
-	});
-
 	// disable view properties
 	SetViewPropertiesEnabled(false);
 	UpdateMaterialMode(EMeshEditingMaterialModes::VertexColor);
 	UpdateWireframeVisibility(false);
-	UpdateFlatShadingSetting(true);
+	UpdateFlatShadingSetting(false);
 
 	// configure panels
 	UpdateSubToolType(FilterProperties->SubToolType);
@@ -391,6 +376,7 @@ void UClothEditorWeightMapPaintTool::Setup()
 
 
 	// update colors
+	UpdateVertexColorOverlay();
 	DynamicMeshComponent->FastNotifyVertexAttributesUpdated(EMeshRenderAttributeFlags::VertexColors);
 	GetToolManager()->PostInvalidation();
 
@@ -565,6 +551,9 @@ void UClothEditorWeightMapPaintTool::OnEndStroke()
 
 	GetActiveBrushOp()->EndStroke(GetSculptMesh(), LastStamp, VertexROI);
 
+	UpdateVertexColorOverlay(&TriangleROI);
+	DynamicMeshComponent->FastNotifyVertexAttributesUpdated(EMeshRenderAttributeFlags::VertexColors);
+
 	// close change record
 	EndChange();
 }
@@ -728,6 +717,15 @@ bool UClothEditorWeightMapPaintTool::ApplyStamp()
 	// yuck
 	FMeshVertexWeightMapEditBrushOp* WeightBrushOp = (FMeshVertexWeightMapEditBrushOp*)UseBrushOp.Get();
 
+	if (FilterProperties->SubToolType == EClothEditorWeightMapPaintInteractionType::Brush)
+	{
+		WeightBrushOp->bApplyRadiusLimit = true;
+	}
+	else
+	{ 
+		WeightBrushOp->bApplyRadiusLimit = false;
+	}
+
 	FDynamicMesh3* Mesh = GetSculptMesh();
 	WeightBrushOp->ApplyStampByVertices(Mesh, CurrentStamp, VertexROI, ROIWeightValueBuffer);
 
@@ -1024,8 +1022,8 @@ void UClothEditorWeightMapPaintTool::ComputeGradient()
 
 	if (bHaveDynamicMeshToWeightConversion)
 	{
-	for (int32 vid : TempROIBuffer)
-	{
+		for (int32 vid : TempROIBuffer)
+		{
 			for (const int32 Idx : WeightToDynamicMesh[DynamicMeshToWeight[vid]])
 			{
 				ActiveWeightEditChangeTracker->SaveVertexOneRingTriangles(Idx, true);
@@ -1036,8 +1034,8 @@ void UClothEditorWeightMapPaintTool::ComputeGradient()
 	{
 		for (int32 vid : TempROIBuffer)
 		{
-		ActiveWeightEditChangeTracker->SaveVertexOneRingTriangles(vid, true);
-	}
+			ActiveWeightEditChangeTracker->SaveVertexOneRingTriangles(vid, true);
+		}
 	}
 
 
@@ -1085,7 +1083,8 @@ void UClothEditorWeightMapPaintTool::ComputeGradient()
 	}
 	}
 
-
+	// update colors
+	UpdateVertexColorOverlay();
 	DynamicMeshComponent->FastNotifyVertexAttributesUpdated(EMeshRenderAttributeFlags::VertexColors);
 	GetToolManager()->PostInvalidation();
 	EndChange();
@@ -1152,18 +1151,19 @@ void UClothEditorWeightMapPaintTool::SetVerticesToWeightMap(const TSet<int32>& V
 	}
 	else
 	{
-	for (int32 vid : TempROIBuffer)
-	{
-		ActiveWeightEditChangeTracker->SaveVertexOneRingTriangles(vid, true);
-	}
-	for (int32 vid : TempROIBuffer)
-	{		
-		ActiveWeightMap->SetValue(vid, &WeightValue);
-	}
+		for (int32 vid : TempROIBuffer)
+		{
+			ActiveWeightEditChangeTracker->SaveVertexOneRingTriangles(vid, true);
+		}
+		for (int32 vid : TempROIBuffer)
+		{		
+			ActiveWeightMap->SetValue(vid, &WeightValue);
+		}
 	}
 
-
-	DynamicMeshComponent->FastNotifyVertexAttributesUpdated(false, true, false);
+	// update colors
+	UpdateVertexColorOverlay();
+	DynamicMeshComponent->FastNotifyVertexAttributesUpdated(EMeshRenderAttributeFlags::VertexColors);
 	GetToolManager()->PostInvalidation();
 	
 
@@ -1510,6 +1510,7 @@ void UClothEditorWeightMapPaintTool::OnTick(float DeltaTime)
 			if (bWeightsModified)
 			{
 				SCOPE_CYCLE_COUNTER(WeightMapPaintTool_Tick_UpdateMeshBlock);
+				UpdateVertexColorOverlay(&TriangleROI);
 				DynamicMeshComponent->FastNotifyTriangleVerticesUpdated(TriangleROI, EMeshRenderAttributeFlags::VertexColors);
 				GetToolManager()->PostInvalidation();
 			}
@@ -1561,8 +1562,8 @@ void UClothEditorWeightMapPaintTool::FloodFillCurrentWeightAction()
 
 	if (bHaveDynamicMeshToWeightConversion)
 	{
-	for (int32 vid : TempROIBuffer)
-	{
+		for (int32 vid : TempROIBuffer)
+		{
 			for (const int32 Idx : WeightToDynamicMesh[DynamicMeshToWeight[vid]])
 			{
 				ActiveWeightEditChangeTracker->SaveVertexOneRingTriangles(Idx, true);
@@ -1574,14 +1575,16 @@ void UClothEditorWeightMapPaintTool::FloodFillCurrentWeightAction()
 	{
 		for (int32 vid : TempROIBuffer)
 		{
-		ActiveWeightEditChangeTracker->SaveVertexOneRingTriangles(vid, true);
-	}
-	for (int32 vid : TempROIBuffer)
-	{
-		ActiveWeightMap->SetValue(vid, &SetWeightValue);
-	}
+			ActiveWeightEditChangeTracker->SaveVertexOneRingTriangles(vid, true);
+		}
+		for (int32 vid : TempROIBuffer)
+		{
+			ActiveWeightMap->SetValue(vid, &SetWeightValue);
+		}
 	}
 
+	// update colors
+	UpdateVertexColorOverlay();
 	DynamicMeshComponent->FastNotifyVertexAttributesUpdated(EMeshRenderAttributeFlags::VertexColors);
 	GetToolManager()->PostInvalidation();
 	EndChange();
@@ -1618,16 +1621,18 @@ void UClothEditorWeightMapPaintTool::ClearAllWeightsAction()
 	}
 	else
 	{
-	for (int32 vid : TempROIBuffer)
-	{
-		ActiveWeightEditChangeTracker->SaveVertexOneRingTriangles(vid, true);
-	}
-	for (int32 vid : TempROIBuffer)
-	{
-		ActiveWeightMap->SetValue(vid, &SetWeightValue);
-	}
+		for (int32 vid : TempROIBuffer)
+		{
+			ActiveWeightEditChangeTracker->SaveVertexOneRingTriangles(vid, true);
+		}
+		for (int32 vid : TempROIBuffer)
+		{
+			ActiveWeightMap->SetValue(vid, &SetWeightValue);
+		}
 	}
 
+	// update colors
+	UpdateVertexColorOverlay();
 	DynamicMeshComponent->FastNotifyVertexAttributesUpdated(EMeshRenderAttributeFlags::VertexColors);
 	GetToolManager()->PostInvalidation();
 	EndChange();
@@ -1930,6 +1935,45 @@ void UClothEditorWeightMapPaintTool::ApplyAction(EClothEditorWeightMapPaintToolA
 }
 
 
+void UClothEditorWeightMapPaintTool::UpdateVertexColorOverlay(const TSet<int>* TrianglesToUpdate)
+{
+	FDynamicMesh3* const Mesh = GetSculptMesh();
+	check(Mesh->HasAttributes());
+	check(Mesh->Attributes()->PrimaryColors());
+	check(ActiveWeightMap);
+
+	FDynamicMeshColorOverlay* const ColorOverlay = Mesh->Attributes()->PrimaryColors();
+
+	auto SetColorsFromWeights = [&](int TriangleID)
+	{
+		const FIndex3i Tri = Mesh->GetTriangle(TriangleID);
+		const FIndex3i ColorElementTri = ColorOverlay->GetTriangle(TriangleID);
+
+		for (int TriVertIndex = 0; TriVertIndex < 3; ++TriVertIndex)
+		{
+			float VertexWeight;
+			ActiveWeightMap->GetValue(Tri[TriVertIndex], &VertexWeight);
+
+			const FVector4f NewColor = FVector4f(VertexWeight, VertexWeight, VertexWeight, 1.0f);
+			ColorOverlay->SetElement(ColorElementTri[TriVertIndex], NewColor);
+		}
+	};
+
+	if (TrianglesToUpdate)
+	{
+		for (const int TriangleID : *TrianglesToUpdate)
+		{
+			SetColorsFromWeights(TriangleID);
+		}
+	}
+	else
+	{
+		for (const int TriangleID : Mesh->TriangleIndicesItr())
+		{
+			SetColorsFromWeights(TriangleID);
+		}
+	}
+}
 
 
 #undef LOCTEXT_NAMESPACE
diff --git a/Engine/Plugins/Experimental/ChaosClothAssetEditor/Source/ChaosClothAssetEditorTools/Public/ChaosClothAsset/ClothWeightMapPaintTool.h b/Engine/Plugins/Experimental/ChaosClothAssetEditor/Source/ChaosClothAssetEditorTools/Public/ChaosClothAsset/ClothWeightMapPaintTool.h
--- a/Engine/Plugins/Experimental/ChaosClothAssetEditor/Source/ChaosClothAssetEditorTools/Public/ChaosClothAsset/ClothWeightMapPaintTool.h
+++ b/Engine/Plugins/Experimental/ChaosClothAssetEditor/Source/ChaosClothAssetEditorTools/Public/ChaosClothAsset/ClothWeightMapPaintTool.h
@@ -480,6 +480,8 @@ protected:
 	FChaosClothAssetAddWeightMapNode* WeightMapNodeToUpdate = nullptr;
 
 	void UpdateSelectedNode();
+
+	void UpdateVertexColorOverlay(const TSet<int>* TrianglesToUpdate = nullptr);
 };
 
 
