Cloth Editor: refresh the Construction view after the Dataflow graph executes. Also rename a member variable from the too-generic "DataflowNode" to "SelectedDataflowNode"

diff --git a/Engine/Plugins/Experimental/ChaosClothAssetEditor/Source/ChaosClothAssetEditor/Private/ChaosClothAsset/ClothEditorToolkit.cpp b/Engine/Plugins/Experimental/ChaosClothAssetEditor/Source/ChaosClothAssetEditor/Private/ChaosClothAsset/ClothEditorToolkit.cpp
--- a/Engine/Plugins/Experimental/ChaosClothAssetEditor/Source/ChaosClothAssetEditor/Private/ChaosClothAsset/ClothEditorToolkit.cpp
+++ b/Engine/Plugins/Experimental/ChaosClothAssetEditor/Source/ChaosClothAssetEditor/Private/ChaosClothAsset/ClothEditorToolkit.cpp
@@ -206,12 +206,12 @@ FChaosClothAssetEditorToolkit::FChaosClothAssetEditorToolkit(UAssetEditor* InOwn
 
 FChaosClothAssetEditorToolkit::~FChaosClothAssetEditorToolkit()
 {
-	if (DataflowNode && OnNodeInvalidatedDelegateHandle.IsValid())
+	if (SelectedDataflowNode && OnNodeInvalidatedDelegateHandle.IsValid())
 	{
-		DataflowNode->GetOnNodeInvalidatedDelegate().Remove(OnNodeInvalidatedDelegateHandle);
-		DataflowNode->OnDeselected();
+		SelectedDataflowNode->GetOnNodeInvalidatedDelegate().Remove(OnNodeInvalidatedDelegateHandle);
+		SelectedDataflowNode->OnDeselected();
 	}
-	DataflowNode.Reset();
+	SelectedDataflowNode.Reset();
 
 	if (ClothPreviewViewportClient)
 	{
@@ -259,6 +259,14 @@ void FChaosClothAssetEditorToolkit::Tick(float DeltaTime)
 		if (OldTimestamp.Value < LastDataflowNodeTimestamp.Value)
 		{
 			OnClothAssetChanged();
+
+			// Refresh the construction viewport
+			if (SelectedDataflowNode)
+			{
+				TSharedPtr<FManagedArrayCollection> Collection = GetClothCollectionIfPossible(SelectedDataflowNode, DataflowContext);
+				UChaosClothAssetEditorMode* const ClothMode = CastChecked<UChaosClothAssetEditorMode>(EditorModeManager->GetActiveScriptableMode(UChaosClothAssetEditorMode::EM_ChaosClothAssetEditorModeId));
+				ClothMode->SetSelectedClothCollection(Collection);
+			}
 		}
 	}
 
@@ -814,6 +822,14 @@ void FChaosClothAssetEditorToolkit::EvaluateNode(FDataflowNode* Node, FDataflowO
 		if (OldTimestamp.Value < LastDataflowNodeTimestamp.Value)
 		{
 			OnClothAssetChanged();
+
+			// Refresh the construction viewport
+			if (SelectedDataflowNode)
+			{
+				TSharedPtr<FManagedArrayCollection> Collection = GetClothCollectionIfPossible(SelectedDataflowNode, DataflowContext);
+				UChaosClothAssetEditorMode* const ClothMode = CastChecked<UChaosClothAssetEditorMode>(EditorModeManager->GetActiveScriptableMode(UChaosClothAssetEditorMode::EM_ChaosClothAssetEditorModeId));
+				ClothMode->SetSelectedClothCollection(Collection);
+			}
 		}
 	}
 };
@@ -939,36 +955,36 @@ void FChaosClothAssetEditorToolkit::OnNodeTitleCommitted(const FText& InNewText,
 	FDataflowEditorCommands::OnNodeTitleCommitted(InNewText, InCommitType, GraphNode);
 }
 
-void FChaosClothAssetEditorToolkit::OnNodeSelectionChanged(const TSet<UObject*>& NewSelection)
+
+TSharedPtr<FManagedArrayCollection> FChaosClothAssetEditorToolkit::GetClothCollectionIfPossible(const TSharedPtr<FDataflowNode> InDataflowNode, const TSharedPtr<Dataflow::FEngineContext> Context)
 {
-	auto GetClothCollectionIfPossible = [](const TSharedPtr<FDataflowNode> InDataflowNode, const TSharedPtr<Dataflow::FEngineContext> Context) -> TSharedPtr<FManagedArrayCollection>
+	if (Context.IsValid())
 	{
-		if (Context.IsValid())
+		for (const FDataflowOutput* const Output : InDataflowNode->GetOutputs())
 		{
-			for (const FDataflowOutput* const Output : InDataflowNode->GetOutputs())
+			if (Output->GetType() == FName("FManagedArrayCollection"))
 			{
-				if (Output->GetType() == FName("FManagedArrayCollection"))
-				{
-					const FManagedArrayCollection DefaultValue;
-					TSharedRef<FManagedArrayCollection> Collection = MakeShared<FManagedArrayCollection>(Output->GetValue<FManagedArrayCollection>(*Context, DefaultValue));
+				const FManagedArrayCollection DefaultValue;
+				TSharedRef<FManagedArrayCollection> Collection = MakeShared<FManagedArrayCollection>(Output->GetValue<FManagedArrayCollection>(*Context, DefaultValue));
 
-					// see if the output collection is a ClothCollection
-					const UE::Chaos::ClothAsset::FCollectionClothConstFacade ClothFacade(Collection);
-					if (ClothFacade.IsValid())
-					{
-						return Collection;
-					}
-
-					// The cloth collection schema must be applied to prevent the dynamic mesh conversion and tools from crashing trying to access invalid facades
-					break;
+				// see if the output collection is a ClothCollection
+				const UE::Chaos::ClothAsset::FCollectionClothConstFacade ClothFacade(Collection);
+				if (ClothFacade.IsValid())
+				{
+					return Collection;
 				}
+
+				// The cloth collection schema must be applied to prevent the dynamic mesh conversion and tools from crashing trying to access invalid facades
+				break;
 			}
 		}
+	}
 
-		return TSharedPtr<FManagedArrayCollection>();
-	};
-
+	return TSharedPtr<FManagedArrayCollection>();
+}
 
+void FChaosClothAssetEditorToolkit::OnNodeSelectionChanged(const TSet<UObject*>& NewSelection)
+{
 	TSharedPtr<FManagedArrayCollection> Collection = nullptr;
 
 	// Get any selected node with a ClothCollection output
@@ -981,14 +997,14 @@ void FChaosClothAssetEditorToolkit::OnNodeSelectionChanged(const TSet<UObject*>&
 
 		if (!NewSelection.Num())
 		{
-			if (DataflowNode)
+			if (SelectedDataflowNode)
 			{
 				if (OnNodeInvalidatedDelegateHandle.IsValid())
 				{
-					DataflowNode->GetOnNodeInvalidatedDelegate().Remove(OnNodeInvalidatedDelegateHandle);
+					SelectedDataflowNode->GetOnNodeInvalidatedDelegate().Remove(OnNodeInvalidatedDelegateHandle);
 				}
-				DataflowNode->OnDeselected();
-				DataflowNode.Reset();
+				SelectedDataflowNode->OnDeselected();
+				SelectedDataflowNode.Reset();
 			}
 		}
 		else
@@ -999,31 +1015,31 @@ void FChaosClothAssetEditorToolkit::OnNodeSelectionChanged(const TSet<UObject*>&
 				{
 					Dataflow->RenderTargets.Add(Node);
 
-					if (DataflowNode != Node->GetDataflowNode())
+					if (SelectedDataflowNode != Node->GetDataflowNode())
 					{
-						if (DataflowNode)
+						if (SelectedDataflowNode)
 						{
 							if (OnNodeInvalidatedDelegateHandle.IsValid())
 							{
-								DataflowNode->GetOnNodeInvalidatedDelegate().Remove(OnNodeInvalidatedDelegateHandle);
+								SelectedDataflowNode->GetOnNodeInvalidatedDelegate().Remove(OnNodeInvalidatedDelegateHandle);
 							}
-							DataflowNode->OnDeselected();
+							SelectedDataflowNode->OnDeselected();
 						}
-						DataflowNode = Node->GetDataflowNode();
+						SelectedDataflowNode = Node->GetDataflowNode();
 
-						if (DataflowNode)
+						if (SelectedDataflowNode)
 						{
-							Collection = GetClothCollectionIfPossible(DataflowNode, DataflowContext);
-							DataflowNode->OnSelected(*DataflowContext);
+							Collection = GetClothCollectionIfPossible(SelectedDataflowNode, DataflowContext);
+							SelectedDataflowNode->OnSelected(*DataflowContext);
 
 							// Set a callback to re-evaluate the node if it is invalidated
-							OnNodeInvalidatedDelegateHandle = DataflowNode->GetOnNodeInvalidatedDelegate().AddLambda(
-								[this, &GetClothCollectionIfPossible](FDataflowNode* InDataflowNode)
+							OnNodeInvalidatedDelegateHandle = SelectedDataflowNode->GetOnNodeInvalidatedDelegate().AddLambda(
+								[this](FDataflowNode* InDataflowNode)
 								{
-									if (DataflowNode.Get() == InDataflowNode)
+									if (SelectedDataflowNode.Get() == InDataflowNode)
 									{
-										GetClothCollectionIfPossible(DataflowNode, DataflowContext);
-										DataflowNode->OnSelected(*DataflowContext);
+										GetClothCollectionIfPossible(SelectedDataflowNode, DataflowContext);
+										SelectedDataflowNode->OnSelected(*DataflowContext);
 									}
 								});
 						}
diff --git a/Engine/Plugins/Experimental/ChaosClothAssetEditor/Source/ChaosClothAssetEditor/Private/ChaosClothAsset/ClothEditorToolkit.h b/Engine/Plugins/Experimental/ChaosClothAssetEditor/Source/ChaosClothAssetEditor/Private/ChaosClothAsset/ClothEditorToolkit.h
--- a/Engine/Plugins/Experimental/ChaosClothAssetEditor/Source/ChaosClothAssetEditor/Private/ChaosClothAsset/ClothEditorToolkit.h
+++ b/Engine/Plugins/Experimental/ChaosClothAssetEditor/Source/ChaosClothAssetEditor/Private/ChaosClothAsset/ClothEditorToolkit.h
@@ -122,6 +122,8 @@ private:
 	void ReinitializeGraphEditorWidget();
 	TSharedPtr<IStructureDetailsView> CreateNodeDetailsEditorWidget(UObject* ObjectToEdit);
 
+	TSharedPtr<FManagedArrayCollection> GetClothCollectionIfPossible(const TSharedPtr<FDataflowNode> InDataflowNode, const TSharedPtr<Dataflow::FEngineContext> Context);
+
 	// DataflowEditorActions
 	void OnPropertyValueChanged(const FPropertyChangedEvent& PropertyChangedEvent);
 	bool OnNodeVerifyTitleCommit(const FText& NewText, UEdGraphNode* GraphNode, FText& OutErrorMessage) const;
@@ -153,7 +155,7 @@ private:
 	TSharedPtr<Dataflow::FEngineContext> DataflowContext;
 	Dataflow::FTimestamp LastDataflowNodeTimestamp = Dataflow::FTimestamp::Invalid;
 	FDelegateHandle OnNodeInvalidatedDelegateHandle;
-	TSharedPtr<FDataflowNode> DataflowNode;
+	TSharedPtr<FDataflowNode> SelectedDataflowNode;
 
 	static const FName GraphCanvasTabId;
 	TSharedPtr<SDockTab> GraphEditorTab;
