Chaos Cloth Asset Editor - Fixed scale reset to 1 on the skin weight transfer tool.

diff --git a/Engine/Plugins/Experimental/ChaosClothAssetEditor/Source/ChaosClothAssetEditorTools/Public/ChaosClothAsset/ClothTransferSkinWeightsTool.h b/Engine/Plugins/Experimental/ChaosClothAssetEditor/Source/ChaosClothAssetEditorTools/Public/ChaosClothAsset/ClothTransferSkinWeightsTool.h
--- a/Engine/Plugins/Experimental/ChaosClothAssetEditor/Source/ChaosClothAssetEditorTools/Public/ChaosClothAsset/ClothTransferSkinWeightsTool.h
+++ b/Engine/Plugins/Experimental/ChaosClothAssetEditor/Source/ChaosClothAssetEditorTools/Public/ChaosClothAsset/ClothTransferSkinWeightsTool.h
@@ -30,13 +30,13 @@ public:
 	TObjectPtr<USkeletalMesh> SourceMesh;
 
 	UPROPERTY(EditAnywhere, Category = "Transform", meta = (DisplayName = "Location", EditCondition = "SourceMesh != nullptr"))
-	FVector3d SourceMeshTranslation;
+	FVector3d SourceMeshTranslation = FVector3d::ZeroVector;
 
 	UPROPERTY(EditAnywhere, Category = "Transform", meta = (DisplayName = "Rotation", EditCondition = "SourceMesh != nullptr"))
-	FVector3d SourceMeshRotation;
+	FVector3d SourceMeshRotation = FVector3d::ZeroVector;
 
 	UPROPERTY(EditAnywhere, Category = "Transform", meta = (DisplayName = "Scale", AllowPreserveRatio, EditCondition = "SourceMesh != nullptr"))
-	FVector3d SourceMeshScale;
+	FVector3d SourceMeshScale = FVector3d::OneVector;
 
 	UPROPERTY(EditAnywhere, Category = Source)
 	bool bHideSourceMesh = false;
