Dataflow - Enabled object pointer node property serialization.

diff --git a/Engine/Plugins/Experimental/ChaosClothAssetEditor/Source/ChaosClothAssetDataflowNodes/Private/ChaosClothAsset/SimulationDefaultConfigNode.cpp b/Engine/Plugins/Experimental/ChaosClothAssetEditor/Source/ChaosClothAssetDataflowNodes/Private/ChaosClothAsset/SimulationDefaultConfigNode.cpp
--- a/Engine/Plugins/Experimental/ChaosClothAssetEditor/Source/ChaosClothAssetDataflowNodes/Private/ChaosClothAsset/SimulationDefaultConfigNode.cpp
+++ b/Engine/Plugins/Experimental/ChaosClothAssetEditor/Source/ChaosClothAssetDataflowNodes/Private/ChaosClothAsset/SimulationDefaultConfigNode.cpp
@@ -9,6 +9,7 @@
 #include "Chaos/CollectionPropertyFacade.h"
 #include "Chaos/PBDLongRangeConstraints.h"  // For Tether modes
 #include "Dataflow/DataflowInputOutput.h"
+#include "UObject/FortniteMainBranchObjectVersion.h"
 
 #include UE_INLINE_GENERATED_CPP_BY_NAME(SimulationDefaultConfigNode)
 
@@ -16,8 +17,8 @@
 
 FChaosClothAssetSimulationDefaultConfigNode::FChaosClothAssetSimulationDefaultConfigNode(const Dataflow::FNodeParameters& InParam, FGuid InGuid)
 	: FDataflowNode(InParam, InGuid)
-	, SimulationConfig(NewObject<UChaosClothConfig>())
-	, SharedSimulationConfig(NewObject<UChaosClothSharedSimConfig>())
+	, SimulationConfig(NewObject<UChaosClothConfig>(InParam.OwningObject))
+	, SharedSimulationConfig(NewObject<UChaosClothSharedSimConfig>(InParam.OwningObject))
 {
 	RegisterInputConnection(&Collection);
 	RegisterOutputConnection(&Collection, &Collection);
@@ -25,17 +26,21 @@ FChaosClothAssetSimulationDefaultConfigNode::FChaosClothAssetSimulationDefaultCo
 
 void FChaosClothAssetSimulationDefaultConfigNode::Serialize(FArchive& Ar)
 {
-	if (!SimulationConfig)
+	Ar.UsingCustomVersion(FFortniteMainBranchObjectVersion::GUID);
+	if (Ar.CustomVer(FFortniteMainBranchObjectVersion::GUID) < FFortniteMainBranchObjectVersion::AddDataflowObjectSerialization)
 	{
-		SimulationConfig = NewObject<UChaosClothConfig>();
-	}
-	SimulationConfig->Serialize(Ar);
+		if (!SimulationConfig)
+		{
+			SimulationConfig = NewObject<UChaosClothConfig>();
+		}
+		SimulationConfig->Serialize(Ar);
 
-	if (!SharedSimulationConfig)
-	{
-		SharedSimulationConfig = NewObject<UChaosClothSharedSimConfig>();
+		if (!SharedSimulationConfig)
+		{
+			SharedSimulationConfig = NewObject<UChaosClothSharedSimConfig>();
+		}
+		SharedSimulationConfig->Serialize(Ar);
 	}
-	SharedSimulationConfig->Serialize(Ar);
 }
 
 void FChaosClothAssetSimulationDefaultConfigNode::Evaluate(Dataflow::FContext& Context, const FDataflowOutput* Out) const
diff --git a/Engine/Plugins/Experimental/Dataflow/Source/DataflowEditor/Private/Dataflow/DataflowSNode.cpp b/Engine/Plugins/Experimental/Dataflow/Source/DataflowEditor/Private/Dataflow/DataflowSNode.cpp
--- a/Engine/Plugins/Experimental/Dataflow/Source/DataflowEditor/Private/Dataflow/DataflowSNode.cpp
+++ b/Engine/Plugins/Experimental/Dataflow/Source/DataflowEditor/Private/Dataflow/DataflowSNode.cpp
@@ -202,7 +202,10 @@ static UDataflowEdNode* CreateNode(UDataflow* Dataflow, UEdGraphPin* FromPin, co
 {
 	if (Dataflow::FNodeFactory* Factory = Dataflow::FNodeFactory::GetInstance())
 	{
-		if (TSharedPtr<FDataflowNode> DataflowNode = Factory->NewNodeFromRegisteredType(*Dataflow->GetDataflow(), { FGuid::NewGuid(), NodeTypeName, NodeUniqueName }))
+		if (TSharedPtr<FDataflowNode> DataflowNode =
+			Factory->NewNodeFromRegisteredType(
+				*Dataflow->GetDataflow(),
+				{ FGuid::NewGuid(), NodeTypeName, NodeUniqueName, Dataflow }))
 		{
 			if (UDataflowEdNode* EdNode = NewObject<UDataflowEdNode>(Dataflow, UDataflowEdNode::StaticClass(), NodeUniqueName))
 			{
diff --git a/Engine/Source/Runtime/Core/Public/UObject/FortniteMainBranchObjectVersion.h b/Engine/Source/Runtime/Core/Public/UObject/FortniteMainBranchObjectVersion.h
--- a/Engine/Source/Runtime/Core/Public/UObject/FortniteMainBranchObjectVersion.h
+++ b/Engine/Source/Runtime/Core/Public/UObject/FortniteMainBranchObjectVersion.h
@@ -387,6 +387,9 @@ struct FFortniteMainBranchObjectVersion
 		// Update paths to keep a flag if they are the widget BP
 		MVVMPropertyPathSelf,
 
+		// Enabled ObjectPtr property serialization for Dataflow nodes
+		AddDataflowObjectSerialization,
+
 		// -----<new versions can be added above this line>-------------------------------------------------
 		VersionPlusOne,
 		LatestVersion = VersionPlusOne - 1
diff --git a/Engine/Source/Runtime/Experimental/Dataflow/Core/Private/Dataflow/DataflowGraph.cpp b/Engine/Source/Runtime/Experimental/Dataflow/Core/Private/Dataflow/DataflowGraph.cpp
--- a/Engine/Source/Runtime/Experimental/Dataflow/Core/Private/Dataflow/DataflowGraph.cpp
+++ b/Engine/Source/Runtime/Experimental/Dataflow/Core/Private/Dataflow/DataflowGraph.cpp
@@ -119,7 +119,15 @@ namespace Dataflow
 
 	}
 
-	void FGraph::Serialize(FArchive& Ar)
+	void FGraph::AddReferencedObjects(FReferenceCollector& Collector)
+	{
+		for (TSharedPtr<FDataflowNode>& Node : Nodes)
+		{
+			Collector.AddPropertyReferencesWithStructARO(Node->TypedScriptStruct(), Node.Get());
+		}
+	}
+
+	void FGraph::Serialize(FArchive& Ar, UObject* OwningObject)
 	{
 		Ar.UsingCustomVersion(FUE5MainStreamObjectVersion::GUID);
 
@@ -186,7 +194,7 @@ namespace Dataflow
 				FName ArNodeName;
 				Ar << ArGuid << ArType << ArNodeName;
 
-				TSharedPtr<FDataflowNode> Node = FNodeFactory::GetInstance()->NewNodeFromRegisteredType(*this, { ArGuid,ArType, ArNodeName });
+				TSharedPtr<FDataflowNode> Node = FNodeFactory::GetInstance()->NewNodeFromRegisteredType(*this, { ArGuid, ArType, ArNodeName, OwningObject });
 				DATAFLOW_OPTIONAL_BLOCK_READ_BEGIN(Node != nullptr)
 				{
 					ensure(!NodeGuidMap.Contains(ArGuid));
diff --git a/Engine/Source/Runtime/Experimental/Dataflow/Core/Public/Dataflow/DataflowGraph.h b/Engine/Source/Runtime/Experimental/Dataflow/Core/Public/Dataflow/DataflowGraph.h
--- a/Engine/Source/Runtime/Experimental/Dataflow/Core/Public/Dataflow/DataflowGraph.h
+++ b/Engine/Source/Runtime/Experimental/Dataflow/Core/Public/Dataflow/DataflowGraph.h
@@ -153,7 +153,9 @@ namespace Dataflow
 		DATAFLOWCORE_API void Connect(FDataflowOutput* OutputConnection, FDataflowInput* InputConnection);
 		DATAFLOWCORE_API void Disconnect(FDataflowOutput* OutputConnection, FDataflowInput* InputConnection);
 
-		DATAFLOWCORE_API virtual void Serialize(FArchive& Ar);
+		DATAFLOWCORE_API void AddReferencedObjects(FReferenceCollector& Collector);
+
+		DATAFLOWCORE_API virtual void Serialize(FArchive& Ar, UObject* OwningObject);
 		const TSet<FName>& GetDisabledNodes() const { return DisabledNodes; }
 
 	};
@@ -171,17 +173,6 @@ FORCEINLINE FArchive& operator<<(Chaos::FChaosArchive& Ar, Dataflow::FLink& Valu
 	return Ar;
 }
 
-FORCEINLINE FArchive& operator<<(FArchive& Ar, Dataflow::FGraph& Value)
-{
-	Value.Serialize(Ar);
-	return Ar;
-}
-
-FORCEINLINE FArchive& operator<<(Chaos::FChaosArchive& Ar, Dataflow::FGraph& Value)
-{
-	Value.Serialize(Ar);
-	return Ar;
-}
 
 
 
diff --git a/Engine/Source/Runtime/Experimental/Dataflow/Core/Public/Dataflow/DataflowNode.h b/Engine/Source/Runtime/Experimental/Dataflow/Core/Public/Dataflow/DataflowNode.h
--- a/Engine/Source/Runtime/Experimental/Dataflow/Core/Public/Dataflow/DataflowNode.h
+++ b/Engine/Source/Runtime/Experimental/Dataflow/Core/Public/Dataflow/DataflowNode.h
@@ -18,6 +18,7 @@ class UScriptStruct;
 namespace Dataflow {
 	struct FNodeParameters {
 		FName Name;
+		UObject* OwningObject = nullptr;
 	};
 	class FGraph;
 }
@@ -316,7 +317,8 @@ namespace Dataflow
 		{A::StaticType(),A::StaticDisplay(),A::StaticCategory(),					\
 			A::StaticTags(),A::StaticToolTip()},									\
 		[](const ::Dataflow::FNewNodeParameters& InParam){							\
-				TUniquePtr<A> Val = MakeUnique<A>(::Dataflow::FNodeParameters{InParam.Name}, InParam.Guid);    \
+				TUniquePtr<A> Val = MakeUnique<A>(::Dataflow::FNodeParameters{		\
+					InParam.Name, InParam.OwningObject}, InParam.Guid);				\
 				Val->ValidateConnections(); return Val;});
 
 #define DATAFLOW_NODE_RENDER_TYPE(A, B)												\
diff --git a/Engine/Source/Runtime/Experimental/Dataflow/Core/Public/Dataflow/DataflowNodeFactory.h b/Engine/Source/Runtime/Experimental/Dataflow/Core/Public/Dataflow/DataflowNodeFactory.h
--- a/Engine/Source/Runtime/Experimental/Dataflow/Core/Public/Dataflow/DataflowNodeFactory.h
+++ b/Engine/Source/Runtime/Experimental/Dataflow/Core/Public/Dataflow/DataflowNodeFactory.h
@@ -17,6 +17,7 @@ namespace Dataflow
 		FGuid Guid;
 		FName Type;
 		FName Name;
+		UObject* OwningObject = nullptr;
 	};
 
 	struct FFactoryParameters {
diff --git a/Engine/Source/Runtime/Experimental/Dataflow/Engine/Private/Dataflow/DataflowObject.cpp b/Engine/Source/Runtime/Experimental/Dataflow/Engine/Private/Dataflow/DataflowObject.cpp
--- a/Engine/Source/Runtime/Experimental/Dataflow/Engine/Private/Dataflow/DataflowObject.cpp
+++ b/Engine/Source/Runtime/Experimental/Dataflow/Engine/Private/Dataflow/DataflowObject.cpp
@@ -45,6 +45,13 @@ void UDataflow::PostEditCallback()
 	// mark as dirty for the UObject
 }
 
+void UDataflow::AddReferencedObjects(UObject* InThis, FReferenceCollector& Collector)
+{
+	Super::AddReferencedObjects(InThis, Collector);
+	UDataflow* const This = CastChecked<UDataflow>(InThis);
+	This->Dataflow->AddReferencedObjects(Collector);
+}
+
 #if WITH_EDITOR
 
 void UDataflow::PostEditChangeProperty(struct FPropertyChangedEvent& PropertyChangedEvent)
@@ -109,7 +116,7 @@ void UDataflow::RemoveRenderTarget(UDataflowEdNode* InNode)
 void UDataflow::Serialize(FArchive& Ar)
 {
 	Super::Serialize(Ar);
-	Ar << *Dataflow.Get();
+	Dataflow->Serialize(Ar, this);
 }
 
 #undef LOCTEXT_NAMESPACE
diff --git a/Engine/Source/Runtime/Experimental/Dataflow/Engine/Public/Dataflow/DataflowObject.h b/Engine/Source/Runtime/Experimental/Dataflow/Engine/Public/Dataflow/DataflowObject.h
--- a/Engine/Source/Runtime/Experimental/Dataflow/Engine/Public/Dataflow/DataflowObject.h
+++ b/Engine/Source/Runtime/Experimental/Dataflow/Engine/Public/Dataflow/DataflowObject.h
@@ -89,6 +89,8 @@ public:
 
 public:
 	/** UObject Interface */
+	static DATAFLOWENGINE_API void AddReferencedObjects(UObject* InThis, FReferenceCollector& Collector);
+
 #if WITH_EDITOR
 	DATAFLOWENGINE_API virtual void PostEditChangeProperty(struct FPropertyChangedEvent& PropertyChangedEvent) override;
 #endif
