Fix cloth asset static mesh import to handle material indices correctly.

diff --git a/Engine/Plugins/Experimental/ChaosClothAssetEditor/Source/ChaosClothAssetDataflowNodes/Private/ChaosClothAsset/StaticMeshImportNode.cpp b/Engine/Plugins/Experimental/ChaosClothAssetEditor/Source/ChaosClothAssetDataflowNodes/Private/ChaosClothAsset/StaticMeshImportNode.cpp
--- a/Engine/Plugins/Experimental/ChaosClothAssetEditor/Source/ChaosClothAssetDataflowNodes/Private/ChaosClothAsset/StaticMeshImportNode.cpp
+++ b/Engine/Plugins/Experimental/ChaosClothAssetEditor/Source/ChaosClothAssetDataflowNodes/Private/ChaosClothAsset/StaticMeshImportNode.cpp
@@ -94,6 +94,7 @@ bool BuildSkeletalMeshModelFromMeshDescription(const FMeshDescription* const InM
 }
 
 bool InitializeDataFromMeshDescription(
+	const UStaticMesh* const StaticMesh,
 	const FMeshDescription* const InMeshDescription,
 	const FMeshBuildSettings& InBuildSettings,
 	const TArray<FStaticMaterial>& StaticMaterials,
@@ -102,10 +103,14 @@ bool InitializeDataFromMeshDescription(
 	FSkeletalMeshLODModel SkeletalMeshModel;
 	if (BuildSkeletalMeshModelFromMeshDescription(InMeshDescription, InBuildSettings, SkeletalMeshModel))
 	{
-		check(SkeletalMeshModel.Sections.Num() == StaticMaterials.Num());
+		FStaticMeshConstAttributes MeshAttributes(*InMeshDescription);
+		TPolygonGroupAttributesConstRef<FName> MaterialSlotNames = MeshAttributes.GetPolygonGroupMaterialSlotNames();
 		for (int32 SectionIndex = 0; SectionIndex < SkeletalMeshModel.Sections.Num(); ++SectionIndex)
 		{
-			const FString RenderMaterialPathName = StaticMaterials[SectionIndex].MaterialInterface ? StaticMaterials[SectionIndex].MaterialInterface->GetPathName() : "";
+			// Section MaterialIndex refers to the polygon group index. Look up which material this corresponds with.
+			const FName& MaterialSlotName = MaterialSlotNames[SkeletalMeshModel.Sections[SectionIndex].MaterialIndex];
+			const int32 MaterialIndex = StaticMesh->GetMaterialIndexFromImportedMaterialSlotName(MaterialSlotName);
+			const FString RenderMaterialPathName = StaticMaterials.IsValidIndex(MaterialIndex)&& StaticMaterials[MaterialIndex].MaterialInterface ? StaticMaterials[MaterialIndex].MaterialInterface->GetPathName() : "";
 			FClothDataflowTools::AddRenderPatternFromSkeletalMeshSection(ClothCollection, SkeletalMeshModel, SectionIndex, RenderMaterialPathName);
 		}
 		return true;
@@ -155,8 +160,8 @@ void FChaosClothAssetStaticMeshImportNode::Evaluate(Dataflow::FContext& Context,
 				}
 				if (bImportRenderMesh)
 				{
-					// Add render data into a single pattern for now
-					if (!InitializeDataFromMeshDescription(MeshDescription, StaticMesh->GetSourceModel(LODIndex).BuildSettings, StaticMesh->GetStaticMaterials(), ClothCollection))
+					// Add render data (section = pattern)
+					if (!InitializeDataFromMeshDescription(StaticMesh, MeshDescription, StaticMesh->GetSourceModel(LODIndex).BuildSettings, StaticMesh->GetStaticMaterials(), ClothCollection))
 					{
 						FClothDataflowTools::LogAndToastWarning(*this,
 							LOCTEXT("InvalidRenderMeshHeadline", "Invalid render mesh."),
