Chaos Cloth Asset Editor:
- Reworked Terminal node evaluations.
- Updated evaluation code to take the DataflowTerminalPath at evaluation time.
- Consolidate evaluations into one single EvaluateNode function.
- Removed duplicated evaluation context, since the context change crash couldn't be reproduced (will need a better fix if it does come back anyway).
- Added a node evaluation progress bar.
#tests Ran a week of tests with various existing Dataflow and new one created from scratch.

diff --git a/Engine/Plugins/Experimental/ChaosClothAssetEditor/Source/ChaosClothAssetEditor/Private/ChaosClothAsset/ClothEditorToolkit.cpp b/Engine/Plugins/Experimental/ChaosClothAssetEditor/Source/ChaosClothAssetEditor/Private/ChaosClothAsset/ClothEditorToolkit.cpp
--- a/Engine/Plugins/Experimental/ChaosClothAssetEditor/Source/ChaosClothAssetEditor/Private/ChaosClothAsset/ClothEditorToolkit.cpp
+++ b/Engine/Plugins/Experimental/ChaosClothAssetEditor/Source/ChaosClothAssetEditor/Private/ChaosClothAsset/ClothEditorToolkit.cpp
@@ -240,35 +240,9 @@ const UDataflow* FChaosClothAssetEditorToolkit::GetDataflow() const
 
 void FChaosClothAssetEditorToolkit::Tick(float DeltaTime)
 {
-	UChaosClothAsset* const ClothAsset = GetAsset();
-
-	if (Dataflow && ClothAsset)
-	{
-		if (!DataflowContext)
-		{
-			DataflowContext = TSharedPtr<Dataflow::FEngineContext>(new Dataflow::FClothAssetDataflowContext(ClothAsset, Dataflow, Dataflow::FTimestamp::Invalid));
-			LastDataflowNodeTimestamp = Dataflow::FTimestamp::Invalid;
-		}
-		DataflowTerminalPath = Private::GetDataflowTerminalFrom(ClothAsset);
-
-		Dataflow::FTimestamp OldTimestamp = LastDataflowNodeTimestamp;
-
-		const TSharedPtr<Dataflow::FEngineContext> EvaluationContext(DataflowContext);  // Copy the context pointer as to not lose the reference during an evaluation (some UI operations can reset the toolkit context mid evaluation)
-		FDataflowEditorCommands::EvaluateTerminalNode(*EvaluationContext, LastDataflowNodeTimestamp, Dataflow, nullptr, nullptr, ClothAsset, DataflowTerminalPath);
-
-		if (OldTimestamp.Value < LastDataflowNodeTimestamp.Value)
-		{
-			OnClothAssetChanged();
-
-			// Refresh the construction viewport
-			if (SelectedDataflowNode)
-			{
-				TSharedPtr<FManagedArrayCollection> Collection = GetClothCollectionIfPossible(SelectedDataflowNode, DataflowContext);
-				UChaosClothAssetEditorMode* const ClothMode = CastChecked<UChaosClothAssetEditorMode>(EditorModeManager->GetActiveScriptableMode(UChaosClothAssetEditorMode::EM_ChaosClothAssetEditorModeId));
-				ClothMode->SetSelectedClothCollection(Collection);
-			}
-		}
-	}
+	// Evaluate terminal node if one is assigned, but only if needed (don't force it)
+	constexpr bool bForceOperation = false;
+	EvaluateNode(nullptr, bForceOperation);
 
 	InvalidateViews();
 }
@@ -336,8 +310,6 @@ void FChaosClothAssetEditorToolkit::CreateWidgets()
 		// TODO: Figure out how to create the GraphEditor widgets when the ClothAsset doesn't have a Dataflow property set
 		if (Dataflow)
 		{
-			DataflowTerminalPath = Private::GetDataflowTerminalFrom(ClothAsset);
-
 			Dataflow->Schema = UDataflowSchema::StaticClass();
 
 			NodeDetailsEditor = CreateNodeDetailsEditorWidget(ClothAsset);
@@ -804,21 +776,107 @@ void FChaosClothAssetEditorToolkit::OnFinishedChangingAssetProperties(const FPro
 	}
 }
 
-void FChaosClothAssetEditorToolkit::EvaluateNode(FDataflowNode* Node, FDataflowOutput* Out)
+void FChaosClothAssetEditorToolkit::EvaluateNode(FDataflowNode* Node, bool bForceOperation)
 {
-	if (Dataflow)
+	UChaosClothAsset* const ClothAsset = GetAsset();
+
+	if (Dataflow && ClothAsset)
 	{
-		UObject* const Asset = GetAsset();
+		const float NumSteps = 2.f;
+		FScopedSlowTask SlowTask(NumSteps, LOCTEXT("ChaosClothAssetEditorToolkitEvaluateNode", "Evaluating nodes..."));
+		SlowTask.MakeDialogDelayed(1.f);
+
+		const TSharedPtr<Dataflow::FGraph> Graph = Dataflow->GetDataflow();
+		check(Graph);
+		
+		const FName DataflowTerminalName = *Private::GetDataflowTerminalFrom(ClothAsset);
+
+		if (!Node)
+		{
+			// Retrieve terminal node
+			if (const TSharedPtr<FDataflowNode> TerminalNode = Graph->FindBaseNode(DataflowTerminalName))
+			{
+				Node = TerminalNode->AsType<FDataflowTerminalNode>();
+			}
+			if (!Node)
+			{
+				return;  // No terminal node set
+			}
+		}
 
+		// Set evaluation context
 		if (!DataflowContext)
 		{
-			DataflowContext = TSharedPtr<Dataflow::FEngineContext>(new Dataflow::FClothAssetDataflowContext(Asset, Dataflow, Dataflow::FTimestamp::Invalid));
+			DataflowContext = TSharedPtr<Dataflow::FEngineContext>(new Dataflow::FClothAssetDataflowContext(ClothAsset, Dataflow, Dataflow::FTimestamp::Invalid));
+			bForceOperation = true;
+		}
+		if (bForceOperation)
+		{
+			LastDataflowNodeTimestamp = Dataflow::FTimestamp::Invalid;
 		}
-		LastDataflowNodeTimestamp = Dataflow::FTimestamp::Invalid;
 
 		Dataflow::FTimestamp OldTimestamp = LastDataflowNodeTimestamp;
-		FDataflowEditorCommands::EvaluateTerminalNode(*DataflowContext.Get(), LastDataflowNodeTimestamp, Dataflow, Node, nullptr, Asset, DataflowTerminalPath);
 
+		// Set all asset values by visiting the graph from the specified node
+		TArray<FDataflowTerminalNode*> TerminalNodes;
+		TerminalNodes.Reserve(Graph->GetTerminalNodes().Num());
+		
+		TQueue<FDataflowNode*> Queue;
+		Queue.Enqueue(Node);
+
+		do
+		{
+			FDataflowNode* VisitedNode;
+			Queue.Dequeue(VisitedNode);
+
+			if (VisitedNode->GetTimestamp() >= LastDataflowNodeTimestamp)
+			{
+				UE_LOG(LogChaosClothAssetEditor, VeryVerbose, TEXT("EvaluateNode - Visiting node %s"), *VisitedNode->GetName().ToString());
+
+				for (FDataflowInput* const Input : VisitedNode->GetInputs())
+				{
+					for (FDataflowOutput* const ConnectedOutput : Input->GetConnectedOutputs())
+					{
+						if (FDataflowNode* const OwningNode = ConnectedOutput->GetOwningNode())
+						{
+							Queue.Enqueue(OwningNode);
+
+							if (FDataflowTerminalNode* const TerminalNode = OwningNode->AsType<FDataflowTerminalNode>())
+							{
+								UE_LOG(LogChaosClothAssetEditor, VeryVerbose, TEXT("EvaluateNode - Found terminal node %s"), *TerminalNode->GetName().ToString());
+								TerminalNodes.Emplace(TerminalNode);
+							}
+						}
+					}
+				}
+			}
+		} while (!Queue.IsEmpty());
+
+		TSet<const FDataflowNode*> ProcessedTerminalNodes;
+		ProcessedTerminalNodes.Reserve(TerminalNodes.Num());  // Only cull the already visited nodes at this later stage, so that it always find nodes using the longest path in the graph
+
+		for (int32 TerminalNodeIndex = TerminalNodes.Num() - 1; TerminalNodeIndex >= 0; --TerminalNodeIndex)
+		{
+			FDataflowTerminalNode* const TerminalNode = TerminalNodes[TerminalNodeIndex];
+
+			bool bHasAlreadyBeenProcessed;
+			ProcessedTerminalNodes.FindOrAdd(TerminalNode, &bHasAlreadyBeenProcessed);
+
+			if (!bHasAlreadyBeenProcessed)
+			{
+				UE_LOG(LogChaosClothAssetEditor, VeryVerbose, TEXT("EvaluateNode - Setting asset value for terminal node %s"), *TerminalNodes[TerminalNodeIndex]->GetName().ToString());
+				TerminalNode->SetAssetValue(ClothAsset, *DataflowContext);
+			}
+		}
+		SlowTask.EnterProgressFrame(1.f);
+
+		// Evaluate the node and update LastDataflowNodeTimestamp if doing so (even if it isn't a terminal node despite the name of the command)
+		// If node is a terminal node, it will have its asset value updated prior to its output being evaluated
+		FDataflowEditorCommands::EvaluateTerminalNode(*DataflowContext, LastDataflowNodeTimestamp, Dataflow, Node, nullptr, ClothAsset);
+
+		SlowTask.EnterProgressFrame(1.f);
+
+		// Refresh editor
 		if (OldTimestamp.Value < LastDataflowNodeTimestamp.Value)
 		{
 			OnClothAssetChanged();
@@ -832,16 +890,16 @@ void FChaosClothAssetEditorToolkit::EvaluateNode(FDataflowNode* Node, FDataflowO
 			}
 		}
 	}
-};
+}
 
 TSharedRef<SDataflowGraphEditor> FChaosClothAssetEditorToolkit::CreateGraphEditorWidget()
 {
 	ensure(Dataflow);
 	using namespace Dataflow;
 
-	const auto EvalLambda = [this](FDataflowNode* Node, FDataflowOutput* Out)
+	const auto EvalLambda = [this](FDataflowNode* Node, FDataflowOutput* /*Out*/)
 	{
-		this->EvaluateNode(Node, Out);
+		EvaluateNode(Node);
 	};
 
 	SGraphEditor::FGraphEditorEvents InEvents;
@@ -866,9 +924,9 @@ void FChaosClothAssetEditorToolkit::ReinitializeGraphEditorWidget()
 {
 	ensure(Dataflow);
 
-	const auto EvalLambda = [this](FDataflowNode* Node, FDataflowOutput* Out)
+	const auto EvalLambda = [this](FDataflowNode* Node, FDataflowOutput* /*Out*/)
 	{
-		EvaluateNode(Node, Out);
+		EvaluateNode(Node);
 	};
 
 	SGraphEditor::FGraphEditorEvents InEvents;
@@ -879,8 +937,6 @@ void FChaosClothAssetEditorToolkit::ReinitializeGraphEditorWidget()
 
 	if (!GraphEditor)
 	{
-		DataflowTerminalPath = Private::GetDataflowTerminalFrom(ClothAsset);
-
 		NodeDetailsEditor = CreateNodeDetailsEditorWidget(ClothAsset);
 		if (NodeDetailsTab.IsValid())
 		{
@@ -964,6 +1020,15 @@ TSharedPtr<FManagedArrayCollection> FChaosClothAssetEditorToolkit::GetClothColle
 		{
 			if (Output->GetType() == FName("FManagedArrayCollection"))
 			{
+				// Set asset value first if this node is a terminal node as the output might need the result of this operation
+				if (const FDataflowTerminalNode* const TerminalNode = InDataflowNode->AsType<const FDataflowTerminalNode>())
+				{
+					if (UChaosClothAsset* const ClothAsset = GetAsset())
+					{
+						TerminalNode->SetAssetValue(ClothAsset, *Context);
+					}
+				}
+
 				const FManagedArrayCollection DefaultValue;
 				TSharedRef<FManagedArrayCollection> Collection = MakeShared<FManagedArrayCollection>(Output->GetValue<FManagedArrayCollection>(*Context, DefaultValue));
 
@@ -1030,7 +1095,10 @@ void FChaosClothAssetEditorToolkit::OnNodeSelectionChanged(const TSet<UObject*>&
 						if (SelectedDataflowNode)
 						{
 							Collection = GetClothCollectionIfPossible(SelectedDataflowNode, DataflowContext);
-							SelectedDataflowNode->OnSelected(*DataflowContext);
+							if (DataflowContext.IsValid())  // TODO: The context shouldn't be nullptr on the first node creation
+							{
+								SelectedDataflowNode->OnSelected(*DataflowContext);
+							}
 
 							// Set a callback to re-evaluate the node if it is invalidated
 							OnNodeInvalidatedDelegateHandle = SelectedDataflowNode->GetOnNodeInvalidatedDelegate().AddLambda(
@@ -1039,7 +1107,10 @@ void FChaosClothAssetEditorToolkit::OnNodeSelectionChanged(const TSet<UObject*>&
 									if (SelectedDataflowNode.Get() == InDataflowNode)
 									{
 										GetClothCollectionIfPossible(SelectedDataflowNode, DataflowContext);
-										SelectedDataflowNode->OnSelected(*DataflowContext);
+										if (DataflowContext.IsValid())  // TODO: The context shouldn't be nullptr on the first node creation
+										{
+											SelectedDataflowNode->OnSelected(*DataflowContext);
+										}
 									}
 								});
 						}
diff --git a/Engine/Plugins/Experimental/ChaosClothAssetEditor/Source/ChaosClothAssetEditor/Private/ChaosClothAsset/ClothEditorToolkit.h b/Engine/Plugins/Experimental/ChaosClothAssetEditor/Source/ChaosClothAssetEditor/Private/ChaosClothAsset/ClothEditorToolkit.h
--- a/Engine/Plugins/Experimental/ChaosClothAssetEditor/Source/ChaosClothAssetEditor/Private/ChaosClothAsset/ClothEditorToolkit.h
+++ b/Engine/Plugins/Experimental/ChaosClothAssetEditor/Source/ChaosClothAssetEditor/Private/ChaosClothAsset/ClothEditorToolkit.h
@@ -117,7 +117,9 @@ private:
 	void InvalidateViews();
 
 	// Dataflow
-	void EvaluateNode(FDataflowNode* Node, FDataflowOutput* Out);
+	UE_DEPRECATED(5.4, "Use EvaluateNode(FDataflowNode*, bool) instead.")
+	void EvaluateNode(FDataflowNode* Node, FDataflowOutput* Out) { EvaluateNode(Node); }
+	void EvaluateNode(FDataflowNode* Node, bool bForceOperation = true);
 	TSharedRef<SDataflowGraphEditor> CreateGraphEditorWidget();
 	void ReinitializeGraphEditorWidget();
 	TSharedPtr<IStructureDetailsView> CreateNodeDetailsEditorWidget(UObject* ObjectToEdit);
@@ -151,7 +153,6 @@ private:
 
 	// Dataflow
 	UDataflow* Dataflow = nullptr;
-	FString DataflowTerminalPath = "";
 	TSharedPtr<Dataflow::FEngineContext> DataflowContext;
 	Dataflow::FTimestamp LastDataflowNodeTimestamp = Dataflow::FTimestamp::Invalid;
 	FDelegateHandle OnNodeInvalidatedDelegateHandle;
diff --git a/Engine/Plugins/Experimental/ChaosClothAssetEditor/Source/ChaosClothAssetTools/Private/ChaosClothAsset/ClothAssetBuilderEditor.cpp b/Engine/Plugins/Experimental/ChaosClothAssetEditor/Source/ChaosClothAssetTools/Private/ChaosClothAsset/ClothAssetBuilderEditor.cpp
--- a/Engine/Plugins/Experimental/ChaosClothAssetEditor/Source/ChaosClothAssetTools/Private/ChaosClothAsset/ClothAssetBuilderEditor.cpp
+++ b/Engine/Plugins/Experimental/ChaosClothAssetEditor/Source/ChaosClothAssetTools/Private/ChaosClothAsset/ClothAssetBuilderEditor.cpp
@@ -6,6 +6,7 @@
 #include "Chaos/CollectionPropertyFacade.h"
 #include "Engine/SkeletalMesh.h"
 #include "Engine/SkinnedAssetCommon.h"
+#include "Misc/ScopedSlowTask.h"
 #include "Modules/ModuleManager.h"
 #include "Rendering/SkeletalMeshLODModel.h"
 #include "Utils/ClothingMeshUtils.h"
@@ -15,6 +16,8 @@
 
 #include UE_INLINE_GENERATED_CPP_BY_NAME(ClothAssetBuilderEditor)
 
+#define LOCTEXT_NAMESPACE "ClothAssetBuilderEditor"
+
 void UClothAssetBuilderEditor::BuildLod(FSkeletalMeshLODModel& LODModel, const UChaosClothAsset& ClothAsset, int32 LodIndex) const
 {
 	using namespace UE::Chaos::ClothAsset;
@@ -76,6 +79,9 @@ void UClothAssetBuilderEditor::BuildLod(FSkeletalMeshLODModel& LODModel, const U
 	const int32 NumSections = ClothFacade.GetNumRenderPatterns();  // Cloth Render Patterns == Skeletal Mesh Sections
 	LODModel.Sections.SetNum(NumSections);
 
+	FScopedSlowTask SlowTask((float)ClothFacade.GetNumRenderFaces(), LOCTEXT("ClothAssetBuildLOD", "Building Cloth Asset LOD sections..."));
+	SlowTask.MakeDialogDelayed(1.f);
+
 	int32 BaseIndex = 0;
 	for(int32 SectionIndex = 0; SectionIndex < NumSections; ++SectionIndex)
 	{
@@ -306,6 +312,8 @@ void UClothAssetBuilderEditor::BuildLod(FSkeletalMeshLODModel& LODModel, const U
 
 		// Copy to user section data, otherwise the section data set above would get lost when the user section gets synced
 		FSkelMeshSourceSectionUserData::GetSourceSectionUserData(LODModel.UserSectionsData, Section);
+
+		SlowTask.EnterProgressFrame(NumFaces);
 	}
 
 
@@ -318,3 +326,5 @@ void UClothAssetBuilderEditor::BuildLod(FSkeletalMeshLODModel& LODModel, const U
 	// Compute the required bones for this model.
 	USkeletalMesh::CalculateRequiredBones(LODModel, ClothAsset.RefSkeleton, nullptr);
 }
+
+#undef LOCTEXT_NAMESPACE
