Chaos Cloth Asset - Removed the redundant FGCObject inheritence from the SimulationDefaultConfig node, after the recent Dataflow change to allow passing of the owning object.
#tests Reloaded old and new assets post and prior the Dataflow owning object change.

diff --git a/Engine/Plugins/Experimental/ChaosClothAssetEditor/Source/ChaosClothAssetDataflowNodes/Private/ChaosClothAsset/SimulationDefaultConfigNode.cpp b/Engine/Plugins/Experimental/ChaosClothAssetEditor/Source/ChaosClothAssetDataflowNodes/Private/ChaosClothAsset/SimulationDefaultConfigNode.cpp
--- a/Engine/Plugins/Experimental/ChaosClothAssetEditor/Source/ChaosClothAssetDataflowNodes/Private/ChaosClothAsset/SimulationDefaultConfigNode.cpp
+++ b/Engine/Plugins/Experimental/ChaosClothAssetEditor/Source/ChaosClothAssetDataflowNodes/Private/ChaosClothAsset/SimulationDefaultConfigNode.cpp
@@ -19,6 +19,7 @@ FChaosClothAssetSimulationDefaultConfigNode::FChaosClothAssetSimulationDefaultCo
 	: FDataflowNode(InParam, InGuid)
 	, SimulationConfig(NewObject<UChaosClothConfig>(InParam.OwningObject))
 	, SharedSimulationConfig(NewObject<UChaosClothSharedSimConfig>(InParam.OwningObject))
+	, OwningObject(InParam.OwningObject)
 {
 	RegisterInputConnection(&Collection);
 	RegisterOutputConnection(&Collection, &Collection);
@@ -31,13 +32,13 @@ void FChaosClothAssetSimulationDefaultConfigNode::Serialize(FArchive& Ar)
 	{
 		if (!SimulationConfig)
 		{
-			SimulationConfig = NewObject<UChaosClothConfig>();
+			SimulationConfig = NewObject<UChaosClothConfig>(OwningObject);
 		}
 		SimulationConfig->Serialize(Ar);
 
 		if (!SharedSimulationConfig)
 		{
-			SharedSimulationConfig = NewObject<UChaosClothSharedSimConfig>();
+			SharedSimulationConfig = NewObject<UChaosClothSharedSimConfig>(OwningObject);
 		}
 		SharedSimulationConfig->Serialize(Ar);
 	}
@@ -84,15 +85,4 @@ void FChaosClothAssetSimulationDefaultConfigNode::Evaluate(Dataflow::FContext& C
 	}
 }
 
-void FChaosClothAssetSimulationDefaultConfigNode::AddReferencedObjects(FReferenceCollector& Collector)
-{
-	Collector.AddReferencedObject(SimulationConfig);
-	Collector.AddReferencedObject(SharedSimulationConfig);
-}
-
-FString FChaosClothAssetSimulationDefaultConfigNode::GetReferencerName() const
-{
-	return TEXT("FChaosClothAssetSimulationDefaultConfigNode");
-}
-
 #undef LOCTEXT_NAMESPACE
diff --git a/Engine/Plugins/Experimental/ChaosClothAssetEditor/Source/ChaosClothAssetDataflowNodes/Public/ChaosClothAsset/SimulationDefaultConfigNode.h b/Engine/Plugins/Experimental/ChaosClothAssetEditor/Source/ChaosClothAssetDataflowNodes/Public/ChaosClothAsset/SimulationDefaultConfigNode.h
--- a/Engine/Plugins/Experimental/ChaosClothAssetEditor/Source/ChaosClothAssetDataflowNodes/Public/ChaosClothAsset/SimulationDefaultConfigNode.h
+++ b/Engine/Plugins/Experimental/ChaosClothAssetEditor/Source/ChaosClothAssetDataflowNodes/Public/ChaosClothAsset/SimulationDefaultConfigNode.h
@@ -11,7 +11,7 @@ class UChaosClothSharedSimConfig;
 
 /** Add default simulation properties to the cloth collection in the format of the skeletal mesh cloth editor. */
 USTRUCT(Meta = (DataflowCloth))
-struct FChaosClothAssetSimulationDefaultConfigNode : public FDataflowNode, public FGCObject
+struct FChaosClothAssetSimulationDefaultConfigNode : public FDataflowNode
 {
 	GENERATED_USTRUCT_BODY()
 	DATAFLOW_NODE_DEFINE_INTERNAL(FChaosClothAssetSimulationDefaultConfigNode, "SimulationDefaultConfig", "Cloth", "Cloth Simulation Default Config")
@@ -36,8 +36,5 @@ private:
 	virtual void Serialize(FArchive& Ar) override;
 	//~ End FDataflowNode Interface
 
-	//~ Begin FGCObject Interface
-	virtual void AddReferencedObjects(FReferenceCollector& Collector) override;
-	virtual FString GetReferencerName() const override;
-	//~ End FGCObject Interface
+	UObject* OwningObject;  // For backward compatibility when null SimulationConfig and SharedSimulationConfig are reloaded
 };
