Cloth Asset Editor - Removed saving from Dataflow asset creation to match the cloth asset behavior.
#tests Created a new Dataflow asset and verified that the saving behavior now matches these of the Cloth asset.

diff --git a/Engine/Plugins/Experimental/ChaosClothAssetEditor/Source/ChaosClothAssetEditor/Private/ChaosClothAsset/AssetDefinition_ClothAsset.cpp b/Engine/Plugins/Experimental/ChaosClothAssetEditor/Source/ChaosClothAssetEditor/Private/ChaosClothAsset/AssetDefinition_ClothAsset.cpp
--- a/Engine/Plugins/Experimental/ChaosClothAssetEditor/Source/ChaosClothAssetEditor/Private/ChaosClothAsset/AssetDefinition_ClothAsset.cpp
+++ b/Engine/Plugins/Experimental/ChaosClothAssetEditor/Source/ChaosClothAssetEditor/Private/ChaosClothAsset/AssetDefinition_ClothAsset.cpp
@@ -59,19 +59,6 @@ namespace ClothAssetDefinitionHelpers
 		// Notify the asset registry
 		FAssetRegistryModule::AssetCreated(NewAsset);
 
-		// Save the package
-		TArray<UPackage*> PackagesToSave;
-		PackagesToSave.Add(NewAsset->GetOutermost());
-		constexpr bool bCheckDirty = false;
-		constexpr bool bPromptToSave = false;
-		FEditorFileUtils::EPromptReturnCode ReturnCode = FEditorFileUtils::PromptForCheckoutAndSave(PackagesToSave, bCheckDirty, bPromptToSave);
-
-		if (ReturnCode != FEditorFileUtils::EPromptReturnCode::PR_Success)
-		{
-			OutDataflowAsset = nullptr;
-			return false;
-		}
-
 		OutDataflowAsset = NewAsset;
 		return true;
 	}
