Chaos Cloth Asset - Added USD Import node.
#tests Successfully Imported the latest USD files provided by our third party partner.

diff --git a/Engine/Plugins/Experimental/ChaosClothAssetEditor/ChaosClothAssetEditor.uplugin b/Engine/Plugins/Experimental/ChaosClothAssetEditor/ChaosClothAssetEditor.uplugin
--- a/Engine/Plugins/Experimental/ChaosClothAssetEditor/ChaosClothAssetEditor.uplugin
+++ b/Engine/Plugins/Experimental/ChaosClothAssetEditor/ChaosClothAssetEditor.uplugin
@@ -76,6 +76,13 @@
 			"SupportedTargetPlatforms": [
 				"Win64"
 			]
+		},
+		{
+			"Name": "USDImporter",
+			"Enabled": true,
+			"SupportedTargetPlatforms": [
+				"Win64"
+			]
 		}
 	]
 }
\ No newline at end of file
diff --git a/Engine/Plugins/Experimental/ChaosClothAssetEditor/Source/ChaosClothAssetDataflowNodes/ChaosClothAssetDataflowNodes.Build.cs b/Engine/Plugins/Experimental/ChaosClothAssetEditor/Source/ChaosClothAssetDataflowNodes/ChaosClothAssetDataflowNodes.Build.cs
--- a/Engine/Plugins/Experimental/ChaosClothAssetEditor/Source/ChaosClothAssetDataflowNodes/ChaosClothAssetDataflowNodes.Build.cs
+++ b/Engine/Plugins/Experimental/ChaosClothAssetEditor/Source/ChaosClothAssetDataflowNodes/ChaosClothAssetDataflowNodes.Build.cs
@@ -12,7 +12,7 @@ public class ChaosClothAssetDataflowNodes : ModuleRules
 				"Core",
 				// ... add other public dependencies that you statically link with here ...
 			}
-			);	
+		);
 		
 		PrivateDependencyModuleNames.AddRange(
 			new string[]
@@ -29,6 +29,7 @@ public class ChaosClothAssetDataflowNodes : ModuleRules
 				"DataflowEngine",
 				"DatasmithCore",
 				"DatasmithImporter",
+				"DesktopWidgets",  // For SFilePathPicker
 				"DetailCustomizations",
 				"DynamicMesh",
 				"Engine",
@@ -44,8 +45,13 @@ public class ChaosClothAssetDataflowNodes : ModuleRules
 				"SlateCore",
 				"StaticMeshDescription",
 				"UnrealEd",
-				// ... add private dependencies that you statically link with here ...	
+				"UnrealUSDWrapper",
+				"USDSchemas",
+				"USDStage",
+				"USDStageImporter",
+				"USDUtilities",
+				// ... add private dependencies that you statically link with here ...
 			}
-			);
+		);
 	}
 }
diff --git a/Engine/Plugins/Experimental/ChaosClothAssetEditor/Source/ChaosClothAssetDataflowNodes/Private/ChaosClothAsset/ChaosClothAssetDataflowNodesModule.cpp b/Engine/Plugins/Experimental/ChaosClothAssetEditor/Source/ChaosClothAssetDataflowNodes/Private/ChaosClothAsset/ChaosClothAssetDataflowNodesModule.cpp
--- a/Engine/Plugins/Experimental/ChaosClothAssetEditor/Source/ChaosClothAssetDataflowNodes/Private/ChaosClothAsset/ChaosClothAssetDataflowNodesModule.cpp
+++ b/Engine/Plugins/Experimental/ChaosClothAssetEditor/Source/ChaosClothAssetDataflowNodes/Private/ChaosClothAsset/ChaosClothAssetDataflowNodesModule.cpp
@@ -8,6 +8,7 @@
 #include "ChaosClothAsset/SelectionGroupCustomization.h"
 #include "ChaosClothAsset/DatasmithImportNode.h"
 #include "ChaosClothAsset/DeleteElementNode.h"
+#include "ChaosClothAsset/ImportFilePathCustomization.h"
 #include "ChaosClothAsset/ImportNode.h"
 #include "ChaosClothAsset/MergeClothCollectionsNode.h"
 #include "ChaosClothAsset/ReverseNormalsNode.h"
@@ -44,6 +45,7 @@
 #include "ChaosClothAsset/TerminalNode.h"
 #include "ChaosClothAsset/TransferSkinWeightsNode.h"
 #include "ChaosClothAsset/TransformUVsNode.h"
+#include "ChaosClothAsset/USDImportNode.h"
 #include "ChaosClothAsset/WeightedValueCustomization.h"
 #include "Dataflow/DataflowNodeColorsRegistry.h"
 #include "Dataflow/DataflowNodeFactory.h"
@@ -96,7 +98,7 @@ namespace UE::Chaos::ClothAsset::Private
 		DATAFLOW_NODE_REGISTER_CREATION_FACTORY(FChaosClothAssetTerminalNode);
 		DATAFLOW_NODE_REGISTER_CREATION_FACTORY(FChaosClothAssetTransferSkinWeightsNode);
 		DATAFLOW_NODE_REGISTER_CREATION_FACTORY(FChaosClothAssetTransformUVsNode);
-
+		DATAFLOW_NODE_REGISTER_CREATION_FACTORY(FChaosClothAssetUSDImportNode);
 		DATAFLOW_NODE_REGISTER_CREATION_FACTORY_NODE_COLORS_BY_CATEGORY("Cloth", FColorScheme::NodeHeader, FColorScheme::NodeBody);
 	}
 }
@@ -114,6 +116,7 @@ void FChaosClothAssetDataflowNodesModule::StartupModule()
 		PropertyModule->RegisterCustomPropertyTypeLayout(FChaosClothAssetWeightedValueNonAnimatable::StaticStruct()->GetFName(), FOnGetPropertyTypeCustomizationInstance::CreateStatic(&FWeightedValueCustomization::MakeInstance));
 		PropertyModule->RegisterCustomPropertyTypeLayout(FChaosClothAssetWeightedValueNonAnimatableNoLowHighRange::StaticStruct()->GetFName(), FOnGetPropertyTypeCustomizationInstance::CreateStatic(&FWeightedValueCustomization::MakeInstance));
 		PropertyModule->RegisterCustomPropertyTypeLayout(FChaosClothAssetNodeSelectionGroup::StaticStruct()->GetFName(), FOnGetPropertyTypeCustomizationInstance::CreateStatic(&FSelectionGroupCustomization::MakeInstance));
+		PropertyModule->RegisterCustomPropertyTypeLayout(FChaosClothAssetImportFilePath::StaticStruct()->GetFName(), FOnGetPropertyTypeCustomizationInstance::CreateStatic(&FImportFilePathCustomization::MakeInstance));
 	}
 
 	// Register modular features
@@ -136,5 +139,5 @@ void FChaosClothAssetDataflowNodesModule::ShutdownModule()
 }
 
 #undef LOCTEXT_NAMESPACE
-	
+
 IMPLEMENT_MODULE(FChaosClothAssetDataflowNodesModule, ChaosClothAssetDataflowNodes)
diff --git a/Engine/Plugins/Experimental/ChaosClothAssetEditor/Source/ChaosClothAssetDataflowNodes/Private/ChaosClothAsset/ClothDataflowTools.cpp b/Engine/Plugins/Experimental/ChaosClothAssetEditor/Source/ChaosClothAssetDataflowNodes/Private/ChaosClothAsset/ClothDataflowTools.cpp
--- a/Engine/Plugins/Experimental/ChaosClothAssetEditor/Source/ChaosClothAssetDataflowNodes/Private/ChaosClothAsset/ClothDataflowTools.cpp
+++ b/Engine/Plugins/Experimental/ChaosClothAssetEditor/Source/ChaosClothAssetDataflowNodes/Private/ChaosClothAsset/ClothDataflowTools.cpp
@@ -3,13 +3,16 @@
 #include "ChaosClothAsset/ClothDataflowTools.h"
 #include "ChaosClothAsset/ClothGeometryTools.h"
 #include "ChaosClothAsset/CollectionClothFacade.h"
-#include "Animation/Skeleton.h"
 #include "Dataflow/DataflowNode.h"
 #include "DynamicMesh/NonManifoldMappingSupport.h"
 #include "Framework/Notifications/NotificationManager.h"
-#include "Math/UnrealMathUtility.h"
+#include "Interfaces/ITargetPlatformManagerModule.h"
+#include "Modules/ModuleManager.h"
 #include "Rendering/SkeletalMeshLODModel.h"
 #include "Widgets/Notifications/SNotificationList.h"
+#include "MeshDescriptionToDynamicMesh.h"
+#include "MeshUtilities.h"
+#include "ReferenceSkeleton.h"
 #include "SkeletalMeshAttributes.h"
 #include "ToDynamicMesh.h"
 
@@ -457,4 +460,71 @@ namespace UE::Chaos::ClothAsset
 		bool bCharsWereRemoved;
 		do { InOutString.TrimCharInline(TEXT('_'), &bCharsWereRemoved); } while (bCharsWereRemoved);
 	}
+
+	static void CopyBuildSettings(const FMeshBuildSettings& InStaticMeshBuildSettings, FSkeletalMeshBuildSettings& OutSkeletalMeshBuildSettings)
+	{
+		OutSkeletalMeshBuildSettings.bRecomputeNormals = InStaticMeshBuildSettings.bRecomputeNormals;
+		OutSkeletalMeshBuildSettings.bRecomputeTangents = InStaticMeshBuildSettings.bRecomputeTangents;
+		OutSkeletalMeshBuildSettings.bUseMikkTSpace = InStaticMeshBuildSettings.bUseMikkTSpace;
+		OutSkeletalMeshBuildSettings.bComputeWeightedNormals = InStaticMeshBuildSettings.bComputeWeightedNormals;
+		OutSkeletalMeshBuildSettings.bRemoveDegenerates = InStaticMeshBuildSettings.bRemoveDegenerates;
+		OutSkeletalMeshBuildSettings.bUseHighPrecisionTangentBasis = InStaticMeshBuildSettings.bUseHighPrecisionTangentBasis;
+		OutSkeletalMeshBuildSettings.bUseFullPrecisionUVs = InStaticMeshBuildSettings.bUseFullPrecisionUVs;
+		OutSkeletalMeshBuildSettings.bUseBackwardsCompatibleF16TruncUVs = InStaticMeshBuildSettings.bUseBackwardsCompatibleF16TruncUVs;
+		// The rest we leave at defaults.
+	}
+
+	bool FClothDataflowTools::BuildSkeletalMeshModelFromMeshDescription(const FMeshDescription* const InMeshDescription, const FMeshBuildSettings& InBuildSettings, FSkeletalMeshLODModel& SkeletalMeshModel)
+	{
+		// This is following StaticToSkeletalMeshConverter.cpp::AddLODFromStaticMeshSourceModel
+		FSkeletalMeshBuildSettings BuildSettings;
+		CopyBuildSettings(InBuildSettings, BuildSettings);
+		FMeshDescription SkeletalMeshGeometry = *InMeshDescription;
+		FSkeletalMeshAttributes SkeletalMeshAttributes(SkeletalMeshGeometry);
+		SkeletalMeshAttributes.Register();
+
+		// Full binding to the root bone.
+		constexpr int32 RootBoneIndex = 0;
+		FSkinWeightsVertexAttributesRef SkinWeights = SkeletalMeshAttributes.GetVertexSkinWeights();
+		UE::AnimationCore::FBoneWeight RootInfluence(RootBoneIndex, 1.0f);
+		UE::AnimationCore::FBoneWeights RootBinding = UE::AnimationCore::FBoneWeights::Create({ RootInfluence });
+
+		for (const FVertexID VertexID : SkeletalMeshGeometry.Vertices().GetElementIDs())
+		{
+			SkinWeights.Set(VertexID, RootBinding);
+		}
+
+		FSkeletalMeshImportData SkeletalMeshImportGeometry = FSkeletalMeshImportData::CreateFromMeshDescription(SkeletalMeshGeometry);
+		// Data needed by BuildSkeletalMesh
+		TArray<FVector3f> LODPoints;
+		TArray<SkeletalMeshImportData::FMeshWedge> LODWedges;
+		TArray<SkeletalMeshImportData::FMeshFace> LODFaces;
+		TArray<SkeletalMeshImportData::FVertInfluence> LODInfluences;
+		TArray<int32> LODPointToRawMap;
+		SkeletalMeshImportGeometry.CopyLODImportData(LODPoints, LODWedges, LODFaces, LODInfluences, LODPointToRawMap);
+		IMeshUtilities::MeshBuildOptions BuildOptions;
+		BuildOptions.TargetPlatform = GetTargetPlatformManagerRef().GetRunningTargetPlatform();
+		BuildOptions.FillOptions(BuildSettings);
+
+		static const FString SkeletalMeshName("ClothAssetStaticMeshImportConvert"); // This is only used by warning messages in the mesh builder.
+		// Build a RefSkeleton with just a root bone. The BuildSkeletalMesh code expects you have a reference skeleton with at least one bone to work.
+		FReferenceSkeleton RootBoneRefSkeleton;
+		FReferenceSkeletonModifier SkeletonModifier(RootBoneRefSkeleton, nullptr);
+		FMeshBoneInfo RootBoneInfo;
+		RootBoneInfo.Name = FName("Root");
+		SkeletonModifier.Add(RootBoneInfo, FTransform());
+		RootBoneRefSkeleton.RebuildRefSkeleton(nullptr, true);
+
+		IMeshUtilities& MeshUtilities = FModuleManager::Get().LoadModuleChecked<IMeshUtilities>("MeshUtilities");
+		TArray<FText> WarningMessages;
+		if (!MeshUtilities.BuildSkeletalMesh(SkeletalMeshModel, SkeletalMeshName, RootBoneRefSkeleton, LODInfluences, LODWedges, LODFaces, LODPoints, LODPointToRawMap, BuildOptions, &WarningMessages))
+		{
+			for (const FText& Message : WarningMessages)
+			{
+				UE_LOG(LogChaosClothAssetDataflowNodes, Warning, TEXT("%s"), *Message.ToString());
+			}
+			return false;
+		}
+		return true;
+	}
 }  // End namespace UE::Chaos::ClothAsset
diff --git a/Engine/Plugins/Experimental/ChaosClothAssetEditor/Source/ChaosClothAssetDataflowNodes/Private/ChaosClothAsset/ClothDataflowTools.h b/Engine/Plugins/Experimental/ChaosClothAssetEditor/Source/ChaosClothAssetDataflowNodes/Private/ChaosClothAsset/ClothDataflowTools.h
--- a/Engine/Plugins/Experimental/ChaosClothAssetEditor/Source/ChaosClothAssetDataflowNodes/Private/ChaosClothAsset/ClothDataflowTools.h
+++ b/Engine/Plugins/Experimental/ChaosClothAssetEditor/Source/ChaosClothAssetDataflowNodes/Private/ChaosClothAsset/ClothDataflowTools.h
@@ -6,6 +6,8 @@
 #include "Logging/LogMacros.h"
 
 struct FManagedArrayCollection;
+struct FMeshBuildSettings;
+struct FMeshDescription;
 struct FDataflowNode;
 class FSkeletalMeshLODModel;
 class FString;
@@ -30,6 +32,8 @@ namespace UE::Chaos::ClothAsset
 		 * It will also have all leading underscore removed, as these names are reserved for internal use.
 		 */
 		static void MakeCollectionName(FString& InOutString);
+
+		static bool BuildSkeletalMeshModelFromMeshDescription(const FMeshDescription* const InMeshDescription, const FMeshBuildSettings& InBuildSettings, FSkeletalMeshLODModel& SkeletalMeshModel);
 	};
 }  // End namespace UE::Chaos::ClothAsset
 
diff --git a/Engine/Plugins/Experimental/ChaosClothAssetEditor/Source/ChaosClothAssetDataflowNodes/Private/ChaosClothAsset/ImportFilePath.h b/Engine/Plugins/Experimental/ChaosClothAssetEditor/Source/ChaosClothAssetDataflowNodes/Private/ChaosClothAsset/ImportFilePath.h
new file mode 100644
--- /dev/null
+++ b/Engine/Plugins/Experimental/ChaosClothAssetEditor/Source/ChaosClothAssetDataflowNodes/Private/ChaosClothAsset/ImportFilePath.h
@@ -0,0 +1,17 @@
+// Copyright Epic Games, Inc. All Rights Reserved.
+
+#pragma once
+
+#include "ImportFilePath.generated.h"
+
+USTRUCT()
+struct FChaosClothAssetImportFilePath
+{
+	GENERATED_USTRUCT_BODY()
+
+	UPROPERTY(EditAnywhere, Category = "Import File Path")
+	FString FilePath;
+
+	UPROPERTY(EditAnywhere, Category = "Import File Path")
+	bool bForceReimport = false;
+};
diff --git a/Engine/Plugins/Experimental/ChaosClothAssetEditor/Source/ChaosClothAssetDataflowNodes/Private/ChaosClothAsset/ImportFilePathCustomization.cpp b/Engine/Plugins/Experimental/ChaosClothAssetEditor/Source/ChaosClothAssetDataflowNodes/Private/ChaosClothAsset/ImportFilePathCustomization.cpp
new file mode 100644
--- /dev/null
+++ b/Engine/Plugins/Experimental/ChaosClothAssetEditor/Source/ChaosClothAssetDataflowNodes/Private/ChaosClothAsset/ImportFilePathCustomization.cpp
@@ -0,0 +1,165 @@
+// Copyright Epic Games, Inc. All Rights Reserved.
+
+#include "ChaosClothAsset/ImportFilePathCustomization.h"
+#include "Misc/MessageDialog.h"
+#include "Misc/PackageName.h"
+#include "Widgets/Images/SImage.h"
+#include "Widgets/Input/SButton.h"
+#include "Widgets/Input/SFilePathPicker.h"
+#include "DetailWidgetRow.h"
+#include "EditorDirectories.h"
+#include "PropertyHandle.h"
+
+#define LOCTEXT_NAMESPACE "ImportFilePathCustomization"
+
+namespace UE::Chaos::ClothAsset
+{
+	TSharedRef<IPropertyTypeCustomization> FImportFilePathCustomization::MakeInstance()
+	{
+		return MakeShareable(new FImportFilePathCustomization);
+	}
+
+	void FImportFilePathCustomization::CustomizeHeader(TSharedRef<IPropertyHandle> StructPropertyHandle, FDetailWidgetRow& HeaderRow, IPropertyTypeCustomizationUtils& CustomizationUtils)
+	{
+		PathStringProperty = StructPropertyHandle->GetChildHandle(TEXT("FilePath"));
+		ForceReimport = StructPropertyHandle->GetChildHandle(TEXT("bForceReimport"));
+
+		// Construct file type filter
+		FString FileTypeFilter;
+
+		const FString& MetaData = StructPropertyHandle->GetMetaData(TEXT("FilePathFilter"));
+
+		bLongPackageName = StructPropertyHandle->HasMetaData(TEXT("LongPackageName"));
+		bRelativeToGameDir = StructPropertyHandle->HasMetaData(TEXT("RelativeToGameDir"));
+
+		if (MetaData.IsEmpty())
+		{
+			FileTypeFilter = TEXT("All files (*.*)|*.*");
+		}
+		else
+		{
+			if (MetaData.Contains(TEXT("|"))) // If MetaData follows the Description|ExtensionList format, use it as is
+			{
+				FileTypeFilter = MetaData;
+			}
+			else
+			{
+				FileTypeFilter = FString::Printf(TEXT("%s files (*.%s)|*.%s"), *MetaData, *MetaData, *MetaData);
+			}
+		}
+
+		// Create path picker widget
+		HeaderRow
+		.NameContent()
+		[
+			StructPropertyHandle->CreatePropertyNameWidget()
+		]
+		.ValueContent()
+		.MinDesiredWidth(250)
+		.MaxDesiredWidth(0)
+		[
+			SNew(SHorizontalBox)
+			+ SHorizontalBox::Slot()
+			[
+				SNew(SFilePathPicker)
+				.BrowseButtonImage(FAppStyle::GetBrush("PropertyWindow.Button_Ellipsis"))
+				.BrowseButtonStyle(FAppStyle::Get(), "HoverHintOnly")
+				.BrowseButtonToolTip(LOCTEXT("FileButtonToolTipText", "Choose a file from this computer"))
+				.BrowseDirectory(FEditorDirectories::Get().GetLastDirectory(ELastDirectory::GENERIC_OPEN))
+				.BrowseTitle(LOCTEXT("PropertyEditorTitle", "File picker..."))
+				.FilePath(this, &FImportFilePathCustomization::HandleFilePathPickerFilePath)
+				.FileTypeFilter(FileTypeFilter)
+				.OnPathPicked(this, &FImportFilePathCustomization::HandleFilePathPickerPathPicked)
+			]
+			+ SHorizontalBox::Slot()
+			.AutoWidth()
+			[
+				SNew(SButton)
+				.ButtonStyle(FAppStyle::Get(), "HoverHintOnly")
+				.ContentPadding(FMargin(10.f, 0))
+				.HAlign(HAlign_Right)
+				.ToolTipText(LOCTEXT("Button_ReimportAsset_Tooltip", "Reimport asset"))
+				.OnClicked(this, &FImportFilePathCustomization::OnClicked)
+				[
+					SNew(SImage)
+					.Image(FAppStyle::GetBrush("Persona.ReimportAsset"))
+				]
+			]
+		];
+	}
+
+	FReply FImportFilePathCustomization::OnClicked()
+	{
+		if (ForceReimport)
+		{
+			ForceReimport->SetValue(true);
+		}
+		return FReply::Handled();
+	}
+
+	// Callbacks lifted from FFilePathStructCustomization
+	FString FImportFilePathCustomization::HandleFilePathPickerFilePath() const
+	{
+		FString FilePath;
+		if (PathStringProperty)
+		{
+			PathStringProperty->GetValue(FilePath);
+		}
+		return FilePath;
+	}
+
+	void FImportFilePathCustomization::HandleFilePathPickerPathPicked(const FString& PickedPath)
+	{
+		FString FinalPath = PickedPath;
+		if (bLongPackageName)
+		{
+			FString LongPackageName;
+			FString StringFailureReason;
+			if (FPackageName::TryConvertFilenameToLongPackageName(PickedPath, LongPackageName, &StringFailureReason) == false)
+			{
+				FMessageDialog::Open(EAppMsgType::Ok, FText::FromString(StringFailureReason));
+			}
+			FinalPath = LongPackageName;
+		}
+		else if (bRelativeToGameDir && !PickedPath.IsEmpty())
+		{
+			//A filepath under the project directory will be made relative to the project directory
+			//Otherwise, the absolute path will be returned unless it doesn't exist, the current path will
+			//be kept. This can happen if it's already relative to project dir (tabbing when selected)
+
+			const FString ProjectDir = FPaths::ProjectDir();
+			const FString AbsoluteProjectDir = FPaths::ConvertRelativePathToFull(ProjectDir);
+			const FString AbsolutePickedPath = FPaths::ConvertRelativePathToFull(PickedPath);
+
+			//Verify if absolute path to file exists. If it was already relative to content directory
+			//the absolute will be to binaries and will possibly be garbage
+			if (FPaths::FileExists(AbsolutePickedPath))
+			{
+				//If file is part of the project dir, chop the project dir part
+				//Otherwise, use the absolute path
+				if (AbsolutePickedPath.StartsWith(AbsoluteProjectDir))
+				{
+					FinalPath = AbsolutePickedPath.RightChop(AbsoluteProjectDir.Len());
+				}
+				else
+				{
+					FinalPath = AbsolutePickedPath;
+				}
+			}
+			else
+			{
+				//If absolute file doesn't exist, it might already be relative to project dir
+				//If not, then it might be a manual entry, so keep it untouched either way
+				FinalPath = PickedPath;
+			}
+		}
+
+		if (PathStringProperty)
+		{
+			PathStringProperty->SetValue(FinalPath);
+		}
+		FEditorDirectories::Get().SetLastDirectory(ELastDirectory::GENERIC_OPEN, FPaths::GetPath(PickedPath));
+	}
+}  // End namespace UE::Chaos::ClothAsset
+
+#undef LOCTEXT_NAMESPACE
diff --git a/Engine/Plugins/Experimental/ChaosClothAssetEditor/Source/ChaosClothAssetDataflowNodes/Private/ChaosClothAsset/ImportFilePathCustomization.h b/Engine/Plugins/Experimental/ChaosClothAssetEditor/Source/ChaosClothAssetDataflowNodes/Private/ChaosClothAsset/ImportFilePathCustomization.h
new file mode 100644
--- /dev/null
+++ b/Engine/Plugins/Experimental/ChaosClothAssetEditor/Source/ChaosClothAssetDataflowNodes/Private/ChaosClothAsset/ImportFilePathCustomization.h
@@ -0,0 +1,35 @@
+// Copyright Epic Games, Inc. All Rights Reserved.
+
+#pragma once
+
+#include "Input/Reply.h"
+#include "IPropertyTypeCustomization.h"
+
+class SButton;
+
+namespace UE::Chaos::ClothAsset
+{
+	/**
+	 * Customization for import file path, modelled after FFilePathStructCustomization with the addition 
+	 * of a reimport button whenever a bForceReimport property is present in the customized struct.
+	 */
+	class FImportFilePathCustomization : public IPropertyTypeCustomization
+	{
+	public:
+		static TSharedRef<IPropertyTypeCustomization> MakeInstance();
+
+		virtual void CustomizeHeader(TSharedRef<IPropertyHandle> StructPropertyHandle, FDetailWidgetRow& HeaderRow, IPropertyTypeCustomizationUtils& CustomizationUtils) override;
+		virtual void CustomizeChildren(TSharedRef<IPropertyHandle> /*StructPropertyHandle*/, IDetailChildrenBuilder& /*ChildBuilder*/, IPropertyTypeCustomizationUtils& /*CustomizationUtils*/) override {}
+
+	private:
+		FReply OnClicked();
+
+		FString HandleFilePathPickerFilePath() const;
+		void HandleFilePathPickerPathPicked(const FString& PickedPath);
+
+		TSharedPtr<IPropertyHandle> PathStringProperty;
+		TSharedPtr<IPropertyHandle> ForceReimport;
+		bool bLongPackageName;
+		bool bRelativeToGameDir;
+	};
+}
diff --git a/Engine/Plugins/Experimental/ChaosClothAssetEditor/Source/ChaosClothAssetDataflowNodes/Private/ChaosClothAsset/StaticMeshImportNode.cpp b/Engine/Plugins/Experimental/ChaosClothAssetEditor/Source/ChaosClothAssetDataflowNodes/Private/ChaosClothAsset/StaticMeshImportNode.cpp
--- a/Engine/Plugins/Experimental/ChaosClothAssetEditor/Source/ChaosClothAssetDataflowNodes/Private/ChaosClothAsset/StaticMeshImportNode.cpp
+++ b/Engine/Plugins/Experimental/ChaosClothAssetEditor/Source/ChaosClothAssetDataflowNodes/Private/ChaosClothAsset/StaticMeshImportNode.cpp
@@ -3,18 +3,13 @@
 #include "ChaosClothAsset/StaticMeshImportNode.h"
 #include "ChaosClothAsset/ClothDataflowTools.h"
 #include "ChaosClothAsset/CollectionClothFacade.h"
-#include "ChaosClothAsset/ClothAsset.h"
 #include "ChaosClothAsset/ClothGeometryTools.h"
 #include "Dataflow/DataflowInputOutput.h"
 #include "Engine/StaticMesh.h"
-#include "Interfaces/ITargetPlatformManagerModule.h"
 #include "MeshDescriptionToDynamicMesh.h"
 #include "SkeletalMeshAttributes.h"
-#include "Rendering/SkeletalMeshLODImporterData.h"
 #include "Rendering/SkeletalMeshLODModel.h"
 #include "Materials/MaterialInterface.h"
-#include "MeshUtilities.h"
-#include "Modules/ModuleManager.h"
 
 #include UE_INLINE_GENERATED_CPP_BY_NAME(StaticMeshImportNode)
 
@@ -22,102 +17,30 @@
 
 namespace UE::Chaos::ClothAsset::Private
 {
-
-void CopyBuildSettings(
-	const FMeshBuildSettings& InStaticMeshBuildSettings,
-	FSkeletalMeshBuildSettings& OutSkeletalMeshBuildSettings
-)
-{
-	OutSkeletalMeshBuildSettings.bRecomputeNormals = InStaticMeshBuildSettings.bRecomputeNormals;
-	OutSkeletalMeshBuildSettings.bRecomputeTangents = InStaticMeshBuildSettings.bRecomputeTangents;
-	OutSkeletalMeshBuildSettings.bUseMikkTSpace = InStaticMeshBuildSettings.bUseMikkTSpace;
-	OutSkeletalMeshBuildSettings.bComputeWeightedNormals = InStaticMeshBuildSettings.bComputeWeightedNormals;
-	OutSkeletalMeshBuildSettings.bRemoveDegenerates = InStaticMeshBuildSettings.bRemoveDegenerates;
-	OutSkeletalMeshBuildSettings.bUseHighPrecisionTangentBasis = InStaticMeshBuildSettings.bUseHighPrecisionTangentBasis;
-	OutSkeletalMeshBuildSettings.bUseFullPrecisionUVs = InStaticMeshBuildSettings.bUseFullPrecisionUVs;
-	OutSkeletalMeshBuildSettings.bUseBackwardsCompatibleF16TruncUVs = InStaticMeshBuildSettings.bUseBackwardsCompatibleF16TruncUVs;
-	// The rest we leave at defaults.
-}
-
-bool BuildSkeletalMeshModelFromMeshDescription(const FMeshDescription* const InMeshDescription, const FMeshBuildSettings& InBuildSettings, FSkeletalMeshLODModel& SkeletalMeshModel)
-{
-	// This is following StaticToSkeletalMeshConverter.cpp::AddLODFromStaticMeshSourceModel
-	FSkeletalMeshBuildSettings BuildSettings;
-	CopyBuildSettings(InBuildSettings, BuildSettings);
-	FMeshDescription SkeletalMeshGeometry = *InMeshDescription;
-	FSkeletalMeshAttributes SkeletalMeshAttributes(SkeletalMeshGeometry);
-	SkeletalMeshAttributes.Register();
-
-	// Full binding to the root bone.
-	constexpr int32 RootBoneIndex = 0;
-	FSkinWeightsVertexAttributesRef SkinWeights = SkeletalMeshAttributes.GetVertexSkinWeights();
-	UE::AnimationCore::FBoneWeight RootInfluence(RootBoneIndex, 1.0f);
-	UE::AnimationCore::FBoneWeights RootBinding = UE::AnimationCore::FBoneWeights::Create({ RootInfluence });
-
-	for (const FVertexID VertexID : SkeletalMeshGeometry.Vertices().GetElementIDs())
-	{
-		SkinWeights.Set(VertexID, RootBinding);
-	}
-
-	FSkeletalMeshImportData SkeletalMeshImportGeometry = FSkeletalMeshImportData::CreateFromMeshDescription(SkeletalMeshGeometry);
-	// Data needed by BuildSkeletalMesh
-	TArray<FVector3f> LODPoints;
-	TArray<SkeletalMeshImportData::FMeshWedge> LODWedges;
-	TArray<SkeletalMeshImportData::FMeshFace> LODFaces;
-	TArray<SkeletalMeshImportData::FVertInfluence> LODInfluences;
-	TArray<int32> LODPointToRawMap;
-	SkeletalMeshImportGeometry.CopyLODImportData(LODPoints, LODWedges, LODFaces, LODInfluences, LODPointToRawMap);
-	IMeshUtilities::MeshBuildOptions BuildOptions;
-	BuildOptions.TargetPlatform = GetTargetPlatformManagerRef().GetRunningTargetPlatform();
-	BuildOptions.FillOptions(BuildSettings);
-
-	static const FString SkeletalMeshName("ClothAssetStaticMeshImportConvert"); // This is only used by warning messages in the mesh builder.
-	// Build a RefSkeleton with just a root bone. The BuildSkeletalMesh code expects you have a reference skeleton with at least one bone to work.
-	FReferenceSkeleton RootBoneRefSkeleton;
-	FReferenceSkeletonModifier SkeletonModifier(RootBoneRefSkeleton, nullptr);
-	FMeshBoneInfo RootBoneInfo;
-	RootBoneInfo.Name = FName("Root");
-	SkeletonModifier.Add(RootBoneInfo, FTransform());
-	RootBoneRefSkeleton.RebuildRefSkeleton(nullptr, true);
-
-	IMeshUtilities& MeshUtilities = FModuleManager::Get().LoadModuleChecked<IMeshUtilities>("MeshUtilities");
-	TArray<FText> WarningMessages;
-	if (!MeshUtilities.BuildSkeletalMesh(SkeletalMeshModel, SkeletalMeshName, RootBoneRefSkeleton, LODInfluences, LODWedges, LODFaces, LODPoints, LODPointToRawMap, BuildOptions, &WarningMessages))
+	static bool InitializeDataFromMeshDescription(
+		const UStaticMesh* const StaticMesh,
+		const FMeshDescription* const InMeshDescription,
+		const FMeshBuildSettings& InBuildSettings,
+		const TArray<FStaticMaterial>& StaticMaterials,
+		const TSharedRef<FManagedArrayCollection>& ClothCollection)
 	{
-		for (const FText& Message : WarningMessages)
+		FSkeletalMeshLODModel SkeletalMeshModel;
+		if (FClothDataflowTools::BuildSkeletalMeshModelFromMeshDescription(InMeshDescription, InBuildSettings, SkeletalMeshModel))
 		{
-			UE_LOG(LogChaosClothAssetDataflowNodes, Warning, TEXT("%s"), *Message.ToString());
+			FStaticMeshConstAttributes MeshAttributes(*InMeshDescription);
+			TPolygonGroupAttributesConstRef<FName> MaterialSlotNames = MeshAttributes.GetPolygonGroupMaterialSlotNames();
+			for (int32 SectionIndex = 0; SectionIndex < SkeletalMeshModel.Sections.Num(); ++SectionIndex)
+			{
+				// Section MaterialIndex refers to the polygon group index. Look up which material this corresponds with.
+				const FName& MaterialSlotName = MaterialSlotNames[SkeletalMeshModel.Sections[SectionIndex].MaterialIndex];
+				const int32 MaterialIndex = StaticMesh->GetMaterialIndexFromImportedMaterialSlotName(MaterialSlotName);
+				const FString RenderMaterialPathName = StaticMaterials.IsValidIndex(MaterialIndex)&& StaticMaterials[MaterialIndex].MaterialInterface ? StaticMaterials[MaterialIndex].MaterialInterface->GetPathName() : "";
+				FClothDataflowTools::AddRenderPatternFromSkeletalMeshSection(ClothCollection, SkeletalMeshModel, SectionIndex, RenderMaterialPathName);
+			}
+			return true;
 		}
 		return false;
 	}
-	return true;
-}
-
-bool InitializeDataFromMeshDescription(
-	const UStaticMesh* const StaticMesh,
-	const FMeshDescription* const InMeshDescription,
-	const FMeshBuildSettings& InBuildSettings,
-	const TArray<FStaticMaterial>& StaticMaterials,
-	const TSharedRef<FManagedArrayCollection>& ClothCollection)
-{
-	FSkeletalMeshLODModel SkeletalMeshModel;
-	if (BuildSkeletalMeshModelFromMeshDescription(InMeshDescription, InBuildSettings, SkeletalMeshModel))
-	{
-		FStaticMeshConstAttributes MeshAttributes(*InMeshDescription);
-		TPolygonGroupAttributesConstRef<FName> MaterialSlotNames = MeshAttributes.GetPolygonGroupMaterialSlotNames();
-		for (int32 SectionIndex = 0; SectionIndex < SkeletalMeshModel.Sections.Num(); ++SectionIndex)
-		{
-			// Section MaterialIndex refers to the polygon group index. Look up which material this corresponds with.
-			const FName& MaterialSlotName = MaterialSlotNames[SkeletalMeshModel.Sections[SectionIndex].MaterialIndex];
-			const int32 MaterialIndex = StaticMesh->GetMaterialIndexFromImportedMaterialSlotName(MaterialSlotName);
-			const FString RenderMaterialPathName = StaticMaterials.IsValidIndex(MaterialIndex)&& StaticMaterials[MaterialIndex].MaterialInterface ? StaticMaterials[MaterialIndex].MaterialInterface->GetPathName() : "";
-			FClothDataflowTools::AddRenderPatternFromSkeletalMeshSection(ClothCollection, SkeletalMeshModel, SectionIndex, RenderMaterialPathName);
-		}
-		return true;
-	}
-	return false;
-}
-
 } // namespace UE::Chaos::ClothAsset::Private
 
 FChaosClothAssetStaticMeshImportNode::FChaosClothAssetStaticMeshImportNode(const Dataflow::FNodeParameters& InParam, FGuid InGuid)
diff --git a/Engine/Plugins/Experimental/ChaosClothAssetEditor/Source/ChaosClothAssetDataflowNodes/Private/ChaosClothAsset/USDImportNode.cpp b/Engine/Plugins/Experimental/ChaosClothAssetEditor/Source/ChaosClothAssetDataflowNodes/Private/ChaosClothAsset/USDImportNode.cpp
new file mode 100644
--- /dev/null
+++ b/Engine/Plugins/Experimental/ChaosClothAssetEditor/Source/ChaosClothAssetDataflowNodes/Private/ChaosClothAsset/USDImportNode.cpp
@@ -0,0 +1,740 @@
+// Copyright Epic Games, Inc. All Rights Reserved.
+
+#include "ChaosClothAsset/USDImportNode.h"
+#include "ChaosClothAsset/ClothAsset.h"
+#include "ChaosClothAsset/ClothDataflowTools.h"
+#include "ChaosClothAsset/CollectionClothFacade.h"
+#include "AssetRegistry/AssetData.h"
+#include "AssetRegistry/AssetRegistryModule.h"
+#include "Dataflow/DataflowInputOutput.h"
+#include "Dataflow/DataflowGraph.h"
+#include "Dataflow/DataflowObject.h"
+#include "Engine/StaticMesh.h"
+#include "HAL/FileManager.h"
+#include "Materials/Material.h"
+#include "Materials/MaterialInstanceConstant.h"
+#include "Misc/Paths.h"
+#include "Misc/ScopedSlowTask.h"
+#include "Rendering/SkeletalMeshLODModel.h"
+#include "StaticMeshAttributes.h"
+#include "UObject/Package.h"
+#include "UsdWrappers/UsdPrim.h"
+#include "USDStageImporter.h"
+#include "USDStageImportContext.h"
+#include "USDStageImportOptions.h"
+#include "USDTypesConversion.h"
+#include "USDValueConversion.h"
+
+#include UE_INLINE_GENERATED_CPP_BY_NAME(USDImportNode)
+
+#define LOCTEXT_NAMESPACE "ChaosClothAssetUSDImportNode"
+
+namespace UE::Chaos::ClothAsset::Private
+{
+	static bool RemoveDegenerateTriangles(
+		const TArray<FIntVector3>& TriangleToVertexIndex,
+		const TArray<FVector2f>& RestPositions2D,
+		const TArray<FVector3f>& DrapedPositions3D,
+		TArray<FIntVector3>& OutTriangleToVertexIndex,
+		TArray<FVector2f>& OutRestPositions2D,
+		TArray<FVector3f>& OutDrapedPositions3D,
+		TArray<int32>& OutIndices)  // Old to new vertices lookup
+	{
+		bool bHasDegenerateTriangles = false;
+
+		check(RestPositions2D.Num() == DrapedPositions3D.Num());
+		const int32 VertexCount = RestPositions2D.Num();
+		const int32 TriangleCount = TriangleToVertexIndex.Num();
+
+		OutTriangleToVertexIndex.Reset(TriangleToVertexIndex.Num());
+
+		// Remap[Index] is the index of the first vertex in a group of degenerated triangles to be callapsed.
+		// When two groups of collapsed vertices are merged, the group with the greatest Remap[index] value must adopt the one from the other group.
+		// For Example:
+		// 1. For all i, Remap[i] = i
+		// 2. Finds one degenerated triangle (7, 9, 4) with collapsed edges (7, 9), (9, 4), and (7, 4) -> Remap[4] = 4, Remap[7] = 4, and Remap[9] = 4
+		// 3. Finds another degenerated triangle (2, 3, 4) with collapsed edges (2, 4) -> Remap[2] = 2, Remap[4] = 2, Remap[7] = 2, and Remap[9] = 2
+		TArray<int32> Remap;
+		Remap.SetNumUninitialized(VertexCount);
+
+		for (int32 Index = 0; Index < VertexCount; ++Index)
+		{
+			Remap[Index] = Index;
+		}
+
+		int32 OutVertexCount = VertexCount;
+
+		auto RemapAndPropagateIndex = [&Remap, &OutVertexCount](int32 Index0, int32 Index1)
+			{
+				if (Remap[Index0] != Remap[Index1])
+				{
+					if (Remap[Index0] > Remap[Index1])  // Always remap from the lowest index to ensure the earlier index is always kept
+					{
+						Swap(Index0, Index1);
+					}
+					// Merge groups with this new first index Remap[Index0]
+					const int32 PrevRemapIndex = Remap[Index1];
+					for (int32 Index = PrevRemapIndex; Index < Remap.Num(); ++Index)  // Only need to start from the first index of the group to merge
+					{
+						if (Remap[Index] == PrevRemapIndex)
+						{
+							Remap[Index] = Remap[Index0];
+						}
+					}
+					--OutVertexCount;
+				}
+			};
+
+		for (int32 TriangleIndex = 0; TriangleIndex < TriangleCount; ++TriangleIndex)
+		{
+			const int32 Index0 = TriangleToVertexIndex[TriangleIndex][0];
+			const int32 Index1 = TriangleToVertexIndex[TriangleIndex][1];
+			const int32 Index2 = TriangleToVertexIndex[TriangleIndex][2];
+
+			const FVector3f& P0 = DrapedPositions3D[Index0];
+			const FVector3f& P1 = DrapedPositions3D[Index1];
+			const FVector3f& P2 = DrapedPositions3D[Index2];
+			const FVector3f P0P1 = P1 - P0;
+			const FVector3f P0P2 = P2 - P0;
+
+			const float TriNormSizeSquared = (P0P1 ^ P0P2).SizeSquared();
+			if (TriNormSizeSquared <= UE_SMALL_NUMBER)
+			{
+				const FVector3f P1P2 = P2 - P1;
+
+				if (P0P1.SquaredLength() <= UE_SMALL_NUMBER)
+				{
+					RemapAndPropagateIndex(Index0, Index1);
+				}
+				if (P0P2.SquaredLength() <= UE_SMALL_NUMBER)
+				{
+					RemapAndPropagateIndex(Index0, Index2);
+				}
+				if (P1P2.SquaredLength() <= UE_SMALL_NUMBER)
+				{
+					RemapAndPropagateIndex(Index1, Index2);
+				}
+			}
+			else
+			{
+				OutTriangleToVertexIndex.Emplace(TriangleToVertexIndex[TriangleIndex]);
+			}
+		}
+
+		const int32 OutTriangleCount = OutTriangleToVertexIndex.Num();
+		bHasDegenerateTriangles = (TriangleCount != OutTriangleCount);
+
+		UE_CLOG(bHasDegenerateTriangles, LogChaosClothAssetDataflowNodes, Warning,
+			TEXT("USD import found %d degenerated triangles out of %d"), TriangleCount - OutTriangleCount, TriangleCount);
+
+		OutRestPositions2D.Reset(OutVertexCount);
+		OutDrapedPositions3D.Reset(OutVertexCount);
+		OutIndices.Reset(VertexCount);
+		int32 OutIndex = -1;
+
+		for (int32 VertexIndex = 0; VertexIndex < VertexCount; ++VertexIndex)
+		{
+			if (Remap[VertexIndex] == VertexIndex)
+			{
+				OutRestPositions2D.Add(RestPositions2D[Remap[VertexIndex]]);
+				OutDrapedPositions3D.Add(DrapedPositions3D[Remap[VertexIndex]]);
+				OutIndices.Add(++OutIndex);
+			}
+			else
+			{
+				const int32 OutRemappedIndex = OutIndices[Remap[VertexIndex]];
+				OutIndices.Add(OutRemappedIndex);
+			}
+		}
+		ensure(OutIndex + 1 == OutVertexCount);
+
+		for (int32 TriangleIndex = 0; TriangleIndex < OutTriangleCount; ++TriangleIndex)
+		{
+			int32& Index0 = OutTriangleToVertexIndex[TriangleIndex][0];
+			int32& Index1 = OutTriangleToVertexIndex[TriangleIndex][1];
+			int32& Index2 = OutTriangleToVertexIndex[TriangleIndex][2];
+
+			Index0 = OutIndices[Index0];
+			Index1 = OutIndices[Index1];
+			Index2 = OutIndices[Index2];
+
+			checkSlow(Index0 != Index1);
+			checkSlow(Index0 != Index2);
+			checkSlow(Index1 != Index2);
+			checkSlow((OutDrapedPositions3D[Index0] - OutDrapedPositions3D[Index1]).SquaredLength() > UE_SMALL_NUMBER);
+			checkSlow((OutDrapedPositions3D[Index0] - OutDrapedPositions3D[Index2]).SquaredLength() > UE_SMALL_NUMBER);
+			checkSlow((OutDrapedPositions3D[Index1] - OutDrapedPositions3D[Index2]).SquaredLength() > UE_SMALL_NUMBER);
+		}
+
+		return bHasDegenerateTriangles;
+	}
+}  // End namespace UE::Chaos::ClothAsset::Private
+
+FChaosClothAssetUSDImportNode::FChaosClothAssetUSDImportNode(const Dataflow::FNodeParameters& InParam, FGuid InGuid)
+	: FDataflowTerminalNode(InParam, InGuid)
+{
+	using namespace UE::Chaos::ClothAsset;
+
+	// Initialize to a valid collection cache
+	const TSharedRef<FManagedArrayCollection> ClothCollection = MakeShared<FManagedArrayCollection>(MoveTemp(CollectionCache));
+	FCollectionClothFacade ClothFacade(ClothCollection);
+	ClothFacade.DefineSchema();
+	CollectionCache = MoveTemp(*ClothCollection);
+
+	// Register connections
+	RegisterOutputConnection(&Collection);
+}
+
+void FChaosClothAssetUSDImportNode::SetAssetValue(TObjectPtr<UObject> Asset, Dataflow::FContext& Context) const
+{
+	using namespace UE::Chaos::ClothAsset;
+
+	if (UChaosClothAsset* const ClothAsset = Cast<UChaosClothAsset>(Asset.Get()))
+	{
+PRAGMA_DISABLE_DEPRECATION_WARNINGS  // TODO: Don't use public property, and have Getter/Setter API instead
+		if (UDataflow* const DataflowAsset = ClothAsset->DataflowAsset)
+PRAGMA_ENABLE_DEPRECATION_WARNINGS
+		{
+			const TSharedPtr<Dataflow::FGraph, ESPMode::ThreadSafe> Dataflow = DataflowAsset->GetDataflow();
+			if (const TSharedPtr<FDataflowNode> BaseNode = Dataflow->FindBaseNode(this->GetGuid()))  // This is basically a safe const_cast
+			{
+				FChaosClothAssetUSDImportNode* const MutableThis = static_cast<FChaosClothAssetUSDImportNode*>(BaseNode.Get());
+				check(MutableThis == this);
+
+				const FString& FilePath = UsdFile.FilePath;
+				const FMD5Hash InFileHash = FilePath.IsEmpty() ?
+					FMD5Hash() :  // Reset to an empty import
+					FPaths::FileExists(FilePath) ?
+					FMD5Hash::HashFile(*FilePath) :  // Update cache
+					FileHash;  // Keep the current cache
+
+				if (FileHash != InFileHash || UsdFile.bForceReimport)
+				{
+					MutableThis->FileHash = InFileHash;
+					MutableThis->UsdFile.bForceReimport = false;
+
+					const FString AssetPath = Asset->GetPackage()->GetPathName();
+
+					FText ErrorText;
+					if (!MutableThis->ImportFromFile(FilePath, AssetPath, ErrorText))
+					{
+						FClothDataflowTools::LogAndToastWarning(*this,
+							LOCTEXT("FailedToExportUsdFileHeadline", "Failed to import USD file from file."),
+							FText::Format(LOCTEXT("FailedToExportUsdDetails", "Error while importing USD cloth from file '{0}':\n{1}"), FText::FromString(FilePath), ErrorText));
+					}
+				}
+			}
+		}
+	}
+}
+
+void FChaosClothAssetUSDImportNode::Evaluate(Dataflow::FContext& Context, const FDataflowOutput* Out) const
+{
+	using namespace UE::Chaos::ClothAsset;
+
+	if (Out->IsA<FManagedArrayCollection>(&Collection))
+	{
+		const TSharedRef<FManagedArrayCollection> ClothCollection = MakeShared<FManagedArrayCollection>();
+
+		// Import from cache
+		FText ErrorText;
+		if (!ImportFromCache(ClothCollection, ErrorText))
+		{
+			FClothDataflowTools::LogAndToastWarning(*this,
+				LOCTEXT("FailedToExportUsdCacheHeadline", "Failed to import USD file from cache."),
+				FText::Format(LOCTEXT("FailedToExportUsdCacheDetails", "Error while importing USD cloth from cache '{0}':\n{1}"), FText::FromString(UsdFile.FilePath), ErrorText));
+		}
+
+		SetValue(Context, MoveTemp(*ClothCollection), &Collection);
+	}
+}
+
+void FChaosClothAssetUSDImportNode::Serialize(FArchive& Archive)
+{
+	using namespace UE::Chaos::ClothAsset;
+
+	::Chaos::FChaosArchive ChaosArchive(Archive);
+	CollectionCache.Serialize(ChaosArchive);
+	if (Archive.IsLoading())
+	{
+		// Make sure to always have a valid cloth collection on reload, some new attributes could be missing from the cached collection
+		const TSharedRef<FManagedArrayCollection> ClothCollection = MakeShared<FManagedArrayCollection>(MoveTemp(CollectionCache));
+		FCollectionClothFacade ClothFacade(ClothCollection);
+		if (!ClothFacade.IsValid())
+		{
+			ClothFacade.DefineSchema();
+		}
+		CollectionCache = MoveTemp(*ClothCollection);
+	}
+
+	Archive << FileHash;
+}
+
+bool FChaosClothAssetUSDImportNode::ImportFromFile(const FString& UsdFilePath, const FString& AssetPath, FText& OutErrorText)
+{
+#if USE_USD_SDK
+	using namespace UE::Chaos::ClothAsset;
+
+	// Reset cache
+	CollectionCache.Reset();
+	PackagePath = FString();
+
+	// Steal the cache's collection and make it a valid cloth collection
+	const TSharedRef<FManagedArrayCollection> ClothCollection = MakeShared<FManagedArrayCollection>(MoveTemp(CollectionCache));
+	FCollectionClothFacade ClothFacade(ClothCollection);
+	ClothFacade.DefineSchema();
+
+	// Make sure to give the collection back to the cache when reaching the end of scope
+	ON_SCOPE_EXIT { CollectionCache = MoveTemp(*ClothCollection); };
+
+	// Empty file
+	if (UsdFilePath.IsEmpty())
+	{
+		return true;
+	}
+
+	// Start slow task
+	constexpr float NumSteps = 4.f;
+	FScopedSlowTask SlowTask(NumSteps, LOCTEXT("ImportingUSDFile", "Importing USD file..."));
+	SlowTask.MakeDialogDelayed(1.f);
+
+	// Open stage
+	constexpr bool bUseStageCache = false;  // Reload from disk, not from cache
+	constexpr EUsdInitialLoadSet UsdInitialLoadSet = EUsdInitialLoadSet::LoadAll;  // TODO: Ideally we should only use LoadNone to start with and load what's needed once the Schema is defined
+
+	const UE::FUsdStage UsdStage = UnrealUSDWrapper::OpenStage(*UsdFilePath, UsdInitialLoadSet, bUseStageCache);
+	if (!UsdStage)
+	{
+		OutErrorText = LOCTEXT("CantCreateNewStage", "Failed to open the specified USD file.");
+		return false;
+	}
+	SlowTask.EnterProgressFrame(1.f);
+
+	// Update import location
+	const uint32 UsdPathHash = GetTypeHash(UsdFile.FilePath);
+	const FString UsdFileName = SlugStringForValidName(FPaths::GetBaseFilename(UsdFile.FilePath));
+	const FString PackageName = FString::Printf(TEXT("%s_%08X"), *UsdFileName, UsdPathHash);
+	PackagePath = FPaths::Combine(AssetPath + TEXT("_Import"), PackageName);
+
+	// Get stage infos
+	const FUsdStageInfo StageInfo(UsdStage);
+
+	// Import recognised assets
+	FUsdStageImportContext ImportContext;
+
+	const TObjectPtr<UUsdStageImportOptions>& ImportOptions = ImportContext.ImportOptions;
+	{
+		check(ImportOptions);
+		// Data to import
+		ImportOptions->bImportActors = false;
+		ImportOptions->bImportGeometry = true;
+		ImportOptions->bImportSkeletalAnimations = false;
+		ImportOptions->bImportLevelSequences = false;
+		ImportOptions->bImportMaterials = true;
+		ImportOptions->bImportGroomAssets = false;
+		ImportOptions->bImportOnlyUsedMaterials = true;
+		// Prims to import
+		ImportOptions->PrimsToImport = TArray<FString>{ TEXT( "/" ) };
+		// USD options
+		ImportOptions->PurposesToImport = (int32)EUsdPurpose::Proxy;
+		ImportOptions->NaniteTriangleThreshold = TNumericLimits<int32>::Max();  // Don't enable Nanite
+		ImportOptions->RenderContextToImport = NAME_None;
+		ImportOptions->MaterialPurpose = NAME_None;  // *UnrealIdentifiers::MaterialPreviewPurpose ???
+		ImportOptions->RootMotionHandling = EUsdRootMotionHandling::NoAdditionalRootMotion;
+		ImportOptions->SubdivisionLevel = 0;
+		ImportOptions->bOverrideStageOptions = true;
+		ImportOptions->StageOptions.MetersPerUnit = 0.001f;  // Assume mm scale for now until this is set in the file
+		ImportOptions->StageOptions.UpAxis = StageInfo.UpAxis;
+		ImportOptions->bImportAtSpecificTimeCode = false;
+		ImportOptions->ImportTimeCode = 0.f;
+		// Groom
+		ImportOptions->GroomInterpolationSettings = TArray<FHairGroupsInterpolation>();
+		// Collision
+		ImportOptions->ExistingActorPolicy = EReplaceActorPolicy::Replace;
+		ImportOptions->ExistingAssetPolicy = EReplaceAssetPolicy::Replace;
+		// Processing
+		ImportOptions->bPrimPathFolderStructure = false;
+		ImportOptions->KindsToCollapse = true;  // TODO: Might need not to merge all static meshes
+		ImportOptions->bMergeIdenticalMaterialSlots = true;
+		ImportOptions->bInterpretLODs = false;
+	}
+
+	constexpr bool bIsAutomated = true;
+	constexpr bool bIsReimport = false;
+	constexpr bool bAllowActorImport = false;
+
+	ImportContext.Stage = UsdStage;  // Set the stage first to prevent re-opening it in the Init function
+	ImportContext.Init(TEXT(""), UsdFilePath, PackagePath, RF_NoFlags, bIsAutomated, bIsReimport, bAllowActorImport);
+
+	UUsdStageImporter UsdStageImporter;
+	UsdStageImporter.ImportFromFile(ImportContext);
+
+	SlowTask.EnterProgressFrame(2.f);
+
+	// Import sim mesh into collection cache 
+	// TODO: Until we have a schema so that we can use the asset cache and remove the collection cache
+	const FUsdStageOptions& StageOptions = ImportOptions->StageOptions;
+	const int AxesOrder[] = { 0, (StageOptions.UpAxis == EUsdUpAxis::ZAxis) ? 1 : 2, (StageOptions.UpAxis == EUsdUpAxis::ZAxis) ? 2 : 1 };
+	const int WindingOrder[] = { 0, (StageOptions.UpAxis == EUsdUpAxis::ZAxis) ? 2 : 1, (StageOptions.UpAxis == EUsdUpAxis::ZAxis) ? 1 : 2 };
+	const float CentimetersPerUnit = StageOptions.MetersPerUnit * 100.f;
+
+	// Sewings
+	TArray<TArray<FIntVector2>> SeamStitches;
+	TArray<TArray<FIntVector2>> SeamPatterns;
+
+	const UE::FSdfPath SewingsPath = UE::FSdfPath(UE::FSdfPath::AbsoluteRootPath()).AppendChild(TEXT("SimulationData")).AppendChild(TEXT("Sewings"));
+	if (const UE::FUsdPrim SewingsPrim = UsdStage.GetPrimAtPath(SewingsPath))
+	{
+		const TArray<UE::FUsdPrim> SewingsPrimChildren = SewingsPrim.GetChildren();
+		const int32 NumSewingsPrimChildren = SewingsPrimChildren.Num();
+		SeamStitches.Reserve(NumSewingsPrimChildren);
+		SeamPatterns.Reserve(NumSewingsPrimChildren);
+
+		for (const UE::FUsdPrim& SewingPrim : SewingsPrimChildren)
+		{
+			const UE::FUsdAttribute PLIndexPairAttr = SewingPrim.GetAttribute(TEXT("PLIndexPair"));
+			if (PLIndexPairAttr.HasValue() && PLIndexPairAttr.GetTypeName() == TEXT("int4[]"))
+			{
+				UE::FVtValue Value;
+				PLIndexPairAttr.Get(Value);
+				UsdUtils::FConvertedVtValue ConvertedVtValue;
+				if (UsdToUnreal::ConvertValue(Value, ConvertedVtValue) && ConvertedVtValue.bIsArrayValued && !ConvertedVtValue.bIsEmpty)
+				{
+					if (const int32 NumSewings = ConvertedVtValue.Entries.Num())
+					{
+						TArray<FIntVector2> SewingPatterns;
+						TArray<FIntVector2> SewingStitches;
+						SewingPatterns.SetNumUninitialized(NumSewings);
+						SewingStitches.SetNumUninitialized(NumSewings);
+
+						check(ConvertedVtValue.Entries[0].Num() == 4);
+						check(ConvertedVtValue.Entries[0][0].IsType<int32>());
+
+						for (int32 Index = 0; Index < NumSewings; ++Index)
+						{
+							SewingPatterns[Index] = FIntVector2(
+								ConvertedVtValue.Entries[Index][0].Get<int32>(),
+								ConvertedVtValue.Entries[Index][2].Get<int32>());
+							SewingStitches[Index] = FIntVector2(
+								ConvertedVtValue.Entries[Index][1].Get<int32>(),
+								ConvertedVtValue.Entries[Index][3].Get<int32>());
+
+							UE_LOG(
+								LogChaosClothAssetDataflowNodes,
+								VeryVerbose,
+								TEXT("Sewing %d:%d-%d:%d"),
+								SewingPatterns[Index][0],
+								SewingStitches[Index][0],
+								SewingPatterns[Index][1],
+								SewingStitches[Index][1]);
+						}
+
+						SeamPatterns.Emplace(MoveTemp(SewingPatterns));
+						SeamStitches.Emplace(MoveTemp(SewingStitches));
+					}
+				}
+			}
+		}
+	}
+
+	// Patterns
+	const UE::FSdfPath PatternsPath = UE::FSdfPath(UE::FSdfPath::AbsoluteRootPath()).AppendChild(TEXT("SimulationData")).AppendChild(TEXT("Patterns"));
+	if (const UE::FUsdPrim PatternsPrim = UsdStage.GetPrimAtPath(PatternsPath))
+	{
+		for (const UE::FUsdPrim& PatternPrim : PatternsPrim.GetChildren())
+		{
+			uint32 TriangleCount = 0;
+			const UE::FUsdAttribute TriangleCountAttr = PatternPrim.GetAttribute(TEXT("TriangleCount"));
+			if (TriangleCountAttr.HasValue() && TriangleCountAttr.GetTypeName() == TEXT("uint"))
+			{
+				UE::FVtValue Value;
+				TriangleCountAttr.Get(Value);
+				const TOptional<uint32> Optional = UsdUtils::GetUnderlyingValue<uint32>(Value);
+				TriangleCount = Optional.IsSet() ? Optional.GetValue() : 0;
+			}
+
+			uint32 VertexCount = 0;
+			const UE::FUsdAttribute VertexCountAttr = PatternPrim.GetAttribute(TEXT("VertexCount"));
+			if (VertexCountAttr.HasValue() && VertexCountAttr.GetTypeName() == TEXT("uint"))
+			{
+				UE::FVtValue Value;
+				VertexCountAttr.Get(Value);
+				const TOptional<uint32> Optional = UsdUtils::GetUnderlyingValue<uint32>(Value);
+				VertexCount = Optional.IsSet() ? Optional.GetValue() : 0;
+			}
+
+			uint32 PatternId = 0;
+			const UE::FUsdAttribute PatternIdAttr = PatternPrim.GetAttribute(TEXT("PatternId"));
+			if (PatternIdAttr.HasValue() && PatternIdAttr.GetTypeName() == TEXT("uint"))
+			{
+				UE::FVtValue Value;
+				PatternIdAttr.Get(Value);
+				const TOptional<uint32> Optional = UsdUtils::GetUnderlyingValue<uint32>(Value);
+				PatternId = Optional.IsSet() ? Optional.GetValue() : 0;
+			}
+
+			UE_LOG(LogChaosClothAssetDataflowNodes,
+				Display,
+				TEXT("Found SimPattern %s, ID %d: %d triangles, %d vertices"), *PatternPrim.GetName().ToString(), PatternId, TriangleCount, VertexCount);
+
+			TArray<FVector2f> RestPositions2D;
+			const UE::FUsdAttribute RestPositions2DAttr = PatternPrim.GetAttribute(TEXT("RestPositions2D"));
+			if (RestPositions2DAttr.HasValue() && RestPositions2DAttr.GetTypeName() == TEXT("float2[]"))
+			{
+				UE::FVtValue Value;
+				RestPositions2DAttr.Get(Value);
+				UsdUtils::FConvertedVtValue ConvertedVtValue;
+				if (UsdToUnreal::ConvertValue(Value, ConvertedVtValue) && ConvertedVtValue.bIsArrayValued && !ConvertedVtValue.bIsEmpty)
+				{
+					RestPositions2D.SetNum(ConvertedVtValue.Entries.Num());
+					check(ConvertedVtValue.Entries[0].Num() == 2);
+					check(ConvertedVtValue.Entries[0][0].IsType<float>());
+
+					for (int32 Index = 0; Index < RestPositions2D.Num(); ++Index)
+					{
+						RestPositions2D[Index] = FVector2f(
+							ConvertedVtValue.Entries[Index][0].Get<float>(),
+							ConvertedVtValue.Entries[Index][1].Get<float>());
+
+						UE_LOG(LogChaosClothAssetDataflowNodes,
+							VeryVerbose,
+							TEXT("RestPositions2D %f %f"), RestPositions2D[Index][0], RestPositions2D[Index][1]);
+					}
+				}
+			}
+			ensure(RestPositions2D.Num() == (int32)VertexCount);
+
+			TArray<FVector3f> DrapedPositions3D;
+			const UE::FUsdAttribute DrapedPositions3DAttr = PatternPrim.GetAttribute(TEXT("DrapedPositions3D"));
+			if (DrapedPositions3DAttr.HasValue() && DrapedPositions3DAttr.GetTypeName() == TEXT("float3[]"))
+			{
+				UE::FVtValue Value;
+				DrapedPositions3DAttr.Get(Value);
+				UsdUtils::FConvertedVtValue ConvertedVtValue;
+				if (UsdToUnreal::ConvertValue(Value, ConvertedVtValue) && ConvertedVtValue.bIsArrayValued && !ConvertedVtValue.bIsEmpty)
+				{
+					DrapedPositions3D.SetNum(ConvertedVtValue.Entries.Num());
+					check(ConvertedVtValue.Entries[0].Num() == 3);
+					check(ConvertedVtValue.Entries[0][0].IsType<float>());
+
+					for (int32 Index = 0; Index < DrapedPositions3D.Num(); ++Index)
+					{
+						DrapedPositions3D[Index] = FVector3f(
+							ConvertedVtValue.Entries[Index][AxesOrder[0]].Get<float>(),
+							ConvertedVtValue.Entries[Index][AxesOrder[1]].Get<float>(),
+							ConvertedVtValue.Entries[Index][AxesOrder[2]].Get<float>()) * CentimetersPerUnit;
+
+						UE_LOG(LogChaosClothAssetDataflowNodes,
+							VeryVerbose,
+							TEXT("RestPositions2D %f %f %f"), DrapedPositions3D[Index][0], DrapedPositions3D[Index][1], DrapedPositions3D[Index][2]);
+					}
+				}
+			}
+			ensure(DrapedPositions3D.Num() == (int32)VertexCount);
+
+			TArray<FIntVector3> TriangleToVertexIndex;
+			const UE::FUsdAttribute TriangleToVertexIndexAttr = PatternPrim.GetAttribute(TEXT("TriangleToVertexIndex"));
+			if (TriangleToVertexIndexAttr.HasValue() && TriangleToVertexIndexAttr.GetTypeName() == TEXT("int3[]"))
+			{
+				UE::FVtValue Value;
+				TriangleToVertexIndexAttr.Get(Value);
+				UsdUtils::FConvertedVtValue ConvertedVtValue;
+				if (UsdToUnreal::ConvertValue(Value, ConvertedVtValue) && ConvertedVtValue.bIsArrayValued && !ConvertedVtValue.bIsEmpty)
+				{
+					TriangleToVertexIndex.SetNum(ConvertedVtValue.Entries.Num());
+					check(ConvertedVtValue.Entries[0].Num() == 3);
+					check(ConvertedVtValue.Entries[0][0].IsType<int32>());
+
+					for (int32 Index = 0; Index < TriangleToVertexIndex.Num(); ++Index)
+					{
+						TriangleToVertexIndex[Index] = FIntVector3(
+							ConvertedVtValue.Entries[Index][WindingOrder[0]].Get<int32>(),
+							ConvertedVtValue.Entries[Index][WindingOrder[1]].Get<int32>(),
+							ConvertedVtValue.Entries[Index][WindingOrder[2]].Get<int32>());
+
+						UE_LOG(LogChaosClothAssetDataflowNodes,
+							VeryVerbose,
+							TEXT("TriangleToVertexIndex %d %d %d"), TriangleToVertexIndex[Index][0], TriangleToVertexIndex[Index][1], TriangleToVertexIndex[Index][2]);
+					}
+				}
+			}
+			ensure(TriangleToVertexIndex.Num() == (int32)TriangleCount);
+
+			// Save pattern to the collection cache
+			if (TriangleCount && VertexCount)
+			{
+				// Remove degenerated triangles
+				TArray<FIntVector3> OutTriangleToVertexIndex;
+				TArray<FVector2f> OutRestPositions2D;
+				TArray<FVector3f> OutDrapedPositions3D;
+				TArray<int32> OutIndices;
+
+				const bool bHasDegenerateTriangles = Private::RemoveDegenerateTriangles(
+					TriangleToVertexIndex,
+					RestPositions2D,
+					DrapedPositions3D,
+					OutTriangleToVertexIndex,
+					OutRestPositions2D,
+					OutDrapedPositions3D,
+					OutIndices);
+
+				// Add the new pattern
+				const int32 SimPatternIndex = ClothFacade.AddSimPattern();
+				FCollectionClothSimPatternFacade SimPattern = ClothFacade.GetSimPattern(SimPatternIndex);
+				SimPattern.Initialize(OutRestPositions2D, OutDrapedPositions3D, OutTriangleToVertexIndex);
+
+				// Remap this pattern's stitches
+				check(SeamPatterns.Num() == SeamStitches.Num());
+				for (int32 SeamIndex = 0; SeamIndex < SeamStitches.Num(); ++SeamIndex)
+				{
+					check(SeamPatterns[SeamIndex].Num() == SeamStitches[SeamIndex].Num());
+
+					for (int32 StitchIndex = 0; StitchIndex < SeamStitches[SeamIndex].Num(); ++StitchIndex)
+					{
+						FIntVector2& SeamPattern = SeamPatterns[SeamIndex][StitchIndex];
+						FIntVector2& SeamStitch = SeamStitches[SeamIndex][StitchIndex];
+
+						for (int32 Side = 0; Side < 2; ++Side)
+						{
+							if (SeamPattern[Side] == (int32)PatternId)
+							{
+								SeamPattern[Side] = INDEX_NONE;  // In case two patterns were to be exported with the same id.
+								SeamStitch[Side] = OutIndices[SeamStitch[Side]] + SimPattern.GetSimVertices2DOffset();
+							}
+						}
+					}
+				}
+
+				// Flag collapsed vertices for info
+				if (bHasDegenerateTriangles)
+				{
+					// TODO: Make this a feature or remove it?
+					const FName WeightMapName(TEXT("_Degenerate"));  // The undescore means this is an internal weight map name
+					const TConstArrayView<int32> SimVertex3DLookup = static_cast<FCollectionClothSimPatternConstFacade&>(SimPattern).GetSimVertex3DLookup();
+					ClothFacade.AddWeightMap(WeightMapName);
+					const TArrayView<float> WeightMap = ClothFacade.GetWeightMap(WeightMapName);
+
+					for (int32 TriangleIndex = 0, OutTriangleIndex = 0; OutTriangleIndex < OutTriangleToVertexIndex.Num(); ++TriangleIndex)
+					{
+						const int32 Index0 = OutIndices[TriangleToVertexIndex[TriangleIndex][0]];
+						const int32 Index1 = OutIndices[TriangleToVertexIndex[TriangleIndex][1]];
+						const int32 Index2 = OutIndices[TriangleToVertexIndex[TriangleIndex][2]];
+
+						if (Index0 == OutTriangleToVertexIndex[OutTriangleIndex][0] &&
+							Index1 == OutTriangleToVertexIndex[OutTriangleIndex][1] &&
+							Index2 == OutTriangleToVertexIndex[OutTriangleIndex][2])
+						{
+							++OutTriangleIndex;
+						}
+						else
+						{
+							WeightMap[SimVertex3DLookup[Index0]] = WeightMap[SimVertex3DLookup[Index1]] = WeightMap[SimVertex3DLookup[Index2]] = 1.f;
+						}
+					}
+				}
+			}
+		}
+	}
+
+	// Add seams
+	for (int32 SeamIndex = 0; SeamIndex < SeamStitches.Num(); ++SeamIndex)
+	{
+		FCollectionClothSeamFacade Seam = ClothFacade.AddGetSeam();
+		Seam.Initialize(SeamStitches[SeamIndex]);
+	}
+
+	SlowTask.EnterProgressFrame(1.f);
+	SlowTask.ForceRefresh();
+	return true;
+
+#else  // #if USE_USD_SDK
+
+	OutErrorText = LOCTEXT("NoUsdSdk", "The ChaosClothAssetDataflowNodes module has been compiled without the USD SDK enabled.");
+	return false;
+
+#endif  // #else #if USE_USD_SDK
+}
+
+bool FChaosClothAssetUSDImportNode::ImportFromCache(const TSharedRef<FManagedArrayCollection>& OutClothCollection, FText& OutErrorText) const
+{
+	using namespace UE::Chaos::ClothAsset;
+
+	// Initialize from collection cache
+	// TODO: Until we have a schema so that we can use the asset cache and remove the collection cache
+	*OutClothCollection = CollectionCache;
+
+	// Initialize from asset registry
+	TArray<FAssetData> AssetData;
+
+	if (!PackagePath.IsEmpty())
+	{
+		const FAssetRegistryModule& AssetRegistryModule = FModuleManager::LoadModuleChecked<FAssetRegistryModule>("AssetRegistry");
+
+		const UClass* const Class = UStaticMesh::StaticClass();
+		constexpr bool bRecursive = true;
+		constexpr bool bIncludeOnlyOnDiskAssets = false;
+		AssetRegistryModule.Get().GetAssetsByPath(FName(*PackagePath), AssetData, bRecursive, bIncludeOnlyOnDiskAssets);
+	}
+
+	for (const FAssetData& AssetDatum : AssetData)
+	{
+		if (!(AssetDatum.IsUAsset() && AssetDatum.IsTopLevelAsset()))
+		{
+			continue;
+		}
+
+		UE_LOG(LogChaosClothAssetDataflowNodes,
+			Verbose,
+			TEXT("Imported USD Object %s of type %s, path: %s"),
+			*AssetDatum.AssetName.ToString(),
+			*AssetDatum.AssetClassPath.ToString(),
+			*AssetDatum.GetFullName());
+
+		if (AssetDatum.GetClass() == UStaticMesh::StaticClass())
+		{
+			UStaticMesh* const StaticMesh = CastChecked<UStaticMesh>(AssetDatum.GetAsset());
+
+			if (StaticMesh->GetNumSourceModels() > 0)  // Only deals with LOD 0 for now
+			{
+				using namespace UE::Chaos::ClothAsset::Private;
+
+				constexpr int32 LODIndex = 0;
+				const FMeshDescription* const MeshDescription = StaticMesh->GetMeshDescription(LODIndex);
+				const FMeshBuildSettings& BuildSettings = StaticMesh->GetSourceModel(LODIndex).BuildSettings;
+				const TArray<FStaticMaterial>& StaticMaterials = StaticMesh->GetStaticMaterials();
+
+				FSkeletalMeshLODModel SkeletalMeshModel;
+				if (FClothDataflowTools::BuildSkeletalMeshModelFromMeshDescription(MeshDescription, BuildSettings, SkeletalMeshModel))
+				{
+					FStaticMeshConstAttributes MeshAttributes(*MeshDescription);
+					TPolygonGroupAttributesConstRef<FName> MaterialSlotNames = MeshAttributes.GetPolygonGroupMaterialSlotNames();
+					for (int32 SectionIndex = 0; SectionIndex < SkeletalMeshModel.Sections.Num(); ++SectionIndex)
+					{
+						// Section MaterialIndex refers to the polygon group index. Look up which material this corresponds with.
+						const FName& MaterialSlotName = MaterialSlotNames[SkeletalMeshModel.Sections[SectionIndex].MaterialIndex];
+						const int32 MaterialIndex = StaticMesh->GetMaterialIndexFromImportedMaterialSlotName(MaterialSlotName);
+						const FString RenderMaterialPathName = StaticMaterials.IsValidIndex(MaterialIndex) && StaticMaterials[MaterialIndex].MaterialInterface ? 
+							StaticMaterials[MaterialIndex].MaterialInterface->GetPathName() :
+							FString();
+						FClothDataflowTools::AddRenderPatternFromSkeletalMeshSection(OutClothCollection, SkeletalMeshModel, SectionIndex, RenderMaterialPathName);
+					}
+				}
+			}
+		}
+	}
+
+	// Apply transforms  TODO: Use a transform node/tool instead
+	FCollectionClothFacade ClothFacade(OutClothCollection);
+	for (FVector3f& SimPosition3D : ClothFacade.GetSimPosition3D())
+	{
+		SimPosition3D = Transform.TransformPosition(SimPosition3D);
+	}
+	for (FVector3f& RenderPosition : ClothFacade.GetRenderPosition())
+	{
+		RenderPosition = Transform.TransformPosition(RenderPosition);
+	}
+
+	return true;
+}
+
+#undef LOCTEXT_NAMESPACE
diff --git a/Engine/Plugins/Experimental/ChaosClothAssetEditor/Source/ChaosClothAssetDataflowNodes/Public/ChaosClothAsset/USDImportNode.h b/Engine/Plugins/Experimental/ChaosClothAssetEditor/Source/ChaosClothAssetDataflowNodes/Public/ChaosClothAsset/USDImportNode.h
new file mode 100644
--- /dev/null
+++ b/Engine/Plugins/Experimental/ChaosClothAssetEditor/Source/ChaosClothAssetDataflowNodes/Public/ChaosClothAsset/USDImportNode.h
@@ -0,0 +1,49 @@
+// Copyright Epic Games, Inc. All Rights Reserved.
+#pragma once 
+
+#include "ChaosClothAsset/ImportFilePath.h"
+#include "Dataflow/DataflowTerminalNode.h"
+#include "GeometryCollection/ManagedArrayCollection.h"
+#include "Misc/SecureHash.h"
+#include "USDImportNode.generated.h"
+
+/** Import a USD file from a third party garment construction software. */
+USTRUCT(meta = (DataflowCloth))
+struct FChaosClothAssetUSDImportNode : public FDataflowTerminalNode
+{
+	GENERATED_USTRUCT_BODY()
+	DATAFLOW_NODE_DEFINE_INTERNAL(FChaosClothAssetUSDImportNode, "USDImport", "Cloth", "Cloth USD Import")
+
+public:
+	UPROPERTY(meta = (DataflowOutput))
+	FManagedArrayCollection Collection;
+
+	/** Name of the USD file to import. */
+	UPROPERTY(EditAnywhere, Category = "USD Import", Meta = (DisplayName = "USD File"))
+	FChaosClothAssetImportFilePath UsdFile;
+
+	/**
+	 * The transform to apply to the imported 3D positions. 
+	 * Experimental, will soon be removed in favour of a tool supported transform node.
+	 */
+	UPROPERTY(EditAnywhere, Category = "USD Import", Meta = (Exprimental))
+	FTransform3f Transform = FTransform3f::Identity;
+
+	FChaosClothAssetUSDImportNode(const Dataflow::FNodeParameters& InParam, FGuid InGuid = FGuid::NewGuid());
+
+private:
+	//~ Begin FDataflowNode interface
+	virtual void SetAssetValue(TObjectPtr<UObject> Asset, Dataflow::FContext& Context) const override;
+	virtual void Evaluate(Dataflow::FContext& Context, const FDataflowOutput* Out) const override;
+	virtual void Serialize(FArchive& Archive) override;
+	//~ End FDataflowNode interface
+
+	bool ImportFromFile(const FString& UsdPath, const FString& AssetPath, class FText& OutErrorText);
+	bool ImportFromCache(const TSharedRef<FManagedArrayCollection>& ClothCollection, class FText& OutErrorText) const;
+
+	UPROPERTY(VisibleAnywhere, Category = "USD Import")
+	FString PackagePath;  // Content folder where all the USD asset are imported
+
+	FMD5Hash FileHash;
+	FManagedArrayCollection CollectionCache;  // Content cache for data that hasn't got a USD schema yet
+};
