Cloth Editor  layout changes:
- Dataflow Graph Editor extends to the left of the editor window
- Dataflow Node Details tab is horizontally aligned with the Asset Details and PreviewScene Details
- Outliner is in the same tab stack as Dataflow Graph Editor

diff --git a/Engine/Plugins/Experimental/ChaosClothAssetEditor/Source/ChaosClothAssetEditor/Private/ChaosClothAsset/ClothEditorToolkit.cpp b/Engine/Plugins/Experimental/ChaosClothAssetEditor/Source/ChaosClothAssetEditor/Private/ChaosClothAsset/ClothEditorToolkit.cpp
--- a/Engine/Plugins/Experimental/ChaosClothAssetEditor/Source/ChaosClothAssetEditor/Private/ChaosClothAsset/ClothEditorToolkit.cpp
+++ b/Engine/Plugins/Experimental/ChaosClothAssetEditor/Source/ChaosClothAssetEditor/Private/ChaosClothAsset/ClothEditorToolkit.cpp
@@ -96,75 +96,75 @@ FChaosClothAssetEditorToolkit::FChaosClothAssetEditorToolkit(UAssetEditor* InOwn
 	// Note: Changes to the layout should include a increment to the layout's ID, i.e.
 	// ChaosClothAssetEditorLayout[X] -> ChaosClothAssetEditorLayout[X+1]. Otherwise, layouts may be messed up
 	// without a full reset to layout defaults inside the editor.
-	StandaloneDefaultLayout = FTabManager::NewLayout(FName("ChaosClothAssetEditorLayout4"))
+	StandaloneDefaultLayout = FTabManager::NewLayout(FName("ChaosClothAssetEditorLayout5"))
 		->AddArea
 		(
-			FTabManager::NewPrimaryArea()->SetOrientation(Orient_Vertical)
+			FTabManager::NewPrimaryArea()->SetOrientation(Orient_Horizontal)
 			->Split
 			(
-				FTabManager::NewSplitter()->SetOrientation(Orient_Horizontal)
+				FTabManager::NewSplitter()->SetOrientation(Orient_Vertical)
+				->SetSizeCoefficient(0.8f)	// Relative width of (Tools Panel, Construction Viewport, Preview Viewport, Dataflow Graph Editor, Outliner) vs (Asset Details, Preview Scene Details, Dataflow Node Details)
 				->Split
 				(
-					FTabManager::NewStack()
-					->SetSizeCoefficient(0.1f)
-					->SetExtensionId(UChaosClothAssetEditorUISubsystem::EditorSidePanelAreaName)
-					->SetHideTabWell(true)
-				)
-				->Split
-				(
-					FTabManager::NewSplitter()->SetOrientation(Orient_Vertical)
+					FTabManager::NewSplitter()->SetOrientation(Orient_Horizontal)
+					->SetSizeCoefficient(0.55f)	// Relative height of (Tools Panel, Construction Viewport, Preview Viewport) vs (Dataflow Graph Editor, Outliner)
 					->Split
 					(
-						FTabManager::NewSplitter()->SetOrientation(Orient_Horizontal)
-						->Split
-						(
-							FTabManager::NewStack()
-							->SetSizeCoefficient(0.5f)
-							->AddTab(ViewportTabID, ETabState::OpenedTab)
-							->SetExtensionId("RestSpaceViewportArea")
-							->SetHideTabWell(true)
-						)
-						->Split
-						(
-							FTabManager::NewStack()
-							->SetSizeCoefficient(0.5f)
-							->AddTab(ClothPreviewTabID, ETabState::OpenedTab)
-							->SetExtensionId("Viewport3DArea")
-							->SetHideTabWell(true)
-						)
+						FTabManager::NewStack()
+						->SetSizeCoefficient(0.1f)		// Relative width of (Tools Panel) vs (Construction Viewport, Preview Viewport)
+						->SetExtensionId(UChaosClothAssetEditorUISubsystem::EditorSidePanelAreaName)
+						->SetHideTabWell(true)
 					)
 					->Split
 					(
-						FTabManager::NewSplitter()->SetOrientation(Orient_Horizontal)
-						->Split
-						(
-							FTabManager::NewStack()
-							->SetSizeCoefficient(0.7f)
-							->AddTab(GraphCanvasTabId, ETabState::OpenedTab)
-							->SetExtensionId("GraphEditorArea")
-							->SetHideTabWell(true)
-						)
-						->Split
-						(
-							FTabManager::NewStack()
-							->SetSizeCoefficient(0.3f)
-							->AddTab(NodeDetailsTabId, ETabState::OpenedTab)
-							->SetExtensionId("NodeDetailsArea")
-							->SetHideTabWell(true)
-						)
+						FTabManager::NewStack()
+						->SetSizeCoefficient(0.45f)		// Relative width of (Construction Viewport) vs (Tools Panel, Preview Viewport)
+						->AddTab(ViewportTabID, ETabState::OpenedTab)
+						->SetExtensionId("RestSpaceViewportArea")
+						->SetHideTabWell(true)
+					)
+					->Split
+					(
+						FTabManager::NewStack()
+						->SetSizeCoefficient(0.45f)		// Relative width of (Preview Viewport) vs (Tools Panel, Construction Viewport)
+						->AddTab(ClothPreviewTabID, ETabState::OpenedTab)
+						->SetExtensionId("Viewport3DArea")
+						->SetHideTabWell(true)
 					)
 				)
 				->Split
 				(
 					FTabManager::NewStack()
-					->SetSizeCoefficient(0.3f)
-					->AddTab(DetailsTabID, ETabState::OpenedTab)
+					->SetSizeCoefficient(0.45f)	// Relative height of (Dataflow Graph Editor, Outliner) vs (Tools Panel, Construction Viewport, Preview Viewport)
+					->AddTab(GraphCanvasTabId, ETabState::OpenedTab)
 					->AddTab(OutlinerTabID, ETabState::OpenedTab)
+					->SetExtensionId("GraphEditorArea")
+					->SetHideTabWell(false)
+					->SetForegroundTab(GraphCanvasTabId)
+				)
+			)
+			->Split
+			(
+				FTabManager::NewSplitter()->SetOrientation(Orient_Vertical)
+				->SetSizeCoefficient(0.2f)	// Relative width of (Asset Details, Preview Scene Details, Dataflow Node Details) vs (Tools Panel, Construction Viewport, Preview Viewport, Dataflow Graph Editor, Outliner)
+				->Split
+				(
+					FTabManager::NewStack()
+					->SetSizeCoefficient(0.65f)	// Relative height of (Asset Details, Preview Scene Details) vs (Dataflow Node Details)
+					->AddTab(DetailsTabID, ETabState::OpenedTab)
 					->AddTab(PreviewSceneDetailsTabID, ETabState::OpenedTab)
 					->SetExtensionId("DetailsArea")
 					->SetHideTabWell(true)
 					->SetForegroundTab(DetailsTabID)
 				)
+				->Split
+				(
+					FTabManager::NewStack()
+					->SetSizeCoefficient(0.35f)	// Relative height of (Dataflow Node Details) vs (Asset Details, Preview Scene Details)
+					->AddTab(NodeDetailsTabId, ETabState::OpenedTab)
+					->SetExtensionId("NodeDetailsArea")
+					->SetHideTabWell(false)
+				)
 			)
 		);
 
