Cloth Editor: move Tool buttons to the top toolbar

diff --git a/Engine/Plugins/Experimental/ChaosClothAssetEditor/Source/ChaosClothAssetEditor/Private/ChaosClothAsset/ClothEditorCommands.cpp b/Engine/Plugins/Experimental/ChaosClothAssetEditor/Source/ChaosClothAssetEditor/Private/ChaosClothAsset/ClothEditorCommands.cpp
--- a/Engine/Plugins/Experimental/ChaosClothAssetEditor/Source/ChaosClothAssetEditor/Private/ChaosClothAsset/ClothEditorCommands.cpp
+++ b/Engine/Plugins/Experimental/ChaosClothAssetEditor/Source/ChaosClothAssetEditor/Private/ChaosClothAsset/ClothEditorCommands.cpp
@@ -42,13 +42,13 @@ void FChaosClothAssetEditorCommands::RegisterCommands()
 	UI_COMMAND(BeginRemeshTool, "Remesh", "Remesh the selected mesh", EUserInterfaceActionType::Button, FInputChord());
 	UI_COMMAND(BeginAttributeEditorTool, "AttrEd", "Edit/configure mesh attributes", EUserInterfaceActionType::Button, FInputChord());
 	
-	UI_COMMAND(BeginWeightMapPaintTool, "MapPnt", "Paint weight maps on the mesh", EUserInterfaceActionType::None, FInputChord());
-	UI_COMMAND(AddWeightMapNode, "MapPnt", "Paint weight maps on the mesh", EUserInterfaceActionType::Button, FInputChord());
+	UI_COMMAND(BeginWeightMapPaintTool, "Add Weight Map", "Paint weight maps on the mesh", EUserInterfaceActionType::None, FInputChord());
+	UI_COMMAND(AddWeightMapNode, "Add Weight Map", "Paint weight maps on the mesh", EUserInterfaceActionType::Button, FInputChord());
 	UI_COMMAND(BeginMeshSelectionTool, "Select", "Select mesh elements", EUserInterfaceActionType::None, FInputChord());
 	UI_COMMAND(AddMeshSelectionNode, "Select", "Select mesh elements", EUserInterfaceActionType::Button, FInputChord());
 
-	UI_COMMAND(BeginTransferSkinWeightsTool, "TransferSkinWeights", "Transfer skinning weights from a SkeletalMesh", EUserInterfaceActionType::None, FInputChord());
-	UI_COMMAND(AddTransferSkinWeightsNode, "TransferSkinWeights", "Transfer skinning weights from a SkeletalMesh", EUserInterfaceActionType::Button, FInputChord());
+	UI_COMMAND(BeginTransferSkinWeightsTool, "Transfer Skin Weights", "Transfer skinning weights from a SkeletalMesh", EUserInterfaceActionType::None, FInputChord());
+	UI_COMMAND(AddTransferSkinWeightsNode, "Transfer Skin Weights", "Transfer skinning weights from a SkeletalMesh", EUserInterfaceActionType::Button, FInputChord());
 
 	UI_COMMAND(SetConstructionMode2D, "2D Sim", "Switches the viewport to 2D simulation mesh view", EUserInterfaceActionType::RadioButton, FInputChord());
 	UI_COMMAND(SetConstructionMode3D, "3D Sim", "Switches the viewport to 3D simulation mesh view", EUserInterfaceActionType::RadioButton, FInputChord());
diff --git a/Engine/Plugins/Experimental/ChaosClothAssetEditor/Source/ChaosClothAssetEditor/Private/ChaosClothAsset/ClothEditorModeToolkit.cpp b/Engine/Plugins/Experimental/ChaosClothAssetEditor/Source/ChaosClothAssetEditor/Private/ChaosClothAsset/ClothEditorModeToolkit.cpp
--- a/Engine/Plugins/Experimental/ChaosClothAssetEditor/Source/ChaosClothAssetEditor/Private/ChaosClothAsset/ClothEditorModeToolkit.cpp
+++ b/Engine/Plugins/Experimental/ChaosClothAssetEditor/Source/ChaosClothAssetEditor/Private/ChaosClothAsset/ClothEditorModeToolkit.cpp
@@ -140,18 +140,23 @@ void FChaosClothAssetEditorModeToolkit::Init(const TSharedPtr<IToolkitHost>& Ini
 	];
 }
 
-
-void FChaosClothAssetEditorModeToolkit::BuildToolPalette(FName PaletteIndex, class FToolBarBuilder& ToolbarBuilder)
+void FChaosClothAssetEditorModeToolkit::BuildEditorToolBar(const FName& EditorToolBarName)
 {
-	const FChaosClothAssetEditorCommands& Commands = FChaosClothAssetEditorCommands::Get();
+	check(EditorToolBarName != FName());
 
-	// TODO: make separate palettes to split up the tools tab
-	if (PaletteIndex == FBaseCharacterFXEditorModeToolkit::ToolsTabName)
-	{
-		ToolbarBuilder.AddToolBarButton(Commands.AddWeightMapNode);
-		ToolbarBuilder.AddToolBarButton(Commands.AddMeshSelectionNode);
-		ToolbarBuilder.AddToolBarButton(Commands.AddTransferSkinWeightsNode);
-	}
+	const TSharedRef<const FUICommandList> EdModeToolkitCommands = GetToolkitCommands();
+
+	UToolMenu* const ToolBarMenu = UToolMenus::Get()->ExtendMenu(EditorToolBarName);
+	FToolMenuSection& Section = ToolBarMenu->FindOrAddSection("ClothTools");
+
+	FToolMenuEntry& WeightMapButtonEntry = Section.AddEntry(FToolMenuEntry::InitToolBarButton(FChaosClothAssetEditorCommands::Get().AddWeightMapNode));
+	WeightMapButtonEntry.SetCommandList(EdModeToolkitCommands);
+
+	FToolMenuEntry& SkinWeightsButtonEntry = Section.AddEntry(FToolMenuEntry::InitToolBarButton(FChaosClothAssetEditorCommands::Get().AddTransferSkinWeightsNode));
+	SkinWeightsButtonEntry.SetCommandList(EdModeToolkitCommands);
+
+	FToolMenuEntry& SelectionButtonEntry = Section.AddEntry(FToolMenuEntry::InitToolBarButton(FChaosClothAssetEditorCommands::Get().AddMeshSelectionNode));
+	SelectionButtonEntry.SetCommandList(EdModeToolkitCommands);
 }
 
 const FSlateBrush* FChaosClothAssetEditorModeToolkit::GetActiveToolIcon(const FString& ActiveToolIdentifier) const
diff --git a/Engine/Plugins/Experimental/ChaosClothAssetEditor/Source/ChaosClothAssetEditor/Private/ChaosClothAsset/ClothEditorModeToolkit.h b/Engine/Plugins/Experimental/ChaosClothAssetEditor/Source/ChaosClothAssetEditor/Private/ChaosClothAsset/ClothEditorModeToolkit.h
--- a/Engine/Plugins/Experimental/ChaosClothAssetEditor/Source/ChaosClothAssetEditor/Private/ChaosClothAsset/ClothEditorModeToolkit.h
+++ b/Engine/Plugins/Experimental/ChaosClothAssetEditor/Source/ChaosClothAssetEditor/Private/ChaosClothAsset/ClothEditorModeToolkit.h
@@ -23,8 +23,9 @@ public:
 
 	void Init(const TSharedPtr<IToolkitHost>& InitToolkitHost, TWeakObjectPtr<UEdMode> InOwningMode);
 
-    /** For a specific tool palette category, construct and fill ToolbarBuilder with the category's tools **/
-	virtual void BuildToolPalette(FName PaletteName, class FToolBarBuilder& ToolbarBuilder) override;
+	// Add buttons to the toolbar with the specified name
+	// Note: Most FModeToolkits would define BuildToolPalette, but we are putting buttons in the top toolbar instead
+	void BuildEditorToolBar(const FName& EditorToolBarName);
 
 	virtual const FSlateBrush* GetActiveToolIcon(const FString& Identifier) const override;
 
@@ -41,6 +42,7 @@ public:
 private:
 
 	// Get the viewport widget associated with the given manager
+	// TODO: This should not be necessary any more as we do not run tools in the Cloth Preview Viewport (JIRA UE-201248)
 	SBaseCharacterFXEditorViewport* GetViewportWidgetForManager(UInteractiveToolManager* Manager);
 
 	UEditorInteractiveToolsContext* GetCurrentToolsContext();
diff --git a/Engine/Plugins/Experimental/ChaosClothAssetEditor/Source/ChaosClothAssetEditor/Private/ChaosClothAsset/ClothEditorToolkit.cpp b/Engine/Plugins/Experimental/ChaosClothAssetEditor/Source/ChaosClothAssetEditor/Private/ChaosClothAsset/ClothEditorToolkit.cpp
--- a/Engine/Plugins/Experimental/ChaosClothAssetEditor/Source/ChaosClothAssetEditor/Private/ChaosClothAsset/ClothEditorToolkit.cpp
+++ b/Engine/Plugins/Experimental/ChaosClothAssetEditor/Source/ChaosClothAssetEditor/Private/ChaosClothAsset/ClothEditorToolkit.cpp
@@ -282,6 +282,10 @@ void FChaosClothAssetEditorToolkit::InitializeEdMode(UBaseCharacterFXEditorMode*
 		FChaosClothAssetEditorModeToolkit* ClothModeToolkit = static_cast<FChaosClothAssetEditorModeToolkit*>(ModeToolkit.Get());
 		ClothModeToolkit->SetRestSpaceViewportWidget(RestSpaceViewportWidget);
 		ClothModeToolkit->SetPreviewViewportWidget(PreviewViewportWidget);
+
+		FName ParentToolbarName;
+		const FName ToolBarName = GetToolMenuToolbarName(ParentToolbarName);
+		ClothModeToolkit->BuildEditorToolBar(ToolBarName);
 	}
 }
 
