Chaos Cloth Asset :
- Enabled KindsToCollapse to import all USD meshes as a single static mesh.
- Enabled USD Stage MetersPerUnit and UpAxis, and temporarily fix the render mesh scale to cm.
- Added function to remove duplicated triangles.
- Added function to remove duplicated stitches.

#tests Imported various USD files using the latest importer version.

diff --git a/Engine/Plugins/Experimental/ChaosClothAssetEditor/Source/ChaosClothAssetDataflowNodes/Private/ChaosClothAsset/USDImportNode.cpp b/Engine/Plugins/Experimental/ChaosClothAssetEditor/Source/ChaosClothAssetDataflowNodes/Private/ChaosClothAsset/USDImportNode.cpp
--- a/Engine/Plugins/Experimental/ChaosClothAssetEditor/Source/ChaosClothAssetDataflowNodes/Private/ChaosClothAsset/USDImportNode.cpp
+++ b/Engine/Plugins/Experimental/ChaosClothAssetEditor/Source/ChaosClothAssetDataflowNodes/Private/ChaosClothAsset/USDImportNode.cpp
@@ -46,7 +46,7 @@ namespace UE::Chaos::ClothAsset::Private
 		const int32 VertexCount = RestPositions2D.Num();
 		const int32 TriangleCount = TriangleToVertexIndex.Num();
 
-		OutTriangleToVertexIndex.Reset(TriangleToVertexIndex.Num());
+		OutTriangleToVertexIndex.Reset(TriangleCount);
 
 		// Remap[Index] is the index of the first vertex in a group of degenerated triangles to be callapsed.
 		// When two groups of collapsed vertices are merged, the group with the greatest Remap[index] value must adopt the one from the other group.
@@ -124,8 +124,8 @@ namespace UE::Chaos::ClothAsset::Private
 		const int32 OutTriangleCount = OutTriangleToVertexIndex.Num();
 		bHasDegenerateTriangles = (TriangleCount != OutTriangleCount);
 
-		UE_CLOG(bHasDegenerateTriangles, LogChaosClothAssetDataflowNodes, Warning,
-			TEXT("USD import found %d degenerated triangles out of %d"), TriangleCount - OutTriangleCount, TriangleCount);
+		UE_CLOG(bHasDegenerateTriangles, LogChaosClothAssetDataflowNodes, Display,
+			TEXT("USD import found and removed %d degenerated triangles out of %d source triangles."), TriangleCount - OutTriangleCount, TriangleCount);
 
 		OutRestPositions2D.Reset(OutVertexCount);
 		OutDrapedPositions3D.Reset(OutVertexCount);
@@ -168,6 +168,114 @@ namespace UE::Chaos::ClothAsset::Private
 
 		return bHasDegenerateTriangles;
 	}
+
+	static bool RemoveDuplicateTriangles(TArray<FIntVector3>& TriangleToVertexIndex)
+	{
+		bool bHasDuplicatedTriangles = false;
+
+		const int32 TriangleCount = TriangleToVertexIndex.Num();
+
+		TSet<FIntVector3> Triangles;
+		Triangles.Reserve(TriangleCount);
+
+		TArray<FIntVector3> OutTriangleToVertexIndex;
+		OutTriangleToVertexIndex.Reserve(TriangleCount);
+
+		auto GetSortedIndices = [](const FIntVector3& TriangleIndices)->FIntVector3
+			{
+				const int32 Index0 = TriangleIndices[0];
+				const int32 Index1 = TriangleIndices[1];
+				const int32 Index2 = TriangleIndices[2];
+
+				return (Index0 < Index1) ?
+					(Index1 < Index2) ? FIntVector3(Index0, Index1, Index2) : (Index0 < Index2) ? FIntVector3(Index0, Index2, Index1) : FIntVector3(Index2, Index0, Index1) :
+					(Index0 < Index2) ? FIntVector3(Index1, Index0, Index2) : (Index1 < Index2) ? FIntVector3(Index1, Index2, Index0) : FIntVector3(Index2, Index1, Index0);
+			};
+
+		for (int32 Index = 0; Index < TriangleCount; ++Index)
+		{
+			const FIntVector3& TriangleIndices = TriangleToVertexIndex[Index];
+			const FIntVector3 TriangleSortedIndices = GetSortedIndices(TriangleIndices);
+
+			bool bIsAlreadyInSet;
+			Triangles.FindOrAdd(TriangleSortedIndices, &bIsAlreadyInSet);
+
+			if (bIsAlreadyInSet)
+			{
+				bHasDuplicatedTriangles = true;
+			}
+			else
+			{
+				OutTriangleToVertexIndex.Emplace(TriangleIndices);
+			}
+		}
+
+		UE_CLOG(bHasDuplicatedTriangles, LogChaosClothAssetDataflowNodes, Display,
+			TEXT("USD import found and removed %d duplicated triangles out of %d source triangles."), TriangleCount - OutTriangleToVertexIndex.Num(), TriangleCount);
+
+		TriangleToVertexIndex = MoveTemp(OutTriangleToVertexIndex);
+
+		return bHasDuplicatedTriangles;
+	}
+
+	bool RemoveDuplicateStitches(TArray<TArray<FIntVector2>>& SeamStitches)
+	{
+		bool bHasDuplicateStitches = false;
+
+		const int32 NumSeamStitches = SeamStitches.Num();
+
+		// Calculate the total number of stitches
+		int32 NumStitches = 0;
+		for (const TArray<FIntVector2>& Stitches : SeamStitches)
+		{
+			NumStitches += Stitches.Num();
+		}
+
+		TSet<FIntVector2> StichSet;
+		StichSet.Reserve(NumStitches);
+
+		int32 OutNumStitches = 0;
+		TArray<TArray<FIntVector2>> OutSeamStitches;
+		OutSeamStitches.Reserve(NumSeamStitches);
+
+		for (const TArray<FIntVector2>& Stitches : SeamStitches)
+		{
+			TArray<FIntVector2> OutStitches;
+			OutStitches.Reserve(Stitches.Num());
+
+			for (const FIntVector2& Stitch : Stitches)
+			{
+				const FIntVector2 SortedStitch = Stitch[0] < Stitch[1] ?
+					FIntVector2(Stitch[0], Stitch[1]) :
+					FIntVector2(Stitch[1], Stitch[0]);
+
+				bool bIsAlreadyInSet;
+				StichSet.FindOrAdd(SortedStitch, &bIsAlreadyInSet);
+
+				if (bIsAlreadyInSet)
+				{
+					bHasDuplicateStitches = true;
+				}
+				else
+				{
+					OutStitches.Emplace(Stitch);
+				}
+			}
+
+			if (OutStitches.Num())
+			{
+				OutSeamStitches.Emplace(OutStitches);
+				OutNumStitches += OutStitches.Num();
+			}
+		}
+
+		UE_CLOG(bHasDuplicateStitches, LogChaosClothAssetDataflowNodes, Display,
+			TEXT("USD import found and removed %d duplicated stitches out of %d source stitches."), NumStitches - OutNumStitches, NumStitches);
+
+		SeamStitches = MoveTemp(OutSeamStitches);
+
+		return bHasDuplicateStitches;
+	}
 }  // End namespace UE::Chaos::ClothAsset::Private
 
 FChaosClothAssetUSDImportNode::FChaosClothAssetUSDImportNode(const Dataflow::FNodeParameters& InParam, FGuid InGuid)
@@ -302,12 +410,31 @@ bool FChaosClothAssetUSDImportNode::ImportFromFile(const FString& UsdFilePath, c
 	constexpr bool bUseStageCache = false;  // Reload from disk, not from cache
 	constexpr EUsdInitialLoadSet UsdInitialLoadSet = EUsdInitialLoadSet::LoadAll;  // TODO: Ideally we should only use LoadNone to start with and load what's needed once the Schema is defined
 
-	const UE::FUsdStage UsdStage = UnrealUSDWrapper::OpenStage(*UsdFilePath, UsdInitialLoadSet, bUseStageCache);
+	UE::FUsdStage UsdStage = UnrealUSDWrapper::OpenStage(*UsdFilePath, UsdInitialLoadSet, bUseStageCache);
 	if (!UsdStage)
 	{
 		OutErrorText = LOCTEXT("CantCreateNewStage", "Failed to open the specified USD file.");
 		return false;
 	}
+
+	// Look for the Mesh prim and set its kind to enable KindsToCollapse
+	const UE::FSdfPath MeshPath(UE::FSdfPath(UE::FSdfPath::AbsoluteRootPath()).AppendChild(TEXT("Mesh")));
+	if (UE::FUsdPrim MeshPrim = UsdStage.GetPrimAtPath(MeshPath))
+	{
+		MeshPrim.SetTypeName(TEXT("Xform"));  // TODO: Ideally this two operations need to be done in the exporter
+		UsdUtils::SetDefaultKind(MeshPrim, EUsdDefaultKind::Component);
+
+		// Look for the SkelRoot prim and disable it to allow the KindsToCollapse to work
+		TArray<UE::FUsdPrim> MeshPrimChildren = MeshPrim.GetChildren();
+		for (UE::FUsdPrim& MeshPrimChild : MeshPrimChildren)
+		{
+			if (MeshPrimChild.GetTypeName() == TEXT("SkelRoot"))
+			{
+				MeshPrimChild.SetActive(false);
+			}
+		}
+	}
+
 	SlowTask.EnterProgressFrame(1.f);
 
 	// Update import location
@@ -343,7 +470,7 @@ bool FChaosClothAssetUSDImportNode::ImportFromFile(const FString& UsdFilePath, c
 		ImportOptions->RootMotionHandling = EUsdRootMotionHandling::NoAdditionalRootMotion;
 		ImportOptions->SubdivisionLevel = 0;
 		ImportOptions->bOverrideStageOptions = true;
-		ImportOptions->StageOptions.MetersPerUnit = 0.001f;  // Assume mm scale for now until this is set in the file
+		ImportOptions->StageOptions.MetersPerUnit = 0.01f;  // REVIEW: The render mesh is currently exported in cm scale despite the stage MetersPerUnit and simulation data always set to mm
 		ImportOptions->StageOptions.UpAxis = StageInfo.UpAxis;
 		ImportOptions->bImportAtSpecificTimeCode = false;
 		ImportOptions->ImportTimeCode = 0.f;
@@ -354,7 +481,7 @@ bool FChaosClothAssetUSDImportNode::ImportFromFile(const FString& UsdFilePath, c
 		ImportOptions->ExistingAssetPolicy = EReplaceAssetPolicy::Replace;
 		// Processing
 		ImportOptions->bPrimPathFolderStructure = false;
-		ImportOptions->KindsToCollapse = true;  // TODO: Might need not to merge all static meshes
+		ImportOptions->KindsToCollapse = (int32)EUsdDefaultKind::Component;
 		ImportOptions->bMergeIdenticalMaterialSlots = true;
 		ImportOptions->bInterpretLODs = false;
 	}
@@ -373,10 +500,11 @@ bool FChaosClothAssetUSDImportNode::ImportFromFile(const FString& UsdFilePath, c
 
 	// Import sim mesh into collection cache 
 	// TODO: Until we have a schema so that we can use the asset cache and remove the collection cache
-	const FUsdStageOptions& StageOptions = ImportOptions->StageOptions;
-	const int AxesOrder[] = { 0, (StageOptions.UpAxis == EUsdUpAxis::ZAxis) ? 1 : 2, (StageOptions.UpAxis == EUsdUpAxis::ZAxis) ? 2 : 1 };
-	const int WindingOrder[] = { 0, (StageOptions.UpAxis == EUsdUpAxis::ZAxis) ? 2 : 1, (StageOptions.UpAxis == EUsdUpAxis::ZAxis) ? 1 : 2 };
-	const float CentimetersPerUnit = StageOptions.MetersPerUnit * 100.f;
+
+	// Retrieve stage infos
+	const int AxesOrder[] = { 0, (StageInfo.UpAxis == EUsdUpAxis::ZAxis) ? 1 : 2, (StageInfo.UpAxis == EUsdUpAxis::ZAxis) ? 2 : 1 };
+	const int WindingOrder[] = { 0, (StageInfo.UpAxis == EUsdUpAxis::ZAxis) ? 2 : 1, (StageInfo.UpAxis == EUsdUpAxis::ZAxis) ? 1 : 2 };
+	const float CentimetersPerUnit = StageInfo.MetersPerUnit * 100.f;
 
 	// Sewings
 	TArray<TArray<FIntVector2>> SeamStitches;
@@ -569,7 +697,7 @@ bool FChaosClothAssetUSDImportNode::ImportFromFile(const FString& UsdFilePath, c
 				TArray<FVector3f> OutDrapedPositions3D;
 				TArray<int32> OutIndices;
 
-				const bool bHasDegenerateTriangles = Private::RemoveDegenerateTriangles(
+				bool bHasRepairedTriangles = Private::RemoveDegenerateTriangles(
 					TriangleToVertexIndex,
 					RestPositions2D,
 					DrapedPositions3D,
@@ -578,6 +706,9 @@ bool FChaosClothAssetUSDImportNode::ImportFromFile(const FString& UsdFilePath, c
 					OutDrapedPositions3D,
 					OutIndices);
 
+				// Remove duplicate triangles
+				bHasRepairedTriangles = Private::RemoveDuplicateTriangles(OutTriangleToVertexIndex) || bHasRepairedTriangles;
+
 				// Add the new pattern
 				const int32 SimPatternIndex = ClothFacade.AddSimPattern();
 				FCollectionClothSimPatternFacade SimPattern = ClothFacade.GetSimPattern(SimPatternIndex);
@@ -605,11 +736,11 @@ bool FChaosClothAssetUSDImportNode::ImportFromFile(const FString& UsdFilePath, c
 					}
 				}
 
-				// Flag collapsed vertices for info
-				if (bHasDegenerateTriangles)
+				// Flag vertices of problem triangles for info
+				if (bHasRepairedTriangles)
 				{
 					// TODO: Make this a feature or remove it?
-					const FName WeightMapName(TEXT("_Degenerate"));  // The undescore means this is an internal weight map name
+					const FName WeightMapName(TEXT("_RepairedTriangles"));  // The undescore means this is an internal weight map name
 					const TConstArrayView<int32> SimVertex3DLookup = static_cast<FCollectionClothSimPatternConstFacade&>(SimPattern).GetSimVertex3DLookup();
 					ClothFacade.AddWeightMap(WeightMapName);
 					const TArrayView<float> WeightMap = ClothFacade.GetWeightMap(WeightMapName);
@@ -636,6 +767,9 @@ bool FChaosClothAssetUSDImportNode::ImportFromFile(const FString& UsdFilePath, c
 		}
 	}
 
+	// Check for duplicate stitches
+	Private::RemoveDuplicateStitches(SeamStitches);
+
 	// Add seams
 	for (int32 SeamIndex = 0; SeamIndex < SeamStitches.Num(); ++SeamIndex)
 	{
