Cloth Editor: fix for clicking on the same node twice making the sim mesh disappear

diff --git a/Engine/Plugins/Experimental/ChaosClothAssetEditor/Source/ChaosClothAssetEditor/Private/ChaosClothAsset/ClothEditorToolkit.cpp b/Engine/Plugins/Experimental/ChaosClothAssetEditor/Source/ChaosClothAssetEditor/Private/ChaosClothAsset/ClothEditorToolkit.cpp
--- a/Engine/Plugins/Experimental/ChaosClothAssetEditor/Source/ChaosClothAssetEditor/Private/ChaosClothAsset/ClothEditorToolkit.cpp
+++ b/Engine/Plugins/Experimental/ChaosClothAssetEditor/Source/ChaosClothAssetEditor/Private/ChaosClothAsset/ClothEditorToolkit.cpp
@@ -1018,7 +1018,7 @@ void FChaosClothAssetEditorToolkit::OnNodeTitleCommitted(const FText& InNewText,
 
 TSharedPtr<FManagedArrayCollection> FChaosClothAssetEditorToolkit::GetClothCollectionIfPossible(const TSharedPtr<FDataflowNode> InDataflowNode, const TSharedPtr<Dataflow::FEngineContext> Context)
 {
-	if (Context.IsValid())
+	if (InDataflowNode && Context)
 	{
 		for (const FDataflowOutput* const Output : InDataflowNode->GetOutputs())
 		{
@@ -1060,14 +1060,21 @@ void FChaosClothAssetEditorToolkit::OnNodeSelectionChanged(const TSet<UObject*>&
 	// Also, set the selected node(s) to be the Dataflow's RenderTargets
 	// TODO: decide if we want selection to be the mechanism for toggling DataflowComponent rendering, or the switch on the Node
 
+	// Despite this function's name, we might not have actually changed which node is selected
+	bool bNodeSelectionChanged = false;
+
 	if (Dataflow)
 	{
 		Dataflow->RenderTargets.Reset();
 
 		if (!NewSelection.Num())
 		{
+			// No new node selected
+
 			if (SelectedDataflowNode)
 			{
+				bNodeSelectionChanged = true;		// current node was deselected
+
 				if (OnNodeInvalidatedDelegateHandle.IsValid())
 				{
 					SelectedDataflowNode->GetOnNodeInvalidatedDelegate().Remove(OnNodeInvalidatedDelegateHandle);
@@ -1086,6 +1093,8 @@ void FChaosClothAssetEditorToolkit::OnNodeSelectionChanged(const TSet<UObject*>&
 
 					if (SelectedDataflowNode != Node->GetDataflowNode())
 					{
+						bNodeSelectionChanged = true;	// changing from one selected node to a different one, or from unselected to selected
+
 						if (SelectedDataflowNode)
 						{
 							if (OnNodeInvalidatedDelegateHandle.IsValid())
@@ -1128,24 +1137,27 @@ void FChaosClothAssetEditorToolkit::OnNodeSelectionChanged(const TSet<UObject*>&
 	}
 
 
-	UChaosClothAssetEditorMode* const ClothMode = CastChecked<UChaosClothAssetEditorMode>(EditorModeManager->GetActiveScriptableMode(UChaosClothAssetEditorMode::EM_ChaosClothAssetEditorModeId));
-	if (ClothMode)
+	if (bNodeSelectionChanged)
 	{
-		// Close any running tool. OnNodeSingleClicked() will start a new tool if a new node was clicked.
-		UEditorInteractiveToolsContext* const ToolsContext = ClothMode->GetInteractiveToolsContext();
-		checkf(ToolsContext, TEXT("No valid ToolsContext found for FChaosClothAssetEditorToolkit"));
-		if (ToolsContext->HasActiveTool())
+		UChaosClothAssetEditorMode* const ClothMode = CastChecked<UChaosClothAssetEditorMode>(EditorModeManager->GetActiveScriptableMode(UChaosClothAssetEditorMode::EM_ChaosClothAssetEditorModeId));
+		if (ClothMode)
 		{
-			ToolsContext->EndTool(EToolShutdownType::Completed);
-		}
+			// Close any running tool. OnNodeSingleClicked() will start a new tool if a new node was clicked.
+			UEditorInteractiveToolsContext* const ToolsContext = ClothMode->GetInteractiveToolsContext();
+			checkf(ToolsContext, TEXT("No valid ToolsContext found for FChaosClothAssetEditorToolkit"));
+			if (ToolsContext->HasActiveTool())
+			{
+				ToolsContext->EndTool(EToolShutdownType::Completed);
+			}
 
-		// Update the Construction viewport with the newly selected node's Collection
-		ClothMode->SetSelectedClothCollection(Collection);
-	}
+			// Update the Construction viewport with the newly selected node's Collection
+			ClothMode->SetSelectedClothCollection(Collection);
+		}
 
-	if (Outliner)
-	{
-		Outliner->SetClothCollection(Collection);
+		if (Outliner)
+		{
+			Outliner->SetClothCollection(Collection);
+		}
 	}
 }
 
@@ -1156,14 +1168,6 @@ void FChaosClothAssetEditorToolkit::OnNodeSingleClicked(UObject* ClickedNode) co
 	{
 		if (GraphEditor && GraphEditor->GetSingleSelectedNode() == ClickedNode)
 		{
-			// Close any running tool
-			UEditorInteractiveToolsContext* const ToolsContext = ClothMode->GetInteractiveToolsContext();
-			checkf(ToolsContext, TEXT("No valid ToolsContext found for FChaosClothAssetEditorToolkit"));
-			if (ToolsContext->HasActiveTool())
-			{
-				ToolsContext->EndTool(EToolShutdownType::Completed);
-			}
-
 			// Start the corresponding tool
 			ClothMode->StartToolForSelectedNode(ClickedNode);
 		}
