Cloth Editor: when creating a new Dataflow asset, add a Terminal Node by default

diff --git a/Engine/Plugins/Experimental/ChaosClothAssetEditor/Source/ChaosClothAssetEditor/Private/ChaosClothAsset/AssetDefinition_ClothAsset.cpp b/Engine/Plugins/Experimental/ChaosClothAssetEditor/Source/ChaosClothAssetEditor/Private/ChaosClothAsset/AssetDefinition_ClothAsset.cpp
--- a/Engine/Plugins/Experimental/ChaosClothAssetEditor/Source/ChaosClothAssetEditor/Private/ChaosClothAsset/AssetDefinition_ClothAsset.cpp
+++ b/Engine/Plugins/Experimental/ChaosClothAssetEditor/Source/ChaosClothAssetEditor/Private/ChaosClothAsset/AssetDefinition_ClothAsset.cpp
@@ -14,6 +14,8 @@
 #include "ThumbnailRendering/SceneThumbnailInfo.h"
 #include "Misc/WarnIfAssetsLoadedInScope.h"
 #include "Dialog/SMessageDialog.h"
+#include "ChaosClothAsset/TerminalNode.h"
+#include "Dataflow/DataflowSNode.h"
 
 #define LOCTEXT_NAMESPACE "AssetDefinition_ClothAsset"
 
@@ -52,7 +54,15 @@ namespace ClothAssetDefinitionHelpers
 
 		const FName NewAssetName(FPackageName::GetLongPackageAssetName(NewPackageName));
 		UPackage* const NewPackage = CreatePackage(*NewPackageName);
-		UObject* const NewAsset = NewObject<UObject>(NewPackage, DataflowClass, NewAssetName, RF_Public | RF_Standalone | RF_Transactional);
+		UDataflow* const NewAsset = NewObject<UDataflow>(NewPackage, DataflowClass, NewAssetName, RF_Public | RF_Standalone | RF_Transactional);
+
+		// Add a ClothAsset Terminal Node to the empty graph
+		const TSharedPtr<FAssetSchemaAction_Dataflow_CreateNode_DataflowEdNode> NodeAction =
+			FAssetSchemaAction_Dataflow_CreateNode_DataflowEdNode::CreateAction(NewAsset, FChaosClothAssetTerminalNode::StaticType());
+		constexpr UEdGraphPin* FromPin = nullptr;
+		constexpr bool bSelectNewNode = true;
+		UEdGraphNode* const NewEdNode = NodeAction->PerformAction(NewAsset, FromPin, FVector2D::Zero(), bSelectNewNode);
+		check(NewEdNode);
 
 		NewAsset->MarkPackageDirty();
 
