Cloth Editor: don't refocus when pressing 'F' if the viewport is empty

diff --git a/Engine/Plugins/Experimental/ChaosClothAssetEditor/Source/ChaosClothAssetEditor/Private/ChaosClothAsset/SClothEditor3DViewport.cpp b/Engine/Plugins/Experimental/ChaosClothAssetEditor/Source/ChaosClothAssetEditor/Private/ChaosClothAsset/SClothEditor3DViewport.cpp
--- a/Engine/Plugins/Experimental/ChaosClothAssetEditor/Source/ChaosClothAssetEditor/Private/ChaosClothAsset/SClothEditor3DViewport.cpp
+++ b/Engine/Plugins/Experimental/ChaosClothAssetEditor/Source/ChaosClothAssetEditor/Private/ChaosClothAsset/SClothEditor3DViewport.cpp
@@ -166,8 +166,15 @@ bool SChaosClothAssetEditor3DViewport::IsVisible() const
 
 void SChaosClothAssetEditor3DViewport::OnFocusViewportToSelection()
 {
-	const FBox PreviewBoundingBox = StaticCastSharedPtr<UE::Chaos::ClothAsset::FChaosClothAssetEditor3DViewportClient>(Client)->PreviewBoundingBox();
-	Client->FocusViewportOnBox(PreviewBoundingBox);
+	using namespace UE::Chaos::ClothAsset;
+
+	const TSharedPtr<const FChaosClothAssetEditor3DViewportClient> ViewportClient = StaticCastSharedPtr<FChaosClothAssetEditor3DViewportClient>(Client);
+	const FBox PreviewBoundingBox = ViewportClient->PreviewBoundingBox();
+
+	if (PreviewBoundingBox.IsValid && !(PreviewBoundingBox.Min == FVector::Zero() && PreviewBoundingBox.Max == FVector::Zero()))
+	{
+		Client->FocusViewportOnBox(PreviewBoundingBox);
+	}
 }
 
 TSharedRef<class SEditorViewport> SChaosClothAssetEditor3DViewport::GetViewportWidget()
diff --git a/Engine/Plugins/Experimental/ChaosClothAssetEditor/Source/ChaosClothAssetEditor/Private/ChaosClothAsset/SClothEditorRestSpaceViewport.cpp b/Engine/Plugins/Experimental/ChaosClothAssetEditor/Source/ChaosClothAssetEditor/Private/ChaosClothAsset/SClothEditorRestSpaceViewport.cpp
--- a/Engine/Plugins/Experimental/ChaosClothAssetEditor/Source/ChaosClothAssetEditor/Private/ChaosClothAsset/SClothEditorRestSpaceViewport.cpp
+++ b/Engine/Plugins/Experimental/ChaosClothAssetEditor/Source/ChaosClothAssetEditor/Private/ChaosClothAsset/SClothEditorRestSpaceViewport.cpp
@@ -167,11 +167,15 @@ void SChaosClothAssetEditorRestSpaceViewport::OnFocusViewportToSelection()
 
 	if (ClothEdMode)
 	{
-		Client->FocusViewportOnBox(ClothEdMode->SelectionBoundingBox());
+		const FBox BoundingBox = ClothEdMode->SelectionBoundingBox();
+		if (BoundingBox.IsValid && !(BoundingBox.Min == FVector::Zero() && BoundingBox.Max == FVector::Zero()))
+		{
+			Client->FocusViewportOnBox(BoundingBox);
 
-		// Reset any changes to the clip planes by the scroll zoom behavior
-		Client->OverrideNearClipPlane(UE_KINDA_SMALL_NUMBER);
-		Client->OverrideFarClipPlane(0);
+			// Reset any changes to the clip planes by the scroll zoom behavior
+			Client->OverrideNearClipPlane(UE_KINDA_SMALL_NUMBER);
+			Client->OverrideFarClipPlane(0);
+		}
 	}
 }
 
