Cloth Editor: visualization of seams in 2D and 3D sim views

diff --git a/Engine/Plugins/Experimental/ChaosClothAsset/Source/ChaosClothAsset/Private/ChaosClothAsset/ClothGeometryTools.cpp b/Engine/Plugins/Experimental/ChaosClothAsset/Source/ChaosClothAsset/Private/ChaosClothAsset/ClothGeometryTools.cpp
--- a/Engine/Plugins/Experimental/ChaosClothAsset/Source/ChaosClothAsset/Private/ChaosClothAsset/ClothGeometryTools.cpp
+++ b/Engine/Plugins/Experimental/ChaosClothAsset/Source/ChaosClothAsset/Private/ChaosClothAsset/ClothGeometryTools.cpp
@@ -957,4 +957,82 @@ namespace UE::Chaos::ClothAsset
 			Cloth.RemoveRenderPatterns(RenderPatternsToRemove);
 		}
 	}
+
+
+	void FClothGeometryTools::BuildConnectedSeams2D(const TSharedRef<const FManagedArrayCollection>& ClothCollection,
+		int32 SeamIndex,
+		const UE::Geometry::FDynamicMesh3& Mesh,
+		TArray<TArray<FIntVector2>>& Seams)
+	{
+		using namespace UE::Geometry;
+
+		FNonManifoldMappingSupport NonManifold(Mesh);
+		checkf(!NonManifold.IsNonManifoldVertexInSource(), TEXT("Cloth source is non-manifold. Cannot use FDynamicMesh to build connected seams"));
+
+		const FCollectionClothConstFacade ClothFacade(ClothCollection);
+		const FCollectionClothSeamConstFacade SeamFacade = ClothFacade.GetSeam(SeamIndex);
+		
+		TArray<FIntVector2> InputStitches(SeamFacade.GetSeamStitch2DEndIndices());
+
+		while (InputStitches.Num() > 0)
+		{
+			TArray<FIntVector2> Seam;
+
+			const FIntVector2 FirstStitch = InputStitches.Last();
+			Seam.Add(FirstStitch);
+			InputStitches.RemoveAt(InputStitches.Num() - 1);
+
+			FIntVector2 CurrStitch = FirstStitch;
+			bool bFoundNextStitch = true;
+			bool bReverseSearch = false;
+			while (InputStitches.Num() > 0 && (bFoundNextStitch || !bReverseSearch))
+			{
+				bFoundNextStitch = false;
+
+				for (int32 TestStitchIndex = 0; TestStitchIndex < InputStitches.Num(); ++TestStitchIndex)
+				{
+					FIntVector2 TestStitch = InputStitches[TestStitchIndex];
+
+					// Stitch (A, B) is connected to stitch (C, D) if there exist edges {(A, C), (B, D)} *or* {(A, D), (B, C)} in the given DynamicMesh.
+
+					const int32 A = CurrStitch[0];
+					const int32 B = CurrStitch[1];
+					const int32 C = TestStitch[0];
+					const int32 D = TestStitch[1];
+
+					if (Mesh.FindEdge(A, C) != FDynamicMesh3::InvalidID && Mesh.FindEdge(B, D) != FDynamicMesh3::InvalidID)
+					{
+						Seam.Add(TestStitch);
+						bFoundNextStitch = true;
+					}
+					else if (Mesh.FindEdge(A, D) != FDynamicMesh3::InvalidID && Mesh.FindEdge(B, C) != FDynamicMesh3::InvalidID)
+					{
+						Swap(TestStitch[0], TestStitch[1]);
+						Seam.Add(TestStitch);
+						bFoundNextStitch = true;
+					}
+
+					if (bFoundNextStitch)
+					{
+						InputStitches.RemoveAt(TestStitchIndex);
+						CurrStitch = TestStitch;
+						break;
+					}
+
+				}
+
+				if (!bFoundNextStitch && !bReverseSearch)
+				{
+					Algo::Reverse(Seam);
+					bReverseSearch = true;
+					bFoundNextStitch = true;
+					CurrStitch = FirstStitch;
+				}
+			}
+
+			// Finished one connected set of seam edges
+			Seams.Add(Seam);
+		}
+	}
+
 }  // End namespace UE::Chaos::ClothAsset
diff --git a/Engine/Plugins/Experimental/ChaosClothAsset/Source/ChaosClothAsset/Public/ChaosClothAsset/ClothGeometryTools.h b/Engine/Plugins/Experimental/ChaosClothAsset/Source/ChaosClothAsset/Public/ChaosClothAsset/ClothGeometryTools.h
--- a/Engine/Plugins/Experimental/ChaosClothAsset/Source/ChaosClothAsset/Public/ChaosClothAsset/ClothGeometryTools.h
+++ b/Engine/Plugins/Experimental/ChaosClothAsset/Source/ChaosClothAsset/Public/ChaosClothAsset/ClothGeometryTools.h
@@ -65,5 +65,16 @@ namespace UE::Chaos::ClothAsset
 		* Remove any empty patterns.
 		*/
 		static void CleanupAndCompactMesh(const TSharedRef<FManagedArrayCollection>& ClothCollection);
+
+		/**
+		* Find sets of connected stitches for the given seam.
+		* Stitch (A, B) is connected to stitch (C, D) if there exist edges {(A, C), (B, D)} *or* {(A, D), (B, C)} in the given DynamicMesh.
+		* ClothCollection meshes must be manifold.
+		*/
+		static void BuildConnectedSeams2D(const TSharedRef<const FManagedArrayCollection>& ClothCollection, 
+			int32 SeamIndex,
+			const UE::Geometry::FDynamicMesh3& Mesh,
+			TArray<TArray<FIntVector2>>& Seams);
+
 	};
 }  // End namespace UE::Chaos::ClothAsset
\ No newline at end of file
diff --git a/Engine/Plugins/Experimental/ChaosClothAssetEditor/Source/ChaosClothAssetEditor/Private/ChaosClothAsset/ClothEditorCommands.cpp b/Engine/Plugins/Experimental/ChaosClothAssetEditor/Source/ChaosClothAssetEditor/Private/ChaosClothAsset/ClothEditorCommands.cpp
--- a/Engine/Plugins/Experimental/ChaosClothAssetEditor/Source/ChaosClothAssetEditor/Private/ChaosClothAsset/ClothEditorCommands.cpp
+++ b/Engine/Plugins/Experimental/ChaosClothAssetEditor/Source/ChaosClothAssetEditor/Private/ChaosClothAsset/ClothEditorCommands.cpp
@@ -23,7 +23,8 @@ const FString FChaosClothAssetEditorCommands::SoftResetSimulationIdentifier = TE
 const FString FChaosClothAssetEditorCommands::HardResetSimulationIdentifier = TEXT("HardResetSimulation");
 const FString FChaosClothAssetEditorCommands::TogglePreviewWireframeIdentifier = TEXT("TogglePreviewWireframe");
 const FString FChaosClothAssetEditorCommands::ToggleConstructionViewWireframeIdentifier = TEXT("ToggleConstructionViewWireframe");
-
+const FString FChaosClothAssetEditorCommands::ToggleConstructionViewSeamsIdentifier = TEXT("ToggleConstructionViewSeams");
+const FString FChaosClothAssetEditorCommands::ToggleConstructionViewSeamsCollapseIdentifier = TEXT("ToggleConstructionViewSeamsCollapse");
 
 FChaosClothAssetEditorCommands::FChaosClothAssetEditorCommands()
 	: TBaseCharacterFXEditorCommands<FChaosClothAssetEditorCommands>("ChaosClothAssetEditor",
@@ -56,7 +57,10 @@ void FChaosClothAssetEditorCommands::RegisterCommands()
 
 	UI_COMMAND(TogglePreviewWireframe, "TogglePreviewWireframe", "Toggle preview wireframe", EUserInterfaceActionType::ToggleButton, FInputChord());
 	UI_COMMAND(ToggleConstructionViewWireframe, "ToggleConstructionViewWireframe", "Toggle construction view wireframe", EUserInterfaceActionType::ToggleButton, FInputChord());
-	
+
+	UI_COMMAND(ToggleConstructionViewSeams, "ToggleConstructionViewSeams", "Toggle construction view seams", EUserInterfaceActionType::ToggleButton, FInputChord());
+	UI_COMMAND(ToggleConstructionViewSeamsCollapse, "ToggleConstructionViewSeamsCollapse", "Toggle drawing one line between seams in construction view", EUserInterfaceActionType::ToggleButton, FInputChord());
+
 	UI_COMMAND(SoftResetSimulation, "SoftResetSimulation", "Soft reset simulation", EUserInterfaceActionType::ToggleButton, FInputChord());
 	UI_COMMAND(HardResetSimulation, "HardResetClothSimulation", "Hard reset simulation", EUserInterfaceActionType::ToggleButton, FInputChord(EModifierKey::Control | EModifierKey::Alt, EKeys::C));
 	UI_COMMAND(ToggleSimulationSuspended, "ToggleSimulationSuspended", "Toggle simulation suspended", EUserInterfaceActionType::ToggleButton, FInputChord());
diff --git a/Engine/Plugins/Experimental/ChaosClothAssetEditor/Source/ChaosClothAssetEditor/Private/ChaosClothAsset/ClothEditorCommands.h b/Engine/Plugins/Experimental/ChaosClothAssetEditor/Source/ChaosClothAssetEditor/Private/ChaosClothAsset/ClothEditorCommands.h
--- a/Engine/Plugins/Experimental/ChaosClothAssetEditor/Source/ChaosClothAssetEditor/Private/ChaosClothAsset/ClothEditorCommands.h
+++ b/Engine/Plugins/Experimental/ChaosClothAssetEditor/Source/ChaosClothAssetEditor/Private/ChaosClothAsset/ClothEditorCommands.h
@@ -55,6 +55,12 @@ public:
 	const static FString ToggleConstructionViewWireframeIdentifier;
 	TSharedPtr<FUICommandInfo> ToggleConstructionViewWireframe;
 
+	const static FString ToggleConstructionViewSeamsIdentifier;
+	TSharedPtr<FUICommandInfo> ToggleConstructionViewSeams;
+
+	const static FString ToggleConstructionViewSeamsCollapseIdentifier;
+	TSharedPtr<FUICommandInfo> ToggleConstructionViewSeamsCollapse;
+
 	TSharedPtr<FUICommandInfo> SetConstructionMode2D;
 	TSharedPtr<FUICommandInfo> SetConstructionMode3D;
 	TSharedPtr<FUICommandInfo> SetConstructionModeRender;
diff --git a/Engine/Plugins/Experimental/ChaosClothAssetEditor/Source/ChaosClothAssetEditor/Private/ChaosClothAsset/ClothEditorMode.cpp b/Engine/Plugins/Experimental/ChaosClothAssetEditor/Source/ChaosClothAssetEditor/Private/ChaosClothAsset/ClothEditorMode.cpp
--- a/Engine/Plugins/Experimental/ChaosClothAssetEditor/Source/ChaosClothAssetEditor/Private/ChaosClothAsset/ClothEditorMode.cpp
+++ b/Engine/Plugins/Experimental/ChaosClothAssetEditor/Source/ChaosClothAssetEditor/Private/ChaosClothAsset/ClothEditorMode.cpp
@@ -64,6 +64,8 @@
 #include "Dataflow/DataflowSNode.h"
 #include "Components/SkeletalMeshComponent.h"
 #include "ChaosClothAsset/ClothEditorToolBuilder.h"
+#include "ChaosClothAsset/ClothGeometryTools.h"
+#include "DynamicMesh/NonManifoldMappingSupport.h"
 
 #include UE_INLINE_GENERATED_CPP_BY_NAME(ClothEditorMode)
 
@@ -72,7 +74,7 @@
 const FEditorModeID UChaosClothAssetEditorMode::EM_ChaosClothAssetEditorModeId = TEXT("EM_ChaosClothAssetEditorMode");
 
 
-namespace ChaosClothAssetEditorModeHelpers
+namespace UE::Chaos::ClothAsset::Private
 {
 	void RemoveClothWeightMaps(UE::Chaos::ClothAsset::FCollectionClothFacade& ClothFacade, const TArray<FName>& WeightMapNames)
 	{
@@ -98,6 +100,7 @@ namespace ChaosClothAssetEditorModeHelpers
 
 		return OutWeightMapNames;
 	}
+
 }
 
 
@@ -193,6 +196,13 @@ void UChaosClothAssetEditorMode::RegisterClothTool(TSharedPtr<FUICommandInfo> UI
 				bConstructionViewWireframe = false;
 			}
 
+			// Seams
+			if (!bShouldRestoreConstructionViewSeams)
+			{
+				bShouldRestoreConstructionViewSeams = bConstructionViewSeamsVisible;
+			}
+			bConstructionViewSeamsVisible = false;
+
 			ActiveToolsContext = ToolsContext;
 			ToolsContext->StartTool(ToolIdentifier);
 		}),
@@ -305,6 +315,12 @@ void UChaosClothAssetEditorMode::OnToolEnded(UInteractiveToolManager* Manager, U
 		bShouldRestoreConstructionViewWireframe = false;
 	}
 
+	if (bShouldRestoreConstructionViewSeams)
+	{
+		bConstructionViewSeamsVisible = true;
+		bShouldRestoreConstructionViewSeams = false;
+	}
+
 	if (bShouldRestoreSavedConstructionViewMode)
 	{
 		SetConstructionViewMode(SavedConstructionViewMode);
@@ -362,11 +378,17 @@ void UChaosClothAssetEditorMode::Exit()
 	DynamicMeshComponent = nullptr;
 	DynamicMeshComponentParentActor = nullptr;
 
-	if (WireframeToTick)
+	if (WireframeDraw)
 	{
-		WireframeToTick->Disconnect();
+		WireframeDraw->Disconnect();
 	}
-	WireframeToTick = nullptr;
+	WireframeDraw = nullptr;
+
+	if (ClothSeamDraw)
+	{
+		ClothSeamDraw->Disconnect();
+	}
+	ClothSeamDraw = nullptr;
 
 	if (DataflowComponent)
 	{
@@ -449,6 +471,153 @@ TSharedPtr<FManagedArrayCollection> UChaosClothAssetEditorMode::GetClothCollecti
 }
 
 
+void UChaosClothAssetEditorMode::InitializeSeamDraw()
+{
+	if (!ClothSeamDraw)
+	{
+		return;
+	}
+
+	ClothSeamDraw->RemoveAllLineSets();
+
+	if (!bConstructionViewSeamsVisible)
+	{
+		return;
+	}
+
+	if (ConstructionViewMode == UE::Chaos::ClothAsset::EClothPatternVertexType::Render)
+	{
+		return;
+	}
+
+	if (!DynamicMeshComponent || !DynamicMeshComponent->GetMesh())
+	{
+		return;
+	}
+
+	const TSharedPtr<const FManagedArrayCollection> Collection = GetClothCollection();
+	if (!Collection)
+	{
+		return;
+	}
+
+	const UE::Geometry::FDynamicMesh3& Mesh = *DynamicMeshComponent->GetMesh();
+
+	// Seam view not available on non-manifold meshes
+	const UE::Geometry::FNonManifoldMappingSupport NonManifold(Mesh);
+	if (NonManifold.IsNonManifoldVertexInSource())
+	{
+		bConstructionViewSeamsVisible = false;
+		return;
+	}
+
+
+	auto PseudoRandomColor = [](int32 NumColorRotations) -> FLinearColor
+	{
+		constexpr uint8 Spread = 157;  // Prime number that gives a good spread of colors without getting too similar as a rand might do.
+		uint8 Seed = Spread;
+		for (int32 Rotation = 0; Rotation < NumColorRotations; ++Rotation)
+		{
+			Seed += Spread;
+		}
+		return FLinearColor::MakeFromHSV8(Seed, 180, 140);
+	};
+
+	const UE::Chaos::ClothAsset::FCollectionClothConstFacade ClothFacade(Collection.ToSharedRef());
+
+	ULineSetComponent* const Lines = ClothSeamDraw->AddLineSet("SeamLines");
+	const bool bDepthTested = (ConstructionViewMode == UE::Chaos::ClothAsset::EClothPatternVertexType::Sim3D);
+	Lines->SetLineMaterial(ToolSetupUtil::GetDefaultLineComponentMaterial(GetToolManager(), bDepthTested));
+	const float LineThickness = (ConstructionViewMode == UE::Chaos::ClothAsset::EClothPatternVertexType::Sim3D) ? 4.0f : 2.0f;
+
+	UPointSetComponent* const Points = ClothSeamDraw->AddPointSet("SeamPoints");
+	constexpr float PointSize = 4.0f;
+
+
+	int32 ConnectedSeamIndex = 0;		// Used to generate different colors for each connected seam, if multiple connected seams are found per input seam
+
+	for (int32 SeamIndex = 0; SeamIndex < ClothFacade.GetNumSeams(); ++SeamIndex)
+	{
+		const UE::Chaos::ClothAsset::FCollectionClothSeamConstFacade SeamFacade = ClothFacade.GetSeam(SeamIndex);
+
+		if (ConstructionViewMode == UE::Chaos::ClothAsset::EClothPatternVertexType::Sim2D)
+		{
+			// Stitches are given in random order, so first construct paths of connected stitches
+			// Note one SeamFacade can contain multiple disjoint paths
+			TArray<TArray<FIntVector2>> ConnectedSeams;
+			UE::Chaos::ClothAsset::FClothGeometryTools::BuildConnectedSeams2D(Collection.ToSharedRef(), SeamIndex, Mesh, ConnectedSeams);
+
+			for (const TArray<FIntVector2>& ConnectedSeam : ConnectedSeams)
+			{
+				const FColor SeamColor = PseudoRandomColor(ConnectedSeamIndex++).ToFColor(true);
+
+				// draw connected edge on each side of the seam
+				for (int32 StitchID = 0; StitchID < ConnectedSeam.Num() - 1; ++StitchID)
+				{
+					const FVector3d PointA = Mesh.GetVertex(ConnectedSeam[StitchID][0]);
+					const FVector3d PointB = Mesh.GetVertex(ConnectedSeam[StitchID][1]);
+					const FVector3d PointC = Mesh.GetVertex(ConnectedSeam[StitchID+1][0]);
+					const FVector3d PointD = Mesh.GetVertex(ConnectedSeam[StitchID+1][1]);
+					Lines->AddLine(PointA, PointC, SeamColor, LineThickness);
+					Lines->AddLine(PointB, PointD, SeamColor, LineThickness);
+				}
+
+				// draw connection between stitch points
+				if (bConstructionViewSeamsCollapse)
+				{
+					const int32 StitchID = ConnectedSeam.Num() / 2;
+					const int32 VertexA = ConnectedSeam[StitchID][0];
+					const int32 VertexB = ConnectedSeam[StitchID][1];
+					const FVector PtA = Mesh.GetVertex(VertexA);
+					const FVector PtB = Mesh.GetVertex(VertexB);
+					Lines->AddLine(PtA, PtB, SeamColor, LineThickness);
+					Points->AddPoint(PtA, SeamColor, PointSize);
+					Points->AddPoint(PtB, SeamColor, PointSize);
+				}
+				else
+				{
+					for (int32 StitchID = 0; StitchID < ConnectedSeam.Num(); ++StitchID)
+					{
+						const int32 VertexA = ConnectedSeam[StitchID][0];
+						const int32 VertexB = ConnectedSeam[StitchID][1];
+						const FVector PtA = Mesh.GetVertex(VertexA);
+						const FVector PtB = Mesh.GetVertex(VertexB);
+						Lines->AddLine(PtA, PtB, SeamColor, 2.0f);
+						Points->AddPoint(PtA, SeamColor, PointSize);
+						Points->AddPoint(PtB, SeamColor, PointSize);
+					}
+				}
+			}
+		}
+		else if (ConstructionViewMode == UE::Chaos::ClothAsset::EClothPatternVertexType::Sim3D)
+		{
+			const TArray<int32> SeamStitches(SeamFacade.GetSeamStitch3DIndex());
+			const FColor SeamColor = PseudoRandomColor(SeamIndex).ToFColor(true);
+
+			// In 3D we should be able to draw the seam edges in any order, doesn't need to be in connected paths
+			for (int32 StitchIndexI = 0; StitchIndexI < SeamStitches.Num(); ++StitchIndexI)
+			{
+				const int32 VertexIndexI = SeamStitches[StitchIndexI];
+				for (int32 StitchIndexJ = StitchIndexI + 1; StitchIndexJ < SeamStitches.Num(); ++StitchIndexJ)
+				{
+					const int32 VertexIndexJ = SeamStitches[StitchIndexJ];
+
+					if (Mesh.FindEdge(VertexIndexI, VertexIndexJ) != UE::Geometry::FDynamicMesh3::InvalidID)
+					{
+						const FVector PointI = Mesh.GetVertex(VertexIndexI);
+						const FVector PointJ = Mesh.GetVertex(VertexIndexJ);
+						Lines->AddLine(PointI, PointJ, SeamColor, LineThickness);
+						Points->AddPoint(PointI, SeamColor, PointSize);
+						Points->AddPoint(PointJ, SeamColor, PointSize);
+					}
+				}
+			}
+
+		}
+	}
+}
+
+
 void UChaosClothAssetEditorMode::ReinitializeDynamicMeshComponents()
 {
 	using namespace UE::Chaos::ClothAsset;
@@ -517,15 +686,21 @@ void UChaosClothAssetEditorMode::ReinitializeDynamicMeshComponents()
 		}
 	}
 
-	if (WireframeToTick)
+	if (WireframeDraw)
+	{
+		WireframeDraw->Disconnect();
+	}
+
+	if (ClothSeamDraw)
 	{
-		WireframeToTick->Disconnect();
+		ClothSeamDraw->Disconnect();
 	}
 
 	PropertyObjectsToTick.Empty();	// TODO: We only want to empty the wireframe display properties. Is anything else using this array?
 	DynamicMeshComponent = nullptr;
 	DynamicMeshComponentParentActor = nullptr;
-	WireframeToTick = nullptr;
+	WireframeDraw = nullptr;
+	ClothSeamDraw = nullptr;
 
 	TSharedPtr<FManagedArrayCollection> Collection = GetClothCollection();
 	if (!Collection)
@@ -560,21 +735,21 @@ void UChaosClothAssetEditorMode::ReinitializeDynamicMeshComponents()
 	DynamicMeshComponent->RegisterComponentWithWorld(this->GetWorld());
 
 	// Set up the wireframe display of the rest space mesh.
-	WireframeToTick = NewObject<UMeshElementsVisualizer>(this);
-	WireframeToTick->CreateInWorld(GetWorld(), FTransform::Identity);
+	WireframeDraw = NewObject<UMeshElementsVisualizer>(this);
+	WireframeDraw->CreateInWorld(GetWorld(), FTransform::Identity);
 
-	WireframeToTick->Settings->DepthBias = 2.0;
-	WireframeToTick->Settings->bAdjustDepthBiasUsingMeshSize = false;
-	WireframeToTick->Settings->bShowWireframe = true;
-	WireframeToTick->Settings->bShowBorders = true;
-	WireframeToTick->Settings->bShowUVSeams = false;
-	WireframeToTick->Settings->bShowNormalSeams = false;
+	WireframeDraw->Settings->DepthBias = 2.0;
+	WireframeDraw->Settings->bAdjustDepthBiasUsingMeshSize = false;
+	WireframeDraw->Settings->bShowWireframe = true;
+	WireframeDraw->Settings->bShowBorders = true;
+	WireframeDraw->Settings->bShowUVSeams = false;
+	WireframeDraw->Settings->bShowNormalSeams = false;
 
 	// These are not exposed at the visualizer level yet
 	// TODO: Should they be?
-	WireframeToTick->WireframeComponent->BoundaryEdgeThickness = 2;
+	WireframeDraw->WireframeComponent->BoundaryEdgeThickness = 2;
 
-	WireframeToTick->SetMeshAccessFunction([this](UMeshElementsVisualizer::ProcessDynamicMeshFunc ProcessFunc)
+	WireframeDraw->SetMeshAccessFunction([this](UMeshElementsVisualizer::ProcessDynamicMeshFunc ProcessFunc)
 	{
 		ProcessFunc(*DynamicMeshComponent->GetMesh());
 	});
@@ -582,27 +757,50 @@ void UChaosClothAssetEditorMode::ReinitializeDynamicMeshComponents()
 	DynamicMeshComponent->OnMeshChanged.Add(
 		FSimpleMulticastDelegate::FDelegate::CreateLambda([this]()
 		{
-			WireframeToTick->NotifyMeshChanged();
+			WireframeDraw->NotifyMeshChanged();
 		}));
 
 	// The settings object and wireframe are not part of a tool, so they won't get ticked like they
 	// are supposed to (to enable property watching), unless we add this here.
-	PropertyObjectsToTick.Add(WireframeToTick->Settings);
+	PropertyObjectsToTick.Add(WireframeDraw->Settings);
+
+	const bool bRestSpaceMeshVisible = DynamicMeshComponent->GetVisibleFlag();
+	WireframeDraw->Settings->bVisible = bRestSpaceMeshVisible && bConstructionViewWireframe;
+
+
+	ClothSeamDraw = NewObject<UPreviewGeometry>(this);
+	ClothSeamDraw->CreateInWorld(GetWorld(), FTransform::Identity);
+	InitializeSeamDraw();
+	ClothSeamDraw->SetAllVisible(bRestSpaceMeshVisible && bConstructionViewSeamsVisible);
+
+	DynamicMeshComponent->OnMeshChanged.Add(
+		FSimpleMulticastDelegate::FDelegate::CreateLambda([this]()
+		{
+			InitializeSeamDraw();
+			const bool bRestSpaceMeshVisible = DynamicMeshComponent->GetVisibleFlag();
+			ClothSeamDraw->SetAllVisible(bRestSpaceMeshVisible&& bConstructionViewSeamsVisible);
+		}));
+
 
 	// Some interactive tools will hide the input DynamicMeshComponent and create their own temporary PreviewMesh for visualization. If this
-	// occurs, we should also hide the corresponding WireframeDisplay (and un-hide it when the tool finishes).
+	// occurs, we should also hide the corresponding Wireframe and Seam drawing (and un-hide it when the tool finishes).
 	UActorComponent::MarkRenderStateDirtyEvent.AddWeakLambda(this, [this](UActorComponent& ActorComponent)
-	{
-		if (!WireframeToTick || !DynamicMeshComponent)
 		{
-			return;
-		}
-		const bool bRestSpaceMeshVisible = DynamicMeshComponent->GetVisibleFlag();
-		WireframeToTick->Settings->bVisible = bRestSpaceMeshVisible && bConstructionViewWireframe;
-	});
+			if (!DynamicMeshComponent)
+			{
+				return;
+			}
+			const bool bRestSpaceMeshVisible = DynamicMeshComponent->GetVisibleFlag();
+			if (WireframeDraw)
+			{
+				WireframeDraw->Settings->bVisible = bRestSpaceMeshVisible && bConstructionViewWireframe;
+			}
+			if (ClothSeamDraw)
+			{
+				ClothSeamDraw->SetAllVisible(bRestSpaceMeshVisible && bConstructionViewSeamsVisible);
+			}
+		});
 
-	const bool bRestSpaceMeshVisible = DynamicMeshComponent->GetVisibleFlag();
-	WireframeToTick->Settings->bVisible = bRestSpaceMeshVisible && bConstructionViewWireframe;
 
 	SelectedComponents->DeselectAll();
 	SelectedComponents->Select(DynamicMeshComponent);
@@ -812,9 +1010,9 @@ void UChaosClothAssetEditorMode::ModeTick(float DeltaTime)
 		}
 	}
 
-	if (WireframeToTick)
+	if (WireframeDraw)
 	{
-		WireframeToTick->OnTick(DeltaTime);
+		WireframeDraw->OnTick(DeltaTime);
 	}
 
 
@@ -1010,6 +1208,52 @@ bool UChaosClothAssetEditorMode::CanSetConstructionViewWireframeActive() const
 	return ClothToolBuilder->CanSetConstructionViewWireframeActive();
 }
 
+
+void UChaosClothAssetEditorMode::ToggleConstructionViewSeams()
+{
+	bConstructionViewSeamsVisible = !bConstructionViewSeamsVisible;
+	ReinitializeDynamicMeshComponents();
+}
+
+bool UChaosClothAssetEditorMode::CanSetConstructionViewSeamsActive() const
+{
+	// Disallow seam view when any tool is active
+	if (GetToolManager()->HasActiveTool(EToolSide::Left))
+	{
+		return false;
+	}
+
+	// Seam view not available on non-manifold meshes
+	if (const UE::Geometry::FDynamicMesh3* const Mesh = DynamicMeshComponent->GetMesh())
+	{
+		const UE::Geometry::FNonManifoldMappingSupport NonManifold(*Mesh);
+		if (NonManifold.IsNonManifoldVertexInSource())
+		{
+			return false;
+		}
+	}
+
+	return true;
+}
+
+
+void UChaosClothAssetEditorMode::ToggleConstructionViewSeamsCollapse()
+{
+	bConstructionViewSeamsCollapse = !bConstructionViewSeamsCollapse;
+	ReinitializeDynamicMeshComponents();
+}
+
+bool UChaosClothAssetEditorMode::CanSetConstructionViewSeamsCollapse() const
+{
+	// Disallow seam view when any tool is active
+	if (GetToolManager()->HasActiveTool(EToolSide::Left))
+	{
+		return false;
+	}
+
+	return bConstructionViewSeamsVisible && (ConstructionViewMode == UE::Chaos::ClothAsset::EClothPatternVertexType::Sim2D);
+}
+
 void UChaosClothAssetEditorMode::SetRestSpaceViewportClient(TWeakPtr<UE::Chaos::ClothAsset::FChaosClothEditorRestSpaceViewportClient, ESPMode::ThreadSafe> InViewportClient)
 {
 	RestSpaceViewportClient = InViewportClient;
diff --git a/Engine/Plugins/Experimental/ChaosClothAssetEditor/Source/ChaosClothAssetEditor/Private/ChaosClothAsset/ClothEditorMode.h b/Engine/Plugins/Experimental/ChaosClothAssetEditor/Source/ChaosClothAssetEditor/Private/ChaosClothAsset/ClothEditorMode.h
--- a/Engine/Plugins/Experimental/ChaosClothAssetEditor/Source/ChaosClothAssetEditor/Private/ChaosClothAsset/ClothEditorMode.h
+++ b/Engine/Plugins/Experimental/ChaosClothAssetEditor/Source/ChaosClothAssetEditor/Private/ChaosClothAsset/ClothEditorMode.h
@@ -38,6 +38,7 @@ class FChaosClothEditorRestSpaceViewportClient;
 }
 class IChaosClothAssetEditorToolBuilder;
 class UEdGraphNode;
+class UPreviewGeometry;
 
 /**
  * The cloth editor mode is the mode used in the cloth asset editor. It holds most of the inter-tool state.
@@ -72,6 +73,20 @@ public:
 		return bConstructionViewWireframe;
 	}
 
+	void ToggleConstructionViewSeams();
+	bool CanSetConstructionViewSeamsActive() const;
+	bool IsConstructionViewSeamsActive() const
+	{
+		return bConstructionViewSeamsVisible;
+	}
+
+	void ToggleConstructionViewSeamsCollapse();
+	bool CanSetConstructionViewSeamsCollapse() const;
+	bool IsConstructionViewSeamsCollapseActive() const
+	{
+		return bConstructionViewSeamsCollapse;
+	}
+
 	// Simulation controls
 	void SoftResetSimulation();
 	void HardResetSimulation();
@@ -194,7 +209,11 @@ private:
 
 	// Rest space wireframe. They have to get ticked to be able to respond to setting changes. 
 	UPROPERTY()
-	TObjectPtr<UMeshElementsVisualizer> WireframeToTick = nullptr;
+	TObjectPtr<UMeshElementsVisualizer> WireframeDraw = nullptr;
+
+	UPROPERTY()
+	TObjectPtr<UPreviewGeometry> ClothSeamDraw = nullptr;
+
 
 	// Preview Scene, here largely for convenience to avoid having to pass it around functions. Owned by the ClothEditorToolkit.
 	UE::Chaos::ClothAsset::FChaosClothPreviewScene* PreviewScene = nullptr;
@@ -236,6 +255,11 @@ private:
 	bool bConstructionViewWireframe = false;
 	bool bShouldRestoreConstructionViewWireframe = false;
 
+	bool bConstructionViewSeamsVisible = false;
+	bool bShouldRestoreConstructionViewSeams = false;
+	bool bConstructionViewSeamsCollapse = false;
+	void InitializeSeamDraw();
+
 	// Create dynamic mesh components from the cloth component's rest space info
 	void ReinitializeDynamicMeshComponents();
 
diff --git a/Engine/Plugins/Experimental/ChaosClothAssetEditor/Source/ChaosClothAssetEditor/Private/ChaosClothAsset/SClothEditorRestSpaceViewport.cpp b/Engine/Plugins/Experimental/ChaosClothAssetEditor/Source/ChaosClothAssetEditor/Private/ChaosClothAsset/SClothEditorRestSpaceViewport.cpp
--- a/Engine/Plugins/Experimental/ChaosClothAssetEditor/Source/ChaosClothAssetEditor/Private/ChaosClothAsset/SClothEditorRestSpaceViewport.cpp
+++ b/Engine/Plugins/Experimental/ChaosClothAssetEditor/Source/ChaosClothAssetEditor/Private/ChaosClothAsset/SClothEditorRestSpaceViewport.cpp
@@ -139,7 +139,11 @@ void SChaosClothAssetEditorRestSpaceViewport::BindCommands()
 		}),
 		FCanExecuteAction::CreateLambda([this]() 
 		{ 
-			return true; 
+			if (const UChaosClothAssetEditorMode* const ClothEdMode = GetEdMode())
+			{
+				return ClothEdMode->CanSetConstructionViewWireframeActive();
+			}
+			return false;
 		}),
 		FIsActionChecked::CreateLambda([this]() 
 		{
@@ -150,6 +154,57 @@ void SChaosClothAssetEditorRestSpaceViewport::BindCommands()
 			return false;
 		}));
 
+	CommandList->MapAction(
+		CommandInfos.ToggleConstructionViewSeams,
+		FExecuteAction::CreateLambda([this]()
+		{
+			if (UChaosClothAssetEditorMode* const ClothEdMode = GetEdMode())
+			{
+				ClothEdMode->ToggleConstructionViewSeams();
+			}
+		}),
+		FCanExecuteAction::CreateLambda([this]()
+		{
+			if (const UChaosClothAssetEditorMode* const ClothEdMode = GetEdMode())
+			{
+				return ClothEdMode->CanSetConstructionViewSeamsActive();
+			}
+			return false;
+		}),
+		FIsActionChecked::CreateLambda([this]()
+		{
+			if (const UChaosClothAssetEditorMode* const ClothEdMode = GetEdMode())
+			{
+				return ClothEdMode->IsConstructionViewSeamsActive();
+			}
+			return false;
+		}));
+
+	CommandList->MapAction(
+		CommandInfos.ToggleConstructionViewSeamsCollapse,
+		FExecuteAction::CreateLambda([this]()
+		{
+			if (UChaosClothAssetEditorMode* const ClothEdMode = GetEdMode())
+			{
+				ClothEdMode->ToggleConstructionViewSeamsCollapse();
+			}
+		}),
+		FCanExecuteAction::CreateLambda([this]()
+		{
+			if (const UChaosClothAssetEditorMode* const ClothEdMode = GetEdMode())
+			{
+				return ClothEdMode->CanSetConstructionViewSeamsCollapse();
+			}
+			return false;
+		}),
+		FIsActionChecked::CreateLambda([this]()
+		{
+			if (const UChaosClothAssetEditorMode* const ClothEdMode = GetEdMode())
+			{
+				return ClothEdMode->IsConstructionViewSeamsCollapseActive();
+			}
+			return false;
+		}));
 
 }
 
diff --git a/Engine/Plugins/Experimental/ChaosClothAssetEditor/Source/ChaosClothAssetEditor/Private/ChaosClothAsset/SClothEditorRestSpaceViewportToolBar.cpp b/Engine/Plugins/Experimental/ChaosClothAssetEditor/Source/ChaosClothAssetEditor/Private/ChaosClothAsset/SClothEditorRestSpaceViewportToolBar.cpp
--- a/Engine/Plugins/Experimental/ChaosClothAssetEditor/Source/ChaosClothAssetEditor/Private/ChaosClothAsset/SClothEditorRestSpaceViewportToolBar.cpp
+++ b/Engine/Plugins/Experimental/ChaosClothAssetEditor/Source/ChaosClothAssetEditor/Private/ChaosClothAsset/SClothEditorRestSpaceViewportToolBar.cpp
@@ -66,6 +66,8 @@ void SChaosClothAssetEditorRestSpaceViewportToolBar::Construct(const FArguments&
 
 TSharedRef<SWidget> SChaosClothAssetEditorRestSpaceViewportToolBar::GenerateClothRestSpaceViewportOptionsMenu()
 {
+	using namespace UE::Chaos::ClothAsset;
+
 	GetInfoProvider().OnFloatingButtonClicked();
 	TSharedRef<SEditorViewport> ViewportRef = GetInfoProvider().GetViewportWidget();
 
@@ -81,6 +83,16 @@ TSharedRef<SWidget> SChaosClothAssetEditorRestSpaceViewportToolBar::GenerateClot
 				OptionsMenuBuilder.AddWidget(GenerateFOVMenu(), LOCTEXT("FOVAngle", "Field of View (H)"));
 			}
 			OptionsMenuBuilder.AddWidget(GenerateLightMenu(), LOCTEXT("LightIntensity", "Render Light Intensity"));
+
+			OptionsMenuBuilder.AddMenuEntry(FChaosClothAssetEditorCommands::Get().ToggleConstructionViewSeams, 
+				NAME_None, 
+				LOCTEXT("ShowSeamsLabel", "Show Seams"), 
+				LOCTEXT("ShowSeamsTooltip", "Display seam information (not available for non-manifold meshes)"));
+
+			OptionsMenuBuilder.AddMenuEntry(FChaosClothAssetEditorCommands::Get().ToggleConstructionViewSeamsCollapse, 
+				NAME_None, 
+				LOCTEXT("SeamsCollapseLabel", "Collapse Seam Lines"), 
+				LOCTEXT("SeamsCollapseTooltip", "Display a single line connecting each seam, rather than all stitches"));
 		}
 		OptionsMenuBuilder.EndSection();
 
