Chaos Cloth Asset:
- Fixed the module name in ChaosClothAssetToolsModule (was ChaosClothAssetEditor instead of ChaosClothAssetToolsModule).
- Removed unnecessary headers, exports, and includes whenever the module class is private and isn't used outside of the implement module requirement.

diff --git a/Engine/Plugins/Experimental/ChaosClothAsset/Source/ChaosClothAsset/Private/ChaosClothAsset/ChaosClothAssetModule.cpp b/Engine/Plugins/Experimental/ChaosClothAsset/Source/ChaosClothAsset/Private/ChaosClothAsset/ChaosClothAssetModule.cpp
--- a/Engine/Plugins/Experimental/ChaosClothAsset/Source/ChaosClothAsset/Private/ChaosClothAsset/ChaosClothAssetModule.cpp
+++ b/Engine/Plugins/Experimental/ChaosClothAsset/Source/ChaosClothAsset/Private/ChaosClothAsset/ChaosClothAssetModule.cpp
@@ -1,7 +1,5 @@
 // Copyright Epic Games, Inc. All Rights Reserved.
 
-#include "CoreMinimal.h"
-#include "Modules/ModuleInterface.h"
 #include "Modules/ModuleManager.h"
 
 namespace UE::Chaos::ClothAsset
diff --git a/Engine/Plugins/Experimental/ChaosClothAssetEditor/Source/ChaosClothAssetDataflowNodes/Private/ChaosClothAsset/ChaosClothAssetDataflowNodesModule.cpp b/Engine/Plugins/Experimental/ChaosClothAssetEditor/Source/ChaosClothAssetDataflowNodes/Private/ChaosClothAsset/ChaosClothAssetDataflowNodesModule.cpp
--- a/Engine/Plugins/Experimental/ChaosClothAssetEditor/Source/ChaosClothAssetDataflowNodes/Private/ChaosClothAsset/ChaosClothAssetDataflowNodesModule.cpp
+++ b/Engine/Plugins/Experimental/ChaosClothAssetEditor/Source/ChaosClothAssetDataflowNodes/Private/ChaosClothAsset/ChaosClothAssetDataflowNodesModule.cpp
@@ -1,6 +1,5 @@
 // Copyright Epic Games, Inc. All Rights Reserved.
 
-#include "ChaosClothAsset/ChaosClothAssetDataflowNodesModule.h"
 #include "ChaosClothAsset/AddWeightMapNode.h"
 #include "ChaosClothAsset/BindToRootBoneNode.h"
 #include "ChaosClothAsset/ColorScheme.h"
@@ -51,93 +50,96 @@
 #include "Dataflow/DataflowNodeFactory.h"
 #include "Modules/ModuleManager.h"
 
-#define LOCTEXT_NAMESPACE "FChaosClothAssetDataflowNodesModule"
-
-namespace UE::Chaos::ClothAsset::Private
+namespace UE::Chaos::ClothAsset
 {
-	static void RegisterDataflowNodes()
+	namespace Private
 	{
-		DATAFLOW_NODE_REGISTER_CREATION_FACTORY(FChaosClothAssetAddWeightMapNode);
-		DATAFLOW_NODE_REGISTER_CREATION_FACTORY(FChaosClothAssetBindToRootBoneNode);
-		DATAFLOW_NODE_REGISTER_CREATION_FACTORY(FChaosClothAssetCopySimulationToRenderMeshNode);
-		DATAFLOW_NODE_REGISTER_CREATION_FACTORY(FChaosClothAssetDatasmithImportNode);
-		DATAFLOW_NODE_REGISTER_CREATION_FACTORY(FChaosClothAssetDeleteElementNode);
-		DATAFLOW_NODE_REGISTER_CREATION_FACTORY(FChaosClothAssetImportNode);
-		DATAFLOW_NODE_REGISTER_CREATION_FACTORY(FChaosClothAssetMergeClothCollectionsNode);
-		DATAFLOW_NODE_REGISTER_CREATION_FACTORY(FChaosClothAssetReverseNormalsNode);
-		DATAFLOW_NODE_REGISTER_CREATION_FACTORY(FChaosClothAssetSelectionNode);
-		DATAFLOW_NODE_REGISTER_CREATION_FACTORY(FChaosClothAssetSelectionToWeightMapNode);
-		DATAFLOW_NODE_REGISTER_CREATION_FACTORY(FChaosClothAssetSetPhysicsAssetNode);
-		DATAFLOW_NODE_REGISTER_CREATION_FACTORY(FChaosClothAssetSimulationAerodynamicsConfigNode);
-		DATAFLOW_NODE_REGISTER_CREATION_FACTORY(FChaosClothAssetSimulationAnimDriveConfigNode);
-		DATAFLOW_NODE_REGISTER_CREATION_FACTORY(FChaosClothAssetSimulationBackstopConfigNode);
-		DATAFLOW_NODE_REGISTER_CREATION_FACTORY(FChaosClothAssetSimulationCollisionConfigNode);
-		DATAFLOW_NODE_REGISTER_CREATION_FACTORY(FChaosClothAssetSimulationDampingConfigNode);
-		DATAFLOW_NODE_REGISTER_CREATION_FACTORY(FChaosClothAssetSimulationDefaultConfigNode);
-		DATAFLOW_NODE_REGISTER_CREATION_FACTORY(FChaosClothAssetSimulationGravityConfigNode);
-		DATAFLOW_NODE_REGISTER_CREATION_FACTORY(FChaosClothAssetSimulationLongRangeAttachmentConfigNode);
-		DATAFLOW_NODE_REGISTER_CREATION_FACTORY(FChaosClothAssetSimulationMassConfigNode);
-		DATAFLOW_NODE_REGISTER_CREATION_FACTORY(FChaosClothAssetSimulationMaxDistanceConfigNode);
-		DATAFLOW_NODE_REGISTER_CREATION_FACTORY(FChaosClothAssetSimulationPBDAreaSpringConfigNode);
-		DATAFLOW_NODE_REGISTER_CREATION_FACTORY(FChaosClothAssetSimulationPBDBendingElementConfigNode);
-		DATAFLOW_NODE_REGISTER_CREATION_FACTORY(FChaosClothAssetSimulationPBDBendingSpringConfigNode);
-		DATAFLOW_NODE_REGISTER_CREATION_FACTORY(FChaosClothAssetSimulationPBDEdgeSpringConfigNode);
-		DATAFLOW_NODE_REGISTER_CREATION_FACTORY(FChaosClothAssetSimulationPressureConfigNode);
-		DATAFLOW_NODE_REGISTER_CREATION_FACTORY(FChaosClothAssetSimulationSelfCollisionConfigNode);
-		DATAFLOW_NODE_REGISTER_CREATION_FACTORY(FChaosClothAssetSimulationSelfCollisionSpheresConfigNode);
-		DATAFLOW_NODE_REGISTER_CREATION_FACTORY(FChaosClothAssetSimulationSolverConfigNode);
-		DATAFLOW_NODE_REGISTER_CREATION_FACTORY(FChaosClothAssetSimulationVelocityScaleConfigNode);
-		DATAFLOW_NODE_REGISTER_CREATION_FACTORY(FChaosClothAssetSimulationXPBDAnisoBendingConfigNode);
-		DATAFLOW_NODE_REGISTER_CREATION_FACTORY(FChaosClothAssetSimulationXPBDAnisoStretchConfigNode);
-		DATAFLOW_NODE_REGISTER_CREATION_FACTORY(FChaosClothAssetSimulationXPBDAreaSpringConfigNode);
-		DATAFLOW_NODE_REGISTER_CREATION_FACTORY(FChaosClothAssetSimulationXPBDBendingElementConfigNode);
-		DATAFLOW_NODE_REGISTER_CREATION_FACTORY(FChaosClothAssetSimulationXPBDBendingSpringConfigNode);
-		DATAFLOW_NODE_REGISTER_CREATION_FACTORY(FChaosClothAssetSimulationXPBDEdgeSpringConfigNode);
-		DATAFLOW_NODE_REGISTER_CREATION_FACTORY(FChaosClothAssetSkeletalMeshImportNode);
-		DATAFLOW_NODE_REGISTER_CREATION_FACTORY(FChaosClothAssetStaticMeshImportNode);
-		DATAFLOW_NODE_REGISTER_CREATION_FACTORY(FChaosClothAssetTerminalNode);
-		DATAFLOW_NODE_REGISTER_CREATION_FACTORY(FChaosClothAssetTransferSkinWeightsNode);
-		DATAFLOW_NODE_REGISTER_CREATION_FACTORY(FChaosClothAssetTransformUVsNode);
-		DATAFLOW_NODE_REGISTER_CREATION_FACTORY(FChaosClothAssetUSDImportNode);
-		DATAFLOW_NODE_REGISTER_CREATION_FACTORY_NODE_COLORS_BY_CATEGORY("Cloth", FColorScheme::NodeHeader, FColorScheme::NodeBody);
-	}
-}
-
-void FChaosClothAssetDataflowNodesModule::StartupModule()
-{
-	using namespace UE::Chaos::ClothAsset;
-
-	Private::RegisterDataflowNodes();
+		static void RegisterDataflowNodes()
+		{
+			DATAFLOW_NODE_REGISTER_CREATION_FACTORY(FChaosClothAssetAddWeightMapNode);
+			DATAFLOW_NODE_REGISTER_CREATION_FACTORY(FChaosClothAssetBindToRootBoneNode);
+			DATAFLOW_NODE_REGISTER_CREATION_FACTORY(FChaosClothAssetCopySimulationToRenderMeshNode);
+			DATAFLOW_NODE_REGISTER_CREATION_FACTORY(FChaosClothAssetDatasmithImportNode);
+			DATAFLOW_NODE_REGISTER_CREATION_FACTORY(FChaosClothAssetDeleteElementNode);
+			DATAFLOW_NODE_REGISTER_CREATION_FACTORY(FChaosClothAssetImportNode);
+			DATAFLOW_NODE_REGISTER_CREATION_FACTORY(FChaosClothAssetMergeClothCollectionsNode);
+			DATAFLOW_NODE_REGISTER_CREATION_FACTORY(FChaosClothAssetReverseNormalsNode);
+			DATAFLOW_NODE_REGISTER_CREATION_FACTORY(FChaosClothAssetSelectionNode);
+			DATAFLOW_NODE_REGISTER_CREATION_FACTORY(FChaosClothAssetSelectionToWeightMapNode);
+			DATAFLOW_NODE_REGISTER_CREATION_FACTORY(FChaosClothAssetSetPhysicsAssetNode);
+			DATAFLOW_NODE_REGISTER_CREATION_FACTORY(FChaosClothAssetSimulationAerodynamicsConfigNode);
+			DATAFLOW_NODE_REGISTER_CREATION_FACTORY(FChaosClothAssetSimulationAnimDriveConfigNode);
+			DATAFLOW_NODE_REGISTER_CREATION_FACTORY(FChaosClothAssetSimulationBackstopConfigNode);
+			DATAFLOW_NODE_REGISTER_CREATION_FACTORY(FChaosClothAssetSimulationCollisionConfigNode);
+			DATAFLOW_NODE_REGISTER_CREATION_FACTORY(FChaosClothAssetSimulationDampingConfigNode);
+			DATAFLOW_NODE_REGISTER_CREATION_FACTORY(FChaosClothAssetSimulationDefaultConfigNode);
+			DATAFLOW_NODE_REGISTER_CREATION_FACTORY(FChaosClothAssetSimulationGravityConfigNode);
+			DATAFLOW_NODE_REGISTER_CREATION_FACTORY(FChaosClothAssetSimulationLongRangeAttachmentConfigNode);
+			DATAFLOW_NODE_REGISTER_CREATION_FACTORY(FChaosClothAssetSimulationMassConfigNode);
+			DATAFLOW_NODE_REGISTER_CREATION_FACTORY(FChaosClothAssetSimulationMaxDistanceConfigNode);
+			DATAFLOW_NODE_REGISTER_CREATION_FACTORY(FChaosClothAssetSimulationPBDAreaSpringConfigNode);
+			DATAFLOW_NODE_REGISTER_CREATION_FACTORY(FChaosClothAssetSimulationPBDBendingElementConfigNode);
+			DATAFLOW_NODE_REGISTER_CREATION_FACTORY(FChaosClothAssetSimulationPBDBendingSpringConfigNode);
+			DATAFLOW_NODE_REGISTER_CREATION_FACTORY(FChaosClothAssetSimulationPBDEdgeSpringConfigNode);
+			DATAFLOW_NODE_REGISTER_CREATION_FACTORY(FChaosClothAssetSimulationPressureConfigNode);
+			DATAFLOW_NODE_REGISTER_CREATION_FACTORY(FChaosClothAssetSimulationSelfCollisionConfigNode);
+			DATAFLOW_NODE_REGISTER_CREATION_FACTORY(FChaosClothAssetSimulationSelfCollisionSpheresConfigNode);
+			DATAFLOW_NODE_REGISTER_CREATION_FACTORY(FChaosClothAssetSimulationSolverConfigNode);
+			DATAFLOW_NODE_REGISTER_CREATION_FACTORY(FChaosClothAssetSimulationVelocityScaleConfigNode);
+			DATAFLOW_NODE_REGISTER_CREATION_FACTORY(FChaosClothAssetSimulationXPBDAnisoBendingConfigNode);
+			DATAFLOW_NODE_REGISTER_CREATION_FACTORY(FChaosClothAssetSimulationXPBDAnisoStretchConfigNode);
+			DATAFLOW_NODE_REGISTER_CREATION_FACTORY(FChaosClothAssetSimulationXPBDAreaSpringConfigNode);
+			DATAFLOW_NODE_REGISTER_CREATION_FACTORY(FChaosClothAssetSimulationXPBDBendingElementConfigNode);
+			DATAFLOW_NODE_REGISTER_CREATION_FACTORY(FChaosClothAssetSimulationXPBDBendingSpringConfigNode);
+			DATAFLOW_NODE_REGISTER_CREATION_FACTORY(FChaosClothAssetSimulationXPBDEdgeSpringConfigNode);
+			DATAFLOW_NODE_REGISTER_CREATION_FACTORY(FChaosClothAssetSkeletalMeshImportNode);
+			DATAFLOW_NODE_REGISTER_CREATION_FACTORY(FChaosClothAssetStaticMeshImportNode);
+			DATAFLOW_NODE_REGISTER_CREATION_FACTORY(FChaosClothAssetTerminalNode);
+			DATAFLOW_NODE_REGISTER_CREATION_FACTORY(FChaosClothAssetTransferSkinWeightsNode);
+			DATAFLOW_NODE_REGISTER_CREATION_FACTORY(FChaosClothAssetTransformUVsNode);
+			DATAFLOW_NODE_REGISTER_CREATION_FACTORY(FChaosClothAssetUSDImportNode);
+			DATAFLOW_NODE_REGISTER_CREATION_FACTORY_NODE_COLORS_BY_CATEGORY("Cloth", FColorScheme::NodeHeader, FColorScheme::NodeBody);
+		}
+	}  // End namespace Private
 
-	// Register type customizations
-	if (FPropertyEditorModule* const PropertyModule = FModuleManager::GetModulePtr<FPropertyEditorModule>("PropertyEditor"))
+	class FChaosClothAssetDataflowNodesModule : public IModuleInterface
 	{
-		PropertyModule->RegisterCustomPropertyTypeLayout(FChaosClothAssetWeightedValue::StaticStruct()->GetFName(), FOnGetPropertyTypeCustomizationInstance::CreateStatic(&FWeightedValueCustomization::MakeInstance));
-		PropertyModule->RegisterCustomPropertyTypeLayout(FChaosClothAssetWeightedValueNonAnimatable::StaticStruct()->GetFName(), FOnGetPropertyTypeCustomizationInstance::CreateStatic(&FWeightedValueCustomization::MakeInstance));
-		PropertyModule->RegisterCustomPropertyTypeLayout(FChaosClothAssetWeightedValueNonAnimatableNoLowHighRange::StaticStruct()->GetFName(), FOnGetPropertyTypeCustomizationInstance::CreateStatic(&FWeightedValueCustomization::MakeInstance));
-		PropertyModule->RegisterCustomPropertyTypeLayout(FChaosClothAssetNodeSelectionGroup::StaticStruct()->GetFName(), FOnGetPropertyTypeCustomizationInstance::CreateStatic(&FSelectionGroupCustomization::MakeInstance));
-		PropertyModule->RegisterCustomPropertyTypeLayout(FChaosClothAssetImportFilePath::StaticStruct()->GetFName(), FOnGetPropertyTypeCustomizationInstance::CreateStatic(&FImportFilePathCustomization::MakeInstance));
-	}
+	public:
+		virtual void StartupModule() override
+		{
+			using namespace UE::Chaos::ClothAsset;
 
-	// Register modular features
-	FChaosClothAssetDatasmithImportNode::RegisterModularFeature();
-}
+			Private::RegisterDataflowNodes();
 
-void FChaosClothAssetDataflowNodesModule::ShutdownModule()
-{
-	// Unregister type customizations
-	if (FPropertyEditorModule* const PropertyModule = FModuleManager::GetModulePtr<FPropertyEditorModule>("PropertyEditor"))
-	{
-		PropertyModule->UnregisterCustomPropertyTypeLayout(FChaosClothAssetWeightedValue::StaticStruct()->GetFName());
-		PropertyModule->UnregisterCustomPropertyTypeLayout(FChaosClothAssetWeightedValueNonAnimatable::StaticStruct()->GetFName());
-		PropertyModule->UnregisterCustomPropertyTypeLayout(FChaosClothAssetWeightedValueNonAnimatableNoLowHighRange::StaticStruct()->GetFName());
-		PropertyModule->UnregisterCustomPropertyTypeLayout(FChaosClothAssetNodeSelectionGroup::StaticStruct()->GetFName());
-	}
+			// Register type customizations
+			if (FPropertyEditorModule* const PropertyModule = FModuleManager::GetModulePtr<FPropertyEditorModule>("PropertyEditor"))
+			{
+				PropertyModule->RegisterCustomPropertyTypeLayout(FChaosClothAssetWeightedValue::StaticStruct()->GetFName(), FOnGetPropertyTypeCustomizationInstance::CreateStatic(&FWeightedValueCustomization::MakeInstance));
+				PropertyModule->RegisterCustomPropertyTypeLayout(FChaosClothAssetWeightedValueNonAnimatable::StaticStruct()->GetFName(), FOnGetPropertyTypeCustomizationInstance::CreateStatic(&FWeightedValueCustomization::MakeInstance));
+				PropertyModule->RegisterCustomPropertyTypeLayout(FChaosClothAssetWeightedValueNonAnimatableNoLowHighRange::StaticStruct()->GetFName(), FOnGetPropertyTypeCustomizationInstance::CreateStatic(&FWeightedValueCustomization::MakeInstance));
+				PropertyModule->RegisterCustomPropertyTypeLayout(FChaosClothAssetNodeSelectionGroup::StaticStruct()->GetFName(), FOnGetPropertyTypeCustomizationInstance::CreateStatic(&FSelectionGroupCustomization::MakeInstance));
+				PropertyModule->RegisterCustomPropertyTypeLayout(FChaosClothAssetImportFilePath::StaticStruct()->GetFName(), FOnGetPropertyTypeCustomizationInstance::CreateStatic(&FImportFilePathCustomization::MakeInstance));
+			}
+
+			// Register modular features
+			FChaosClothAssetDatasmithImportNode::RegisterModularFeature();
+		}
 
-	// Unregister modular features
-	FChaosClothAssetDatasmithImportNode::UnregisterModularFeature();
-}
+		virtual void ShutdownModule() override
+		{
+			// Unregister type customizations
+			if (FPropertyEditorModule* const PropertyModule = FModuleManager::GetModulePtr<FPropertyEditorModule>("PropertyEditor"))
+			{
+				PropertyModule->UnregisterCustomPropertyTypeLayout(FChaosClothAssetWeightedValue::StaticStruct()->GetFName());
+				PropertyModule->UnregisterCustomPropertyTypeLayout(FChaosClothAssetWeightedValueNonAnimatable::StaticStruct()->GetFName());
+				PropertyModule->UnregisterCustomPropertyTypeLayout(FChaosClothAssetWeightedValueNonAnimatableNoLowHighRange::StaticStruct()->GetFName());
+				PropertyModule->UnregisterCustomPropertyTypeLayout(FChaosClothAssetNodeSelectionGroup::StaticStruct()->GetFName());
+			}
 
-#undef LOCTEXT_NAMESPACE
+			// Unregister modular features
+			FChaosClothAssetDatasmithImportNode::UnregisterModularFeature();
+		}
+	};
+}  // End namespace UE::Chaos::ClothAsset
 
-IMPLEMENT_MODULE(FChaosClothAssetDataflowNodesModule, ChaosClothAssetDataflowNodes)
+IMPLEMENT_MODULE(UE::Chaos::ClothAsset::FChaosClothAssetDataflowNodesModule, ChaosClothAssetDataflowNodes)
diff --git a/Engine/Plugins/Experimental/ChaosClothAssetEditor/Source/ChaosClothAssetDataflowNodes/Private/ChaosClothAsset/ChaosClothAssetDataflowNodesModule.h b/Engine/Plugins/Experimental/ChaosClothAssetEditor/Source/ChaosClothAssetDataflowNodes/Private/ChaosClothAsset/ChaosClothAssetDataflowNodesModule.h
deleted file mode 100644
diff --git a/Engine/Plugins/Experimental/ChaosClothAssetEditor/Source/ChaosClothAssetEditor/Private/ChaosClothAsset/ChaosClothAssetEditorModule.cpp b/Engine/Plugins/Experimental/ChaosClothAssetEditor/Source/ChaosClothAssetEditor/Private/ChaosClothAsset/ChaosClothAssetEditorModule.cpp
--- a/Engine/Plugins/Experimental/ChaosClothAssetEditor/Source/ChaosClothAssetEditor/Private/ChaosClothAsset/ChaosClothAssetEditorModule.cpp
+++ b/Engine/Plugins/Experimental/ChaosClothAssetEditor/Source/ChaosClothAssetEditor/Private/ChaosClothAsset/ChaosClothAssetEditorModule.cpp
@@ -1,44 +1,27 @@
 // Copyright Epic Games, Inc. All Rights Reserved.
 
 #include "ChaosClothAsset/ChaosClothAssetEditorModule.h"
-#include "ChaosClothAsset/ClothEditorMode.h"
 #include "ChaosClothAsset/ClothEditorCommands.h"
+#include "ChaosClothAsset/ClothEditorMode.h"
 #include "ChaosClothAsset/ClothEditorStyle.h"
-#include "ContentBrowserMenuContexts.h"
+#include "Modules/ModuleManager.h"
 #include "EditorModeRegistry.h"
-#include "Selection.h"
-#include "ToolMenus.h"
-#include "Editor.h"
-#include "Styling/AppStyle.h"
-#include "AssetToolsModule.h"
-
-#define LOCTEXT_NAMESPACE "FChaosClothAssetEditorModule"
 
 namespace UE::Chaos::ClothAsset
 {
-void FChaosClothAssetEditorModule::StartupModule()
-{
-	FChaosClothAssetEditorStyle::Get(); // Causes the constructor to be called
+	void FChaosClothAssetEditorModule::StartupModule()
+	{
+		FChaosClothAssetEditorStyle::Get(); // Causes the constructor to be called
 
-	FChaosClothAssetEditorCommands::Register();
-
-	// Register asset actions
-	FAssetToolsModule& AssetToolsModule = FAssetToolsModule::GetModule();
-	IAssetTools& AssetTools = AssetToolsModule.Get();
-
-	// TODO: Register details view customizations
-}
-
-void FChaosClothAssetEditorModule::ShutdownModule()
-{
-	FChaosClothAssetEditorCommands::Unregister();
+		FChaosClothAssetEditorCommands::Register();
+	}
 
-	FEditorModeRegistry::Get().UnregisterMode(UChaosClothAssetEditorMode::EM_ChaosClothAssetEditorModeId);
+	void FChaosClothAssetEditorModule::ShutdownModule()
+	{
+		FChaosClothAssetEditorCommands::Unregister();
 
-	// TODO: Unregister details view customizations
-}
+		FEditorModeRegistry::Get().UnregisterMode(UChaosClothAssetEditorMode::EM_ChaosClothAssetEditorModeId);
+	}
 } // namespace UE::Chaos::ClothAsset
 
-#undef LOCTEXT_NAMESPACE
-	
-IMPLEMENT_MODULE(UE::Chaos::ClothAsset::FChaosClothAssetEditorModule, ChaosClothAssetEditor)
\ No newline at end of file
+IMPLEMENT_MODULE(UE::Chaos::ClothAsset::FChaosClothAssetEditorModule, ChaosClothAssetEditor)
diff --git a/Engine/Plugins/Experimental/ChaosClothAssetEditor/Source/ChaosClothAssetEditor/Private/ChaosClothAsset/ChaosClothAssetEditorModule.h b/Engine/Plugins/Experimental/ChaosClothAssetEditor/Source/ChaosClothAssetEditor/Private/ChaosClothAsset/ChaosClothAssetEditorModule.h
--- a/Engine/Plugins/Experimental/ChaosClothAssetEditor/Source/ChaosClothAssetEditor/Private/ChaosClothAsset/ChaosClothAssetEditorModule.h
+++ b/Engine/Plugins/Experimental/ChaosClothAssetEditor/Source/ChaosClothAssetEditor/Private/ChaosClothAsset/ChaosClothAssetEditorModule.h
@@ -2,27 +2,15 @@
 
 #pragma once
 
-#include "CoreMinimal.h"
 #include "BaseCharacterFXEditorModule.h"
 
-class FLayoutExtender;
-
 namespace UE::Chaos::ClothAsset
 {
-class FAssetTypeActions_ClothAsset;
-class FAssetTypeActions_ClothPreset;
-
-class CHAOSCLOTHASSETEDITOR_API FChaosClothAssetEditorModule : public FBaseCharacterFXEditorModule
-{
-public:
-
-	/** IModuleInterface implementation */
-	virtual void StartupModule() override;
-	virtual void ShutdownModule() override;
-
-private:
-	FAssetTypeActions_ClothAsset* AssetTypeActions_ClothAsset;
-	FAssetTypeActions_ClothPreset* AssetTypeActions_ClothPreset;
-
-};
+	class FChaosClothAssetEditorModule final : public FBaseCharacterFXEditorModule
+	{
+	public:
+		/** IModuleInterface implementation */
+		virtual void StartupModule() override;
+		virtual void ShutdownModule() override;
+	};
 } // namespace UE::Chaos::ClothAsset
diff --git a/Engine/Plugins/Experimental/ChaosClothAssetEditor/Source/ChaosClothAssetEditorTools/Private/ChaosClothAsset/ChaosClothAssetEditorToolsModule.cpp b/Engine/Plugins/Experimental/ChaosClothAssetEditor/Source/ChaosClothAssetEditorTools/Private/ChaosClothAsset/ChaosClothAssetEditorToolsModule.cpp
--- a/Engine/Plugins/Experimental/ChaosClothAssetEditor/Source/ChaosClothAssetEditorTools/Private/ChaosClothAsset/ChaosClothAssetEditorToolsModule.cpp
+++ b/Engine/Plugins/Experimental/ChaosClothAssetEditor/Source/ChaosClothAssetEditorTools/Private/ChaosClothAsset/ChaosClothAssetEditorToolsModule.cpp
@@ -1,20 +1,20 @@
 // Copyright Epic Games, Inc. All Rights Reserved.
 
-#include "ChaosClothAsset/ChaosClothAssetEditorToolsModule.h"
+#include "Modules/ModuleManager.h"
 
-#define LOCTEXT_NAMESPACE "FChaosClothAssetEditorToolsModule"
-
-void FChaosClothAssetEditorToolsModule::StartupModule()
+namespace UE::Chaos::ClothAsset
 {
-	// This code will execute after your module is loaded into memory; the exact timing is specified in the .uplugin file per-module
-}
+	class FChaosClothAssetEditorToolsModule : public IModuleInterface
+	{
+	public:
+		virtual void StartupModule() override
+		{
+		}
 
-void FChaosClothAssetEditorToolsModule::ShutdownModule()
-{
-	// This function may be called during shutdown to clean up your module.  For modules that support dynamic reloading,
-	// we call this function before unloading the module.
+		virtual void ShutdownModule() override
+		{
+		}
+	};
 }
 
-#undef LOCTEXT_NAMESPACE
-
-IMPLEMENT_MODULE(FChaosClothAssetEditorToolsModule, ChaosClothAssetEditorTools)
+IMPLEMENT_MODULE(UE::Chaos::ClothAsset::FChaosClothAssetEditorToolsModule, ChaosClothAssetEditorTools)
diff --git a/Engine/Plugins/Experimental/ChaosClothAssetEditor/Source/ChaosClothAssetEditorTools/Private/ChaosClothAsset/ChaosClothAssetEditorToolsModule.h b/Engine/Plugins/Experimental/ChaosClothAssetEditor/Source/ChaosClothAssetEditorTools/Private/ChaosClothAsset/ChaosClothAssetEditorToolsModule.h
deleted file mode 100644
diff --git a/Engine/Plugins/Experimental/ChaosClothAssetEditor/Source/ChaosClothAssetTools/Private/ChaosClothAsset/ChaosClothAssetToolsModule.cpp b/Engine/Plugins/Experimental/ChaosClothAssetEditor/Source/ChaosClothAssetTools/Private/ChaosClothAsset/ChaosClothAssetToolsModule.cpp
--- a/Engine/Plugins/Experimental/ChaosClothAssetEditor/Source/ChaosClothAssetTools/Private/ChaosClothAsset/ChaosClothAssetToolsModule.cpp
+++ b/Engine/Plugins/Experimental/ChaosClothAssetEditor/Source/ChaosClothAssetTools/Private/ChaosClothAsset/ChaosClothAssetToolsModule.cpp
@@ -1,18 +1,15 @@
 // Copyright Epic Games, Inc. All Rights Reserved.
 
-#include "CoreMinimal.h"
+#include "ChaosClothAsset/ClothAssetBuilderEditor.h"
 #include "Features/IModularFeatures.h"
-#include "Modules/ModuleInterface.h"
 #include "Modules/ModuleManager.h"
-#include "AssetToolsModule.h"
 #include "ClothAssetEditorStyle.h"
 #include "ClothComponentEditorStyle.h"
 #include "ClothingAssetToClothAssetExporter.h"
-#include "ChaosClothAsset/ClothAssetBuilderEditor.h"
 
 namespace UE::Chaos::ClothAsset
 {
-	class FClothAssetEditorModule : public IModuleInterface, public IClothingAssetExporterClassProvider, public IClothAssetBuilderClassProvider
+	class FChaosClothAssetToolsModule : public IModuleInterface, public IClothingAssetExporterClassProvider, public IClothAssetBuilderClassProvider
 	{
 	public:
 		// IModuleInterface implementation
@@ -51,4 +48,4 @@ namespace UE::Chaos::ClothAsset
 	};
 }
 
-IMPLEMENT_MODULE(UE::Chaos::ClothAsset::FClothAssetEditorModule, ChaosClothAssetEditor);
+IMPLEMENT_MODULE(UE::Chaos::ClothAsset::FChaosClothAssetToolsModule, ChaosClothAssetToolsModule);
