Chaos Cloth Asset - Restore the USD stage scale on the render mesh import since it has now been fixed in the exporter.

diff --git a/Engine/Plugins/Experimental/ChaosClothAssetEditor/Source/ChaosClothAssetDataflowNodes/Private/ChaosClothAsset/USDImportNode.cpp b/Engine/Plugins/Experimental/ChaosClothAssetEditor/Source/ChaosClothAssetDataflowNodes/Private/ChaosClothAsset/USDImportNode.cpp
--- a/Engine/Plugins/Experimental/ChaosClothAssetEditor/Source/ChaosClothAssetDataflowNodes/Private/ChaosClothAsset/USDImportNode.cpp
+++ b/Engine/Plugins/Experimental/ChaosClothAssetEditor/Source/ChaosClothAssetDataflowNodes/Private/ChaosClothAsset/USDImportNode.cpp
@@ -469,9 +469,7 @@ bool FChaosClothAssetUSDImportNode::ImportFromFile(const FString& UsdFilePath, c
 		ImportOptions->MaterialPurpose = NAME_None;  // *UnrealIdentifiers::MaterialPreviewPurpose ???
 		ImportOptions->RootMotionHandling = EUsdRootMotionHandling::NoAdditionalRootMotion;
 		ImportOptions->SubdivisionLevel = 0;
-		ImportOptions->bOverrideStageOptions = true;
-		ImportOptions->StageOptions.MetersPerUnit = 0.01f;  // REVIEW: The render mesh is currently exported in cm scale despite the stage MetersPerUnit and simulation data always set to mm
-		ImportOptions->StageOptions.UpAxis = StageInfo.UpAxis;
+		ImportOptions->bOverrideStageOptions = false;
 		ImportOptions->bImportAtSpecificTimeCode = false;
 		ImportOptions->ImportTimeCode = 0.f;
 		// Groom
