CIS - Fixed the ChaosClothAssetTools module name that had the suffix "Module" mistakenly added to the name.

diff --git a/Engine/Plugins/Experimental/ChaosClothAssetEditor/Source/ChaosClothAssetTools/Private/ChaosClothAsset/ChaosClothAssetToolsModule.cpp b/Engine/Plugins/Experimental/ChaosClothAssetEditor/Source/ChaosClothAssetTools/Private/ChaosClothAsset/ChaosClothAssetToolsModule.cpp
--- a/Engine/Plugins/Experimental/ChaosClothAssetEditor/Source/ChaosClothAssetTools/Private/ChaosClothAsset/ChaosClothAssetToolsModule.cpp
+++ b/Engine/Plugins/Experimental/ChaosClothAssetEditor/Source/ChaosClothAssetTools/Private/ChaosClothAsset/ChaosClothAssetToolsModule.cpp
@@ -48,4 +48,4 @@ namespace UE::Chaos::ClothAsset
 	};
 }
 
-IMPLEMENT_MODULE(UE::Chaos::ClothAsset::FChaosClothAssetToolsModule, ChaosClothAssetToolsModule);
+IMPLEMENT_MODULE(UE::Chaos::ClothAsset::FChaosClothAssetToolsModule, ChaosClothAssetTools);
