Chaos Cloth Asset - Set default MaxDistance weight map to 1 when MaxDistance properties have been set.

diff --git a/Engine/Plugins/Experimental/ChaosClothAssetEditor/Source/ChaosClothAssetTools/Private/ChaosClothAsset/ClothAssetBuilderEditor.cpp b/Engine/Plugins/Experimental/ChaosClothAssetEditor/Source/ChaosClothAssetTools/Private/ChaosClothAsset/ClothAssetBuilderEditor.cpp
--- a/Engine/Plugins/Experimental/ChaosClothAssetEditor/Source/ChaosClothAssetTools/Private/ChaosClothAsset/ClothAssetBuilderEditor.cpp
+++ b/Engine/Plugins/Experimental/ChaosClothAssetEditor/Source/ChaosClothAssetTools/Private/ChaosClothAsset/ClothAssetBuilderEditor.cpp
@@ -64,13 +64,14 @@ void UClothAssetBuilderEditor::BuildLod(FSkeletalMeshLODModel& LODModel, const U
 	int32 MaxDistancePropertyKeyIndex;
 	FString MaxDistanceString = TEXT("MaxDistance");
 	MaxDistanceString = Properties.GetStringValue(MaxDistanceString, MaxDistanceString, &MaxDistancePropertyKeyIndex);
-	const float MaxDistanceBase = (MaxDistancePropertyKeyIndex != INDEX_NONE) ? Properties.GetLowValue<float>(MaxDistancePropertyKeyIndex) : 0.f;
-	const float MaxDistanceRange = (MaxDistancePropertyKeyIndex != INDEX_NONE) ? Properties.GetHighValue<float>(MaxDistancePropertyKeyIndex) - MaxDistanceBase : 1.f;
+	const bool bHasMaxDistanceProperty = (MaxDistancePropertyKeyIndex != INDEX_NONE);
+	const float MaxDistanceOffset = bHasMaxDistanceProperty ? Properties.GetLowValue<float>(MaxDistancePropertyKeyIndex) : TNumericLimits<float>::Max();  // Uses infinite distance when no MaxDistance properties are set
+	const float MaxDistanceScale = bHasMaxDistanceProperty ? Properties.GetHighValue<float>(MaxDistancePropertyKeyIndex) - MaxDistanceOffset : 0.f;
 	const TConstArrayView<float> MaxDistanceWeightMap = ClothFacade.GetWeightMap(FName(MaxDistanceString));
 
 	const FPointWeightMap MaxDistances = (MaxDistanceWeightMap.Num() == NumLodSimVertices) ?
-		FPointWeightMap(MaxDistanceWeightMap) : 
-		FPointWeightMap(NumLodSimVertices, TNumericLimits<float>::Max());
+		FPointWeightMap(MaxDistanceWeightMap, MaxDistanceOffset, MaxDistanceScale) :
+		FPointWeightMap(NumLodSimVertices, MaxDistanceOffset);
 
 	const int32 NumRenderVertices = ClothFacade.GetNumRenderVertices();
 	LODModel.MeshToImportVertexMap.Reserve(NumRenderVertices);
diff --git a/Engine/Source/Runtime/ClothingSystemRuntimeCommon/Public/PointWeightMap.h b/Engine/Source/Runtime/ClothingSystemRuntimeCommon/Public/PointWeightMap.h
--- a/Engine/Source/Runtime/ClothingSystemRuntimeCommon/Public/PointWeightMap.h
+++ b/Engine/Source/Runtime/ClothingSystemRuntimeCommon/Public/PointWeightMap.h
@@ -55,6 +55,21 @@ struct FPointWeightMap
 #endif
 	{}
 
+	FPointWeightMap(const TConstArrayView<float>& InValues, float Offset, float Scale)
+#if WITH_EDITORONLY_DATA
+		: Name(NAME_None)
+		, CurrentTarget((uint8)EWeightMapTargetCommon::None)
+		, bEnabled(false)
+#endif
+	{
+		const int32 NumPoints = InValues.Num();
+		Values.SetNumUninitialized(NumPoints);
+		for (int32 Index = 0; Index < NumPoints; ++Index)
+		{
+			Values[Index] = Offset + Scale * InValues[Index];
+		}
+	}
+
 	~FPointWeightMap()
 	{}
 
