Chaos Cloth Asset - Fixed crash when evaluating Name output of the AddWeightMap and SelectionToWeightMapNode nodes without evaluating their collection.

diff --git a/Engine/Plugins/Experimental/ChaosClothAssetEditor/Source/ChaosClothAssetDataflowNodes/Private/ChaosClothAsset/AddWeightMapNode.cpp b/Engine/Plugins/Experimental/ChaosClothAssetEditor/Source/ChaosClothAssetDataflowNodes/Private/ChaosClothAsset/AddWeightMapNode.cpp
--- a/Engine/Plugins/Experimental/ChaosClothAssetEditor/Source/ChaosClothAssetDataflowNodes/Private/ChaosClothAsset/AddWeightMapNode.cpp
+++ b/Engine/Plugins/Experimental/ChaosClothAssetEditor/Source/ChaosClothAssetDataflowNodes/Private/ChaosClothAsset/AddWeightMapNode.cpp
@@ -54,6 +54,12 @@ void FChaosClothAssetAddWeightMapNode::Evaluate(Dataflow::FContext& Context, con
 			}
 		}
 		SetValue(Context, MoveTemp(*ClothCollection), &Collection);
+	}
+	else if (Out->IsA<FString>(&Name))
+	{
+		// Make the name a valid attribute name, and replace the value in the UI  // TODO: Remove all mutables from evaluation
+		FWeightMapTools::MakeWeightMapName(Name);
+
 		SetValue(Context, Name, &Name);
 	}
 }
diff --git a/Engine/Plugins/Experimental/ChaosClothAssetEditor/Source/ChaosClothAssetDataflowNodes/Private/ChaosClothAsset/SelectionToWeightMapNode.cpp b/Engine/Plugins/Experimental/ChaosClothAssetEditor/Source/ChaosClothAssetDataflowNodes/Private/ChaosClothAsset/SelectionToWeightMapNode.cpp
--- a/Engine/Plugins/Experimental/ChaosClothAssetEditor/Source/ChaosClothAssetDataflowNodes/Private/ChaosClothAsset/SelectionToWeightMapNode.cpp
+++ b/Engine/Plugins/Experimental/ChaosClothAssetEditor/Source/ChaosClothAssetDataflowNodes/Private/ChaosClothAsset/SelectionToWeightMapNode.cpp
@@ -42,7 +42,7 @@ void FChaosClothAssetSelectionToWeightMapNode::Evaluate(Dataflow::FContext& Cont
 					SelectionGroup == ClothCollectionGroup::SimVertices3D || 
 					SelectionGroup == ClothCollectionGroup::SimFaces)
 				{
-					const FName InMapName(WeightMapName);
+					const FName InMapName(WeightMapName.IsEmpty() ? SelectionName : WeightMapName);
 					ClothFacade.AddWeightMap(InMapName);
 					TArrayView<float> OutClothWeights = ClothFacade.GetWeightMap(InMapName);
 
@@ -112,7 +112,10 @@ void FChaosClothAssetSelectionToWeightMapNode::Evaluate(Dataflow::FContext& Cont
 		}
 
 		SetValue(Context, MoveTemp(*ClothCollection), &Collection);
-		SetValue(Context, WeightMapName, &WeightMapName);
+	}
+	else if (Out->IsA<FString>(&WeightMapName))
+	{
+		SetValue(Context, WeightMapName.IsEmpty() ? SelectionName : WeightMapName, &WeightMapName);
 	}
 }
 
diff --git a/Engine/Plugins/Experimental/ChaosClothAssetEditor/Source/ChaosClothAssetDataflowNodes/Public/ChaosClothAsset/SelectionToWeightMapNode.h b/Engine/Plugins/Experimental/ChaosClothAssetEditor/Source/ChaosClothAssetDataflowNodes/Public/ChaosClothAsset/SelectionToWeightMapNode.h
--- a/Engine/Plugins/Experimental/ChaosClothAssetEditor/Source/ChaosClothAssetDataflowNodes/Public/ChaosClothAsset/SelectionToWeightMapNode.h
+++ b/Engine/Plugins/Experimental/ChaosClothAssetEditor/Source/ChaosClothAssetDataflowNodes/Public/ChaosClothAsset/SelectionToWeightMapNode.h
@@ -22,7 +22,10 @@ public:
 	UPROPERTY(Meta = (DataflowInput))
 	FString SelectionName;
 
-	/** The name of the weight map attribute that will be added to the collection */
+	/**
+	 * The name of the weight map attribute that will be added to the collection.
+	 * If left empty the same name as the selection name will be used instead.
+	 */
 	UPROPERTY(EditAnywhere, Category = "Selection To Weight Map", Meta = (DataflowOutput))
 	FString WeightMapName;
 
