Add a Remesh node for Cloth.

For the sim mesh, the node operates on the 2D representation by first remeshing the seam and boundary edges, then remeshing the entire mesh while constraining those edges. Finally it (un)projects the 2D mesh to 3D using the input 2D/3D mesh representations.

For the render mesh, it remeshes the boundaries first, then the interiors in 3D since there are no seams. (Render Mesh LOD generation is likely to be replaced with something more user-controllable in the future, this is kind of a placeholder.)

Still TODO: additional constraints for other sim needs

diff --git a/Engine/Plugins/Experimental/ChaosClothAsset/Source/ChaosClothAsset/Private/ChaosClothAsset/ClothGeometryTools.cpp b/Engine/Plugins/Experimental/ChaosClothAsset/Source/ChaosClothAsset/Private/ChaosClothAsset/ClothGeometryTools.cpp
--- a/Engine/Plugins/Experimental/ChaosClothAsset/Source/ChaosClothAsset/Private/ChaosClothAsset/ClothGeometryTools.cpp
+++ b/Engine/Plugins/Experimental/ChaosClothAsset/Source/ChaosClothAsset/Private/ChaosClothAsset/ClothGeometryTools.cpp
@@ -718,6 +718,137 @@ namespace UE::Chaos::ClothAsset
 		}
 	}
 
+
+	void FClothGeometryTools::BuildSimMeshFromDynamicMeshes(
+		const TSharedRef<FManagedArrayCollection>& ClothCollection,
+		const UE::Geometry::FDynamicMesh3& Mesh2D,
+		const UE::Geometry::FDynamicMesh3& Mesh3D,
+		int32 PatternIndexLayerID,
+		bool bTransferWeightMaps,
+		bool bTransferSimSkinningData,
+		bool bAppend,
+		TMap<int, int32>& OutDynamicMeshToClothVertexMap)
+	{
+		using namespace UE::Geometry;
+
+		if (!bAppend)
+		{
+			DeleteSimMesh(ClothCollection);
+		}
+		FCollectionClothFacade Cloth(ClothCollection);
+		checkf(Cloth.IsValid(), TEXT("Invalid ClothCollection passed into BuildSimMeshFromDynamicMeshes"));
+
+		check(Mesh2D.HasAttributes());
+		const FDynamicMeshPolygroupAttribute* const PatternLayer = Mesh2D.Attributes()->GetPolygroupLayer((int)PatternIndexLayerID);
+		check(PatternLayer);
+
+		TArray<TArray<int>> PatternIndices;
+		for (int FaceID = 0; FaceID < Mesh2D.MaxTriangleID(); ++FaceID)
+		{
+			const int32 PatternID = PatternLayer->GetValue(FaceID);
+			if (PatternID >= PatternIndices.Num())
+			{
+				PatternIndices.SetNum(PatternID + 1);
+			}
+
+			const FIndex3i Tri = Mesh2D.GetTriangle(FaceID);
+			PatternIndices[PatternID].Add(Tri[0]);
+			PatternIndices[PatternID].Add(Tri[1]);
+			PatternIndices[PatternID].Add(Tri[2]);
+		}
+
+		TMap<int, FIntVector2> MeshVertexToPatternAndVertex;
+
+		for (int32 PatternID = 0; PatternID < PatternIndices.Num(); ++PatternID)
+		{
+			FCollectionClothSimPatternFacade Pattern = Cloth.AddGetSimPattern();
+			const int32 PatternVertexOffset = Pattern.GetSimVertices2DOffset();
+
+			TArray<FVector2f> Positions2D;
+			TArray<FVector3f> Positions3D;
+
+			const TArray<int>& InPatternIndexBuffer = PatternIndices[PatternID];
+			
+			TArray<int> LocalPatternIndexBuffer;
+
+			for (const int VertexIndex : InPatternIndexBuffer)
+			{
+				int32 PatternVertexID;
+
+				if (!MeshVertexToPatternAndVertex.Contains(VertexIndex))
+				{
+					const FVector3d InPosition2D = Mesh2D.GetVertex(VertexIndex);
+					PatternVertexID = Positions2D.Add(FVector2f(InPosition2D[0], InPosition2D[1]));
+
+					const FVector3d InPosition3D = Mesh3D.GetVertex(VertexIndex);
+					Positions3D.Add(FVector3f(InPosition3D));
+
+					MeshVertexToPatternAndVertex.Add({ VertexIndex, FIntVector2{PatternID, PatternVertexID}});
+				}
+				else
+				{
+					check(MeshVertexToPatternAndVertex[VertexIndex][0] == PatternID);
+					PatternVertexID = MeshVertexToPatternAndVertex[VertexIndex][1];
+				}
+
+				LocalPatternIndexBuffer.Add(PatternVertexID);
+			}
+
+			Pattern.Initialize(Positions2D, Positions3D, LocalPatternIndexBuffer);
+		}
+
+		for (int InGlobalVertexIndex = 0; InGlobalVertexIndex < Mesh2D.MaxVertexID(); ++InGlobalVertexIndex)
+		{
+			const int32 PatternID = MeshVertexToPatternAndVertex[InGlobalVertexIndex][0];
+			const int32 VertexID = MeshVertexToPatternAndVertex[InGlobalVertexIndex][1];
+			const int32 ClothGlobalIndex = Cloth.GetSimPattern(PatternID).GetSimVertices2DOffset() + VertexID;
+
+			OutDynamicMeshToClothVertexMap.Add({ InGlobalVertexIndex, ClothGlobalIndex });
+		}
+
+		// Copy skinning data
+		if (bTransferSimSkinningData)
+		{
+			const UE::Geometry::FDynamicMeshVertexSkinWeightsAttribute* const SkinWeights = Mesh2D.Attributes() ? Mesh2D.Attributes()->GetSkinWeightsAttribute(FName("Default")) : nullptr;
+			if (SkinWeights)
+			{
+				TArrayView<TArray<int32>> BoneIndices = Cloth.GetSimBoneIndices();
+				TArrayView<TArray<float>> BoneWeights = Cloth.GetSimBoneWeights();
+				for (int32 MeshVertexIndex : Mesh2D.VertexIndicesItr())
+				{
+					const int32 ClothVertexIndex = OutDynamicMeshToClothVertexMap[MeshVertexIndex];
+					SkinWeights->GetValue(MeshVertexIndex, BoneIndices[ClothVertexIndex], BoneWeights[ClothVertexIndex]);
+				}
+			}
+		}
+
+		// Copy scalar weight maps
+		if (bTransferWeightMaps && Mesh2D.Attributes())
+		{
+			const UE::Geometry::FDynamicMeshAttributeSet* const AttributeSet = Mesh2D.Attributes();
+			for (int32 WeightMapLayerIndex = 0; WeightMapLayerIndex < AttributeSet->NumWeightLayers(); ++WeightMapLayerIndex)
+			{
+				if (const UE::Geometry::FDynamicMeshWeightAttribute* const WeightMapAttribute = AttributeSet->GetWeightLayer(WeightMapLayerIndex))
+				{
+					const FName WeightMapName = WeightMapAttribute->GetName();
+					Cloth.AddWeightMap(WeightMapName);	// Does nothing if weight map already exists
+					TArrayView<float> OutWeightMap = Cloth.GetWeightMap(WeightMapName);
+
+					for (const int32 MeshVertexIndex : Mesh2D.VertexIndicesItr())
+					{
+						float VertexWeight;
+						WeightMapAttribute->GetValue(MeshVertexIndex, &VertexWeight);
+
+						const int32 ClothVertexIndex = OutDynamicMeshToClothVertexMap[MeshVertexIndex];
+						OutWeightMap[ClothVertexIndex] = VertexWeight;
+					}
+				}
+			}
+		}
+
+	}
+
+
 	void FClothGeometryTools::BuildSimMeshFromDynamicMesh(
 		const TSharedRef<FManagedArrayCollection>& ClothCollection,
 		const UE::Geometry::FDynamicMesh3& DynamicMesh, int32 UVChannelIndex, const FVector2f& UVScale, bool bAppend)
@@ -751,6 +882,8 @@ namespace UE::Chaos::ClothAsset
 				FCollectionClothSimPatternFacade Pattern = Cloth.AddGetSimPattern();
 				const int32 VertexOffset = Cloth.GetNumSimVertices3D();
 				Pattern.Initialize(Island.Positions2D, Island.Positions3D, Island.Indices);
+
+				// Copy skinning data
 				if (SkinWeights)
 				{
 					TArrayView<TArray<int32>> BoneIndices = Cloth.GetSimBoneIndices();
@@ -761,6 +894,28 @@ namespace UE::Chaos::ClothAsset
 						SkinWeights->GetValue(Island.PositionToSourceIndex[VertexIndex], BoneIndices[VertexIndex + VertexOffset], BoneWeights[VertexIndex + VertexOffset]);
 					}
 				}
+
+				// Copy scalar weight maps
+				if (AttributeSet)
+				{
+					for (int32 WeightMapLayerIndex = 0; WeightMapLayerIndex < AttributeSet->NumWeightLayers(); ++WeightMapLayerIndex)
+					{
+						if (const UE::Geometry::FDynamicMeshWeightAttribute* const WeightMapAttribute = AttributeSet->GetWeightLayer(WeightMapLayerIndex))
+						{
+							const FName WeightMapName = WeightMapAttribute->GetName();
+							Cloth.AddWeightMap(WeightMapName);	// Does nothing if weight map already exists
+							TArrayView<float> OutWeightMap = Cloth.GetWeightMap(WeightMapName);
+
+							for (int32 VertexIndex = 0; VertexIndex < Island.Positions3D.Num(); ++VertexIndex)
+							{
+								float VertexWeight;
+								WeightMapAttribute->GetValue(Island.PositionToSourceIndex[VertexIndex], &VertexWeight);
+								OutWeightMap[VertexIndex + VertexOffset] = VertexWeight;
+							}
+						}
+					}
+				}
+
 			}
 		}
 
diff --git a/Engine/Plugins/Experimental/ChaosClothAsset/Source/ChaosClothAsset/Public/ChaosClothAsset/ClothGeometryTools.h b/Engine/Plugins/Experimental/ChaosClothAsset/Source/ChaosClothAsset/Public/ChaosClothAsset/ClothGeometryTools.h
--- a/Engine/Plugins/Experimental/ChaosClothAsset/Source/ChaosClothAsset/Public/ChaosClothAsset/ClothGeometryTools.h
+++ b/Engine/Plugins/Experimental/ChaosClothAsset/Source/ChaosClothAsset/Public/ChaosClothAsset/ClothGeometryTools.h
@@ -53,6 +53,29 @@ namespace UE::Chaos::ClothAsset
 			bool bBindSimMesh,
 			bool bBindRenderMesh);
 
+
+		/**
+		* Build (or add to) a ClothCollection Sim Mesh from the given 2D and 3D mesh data. Uses a Polygroup Attribute Layer to specify Pattern topology.
+		* 
+		* @param ClothCollection					The cloth collection whose sim mesh (2D and 3D) will be modified
+		* @param Mesh2D								Input 2D sim mesh data
+		* @param Mesh3D								Input 3D sim mesh data
+		* @param PatternIndexLayerId				Specifies which PolyGroup layer on Mesh2D contains pattern index per triangle information
+		* @param bTransferWeightMaps				Copy any weight map layers from Mesh2D into the ClothCollection sim mesh
+		* @param bTransferSimSkinningData			Copy any skinning weight data from Mesh2D into the ClothCollection sim mesh
+		* @param bAppend							Whether to add the new mesh data to the existing sim mesh, or crete new sim mesh in the collection
+		* @param OutDynamicMeshToClothVertexMap		(Output) Computed map of vertex indices in the input Meshes to vertex indices in the output ClothCollection
+		*/
+		static void BuildSimMeshFromDynamicMeshes(
+			const TSharedRef<FManagedArrayCollection>& ClothCollection,
+			const UE::Geometry::FDynamicMesh3& Mesh2D,
+			const UE::Geometry::FDynamicMesh3& Mesh3D,
+			int32 PatternIndexLayerId,
+			bool bTransferWeightMaps,
+			bool bTransferSimSkinningData,
+			bool bAppend,
+			TMap<int, int32>& OutDynamicMeshToClothVertexMap);
+
 		/**
 		* Unwrap and build SimMesh data from a DynamicMesh
 		*/
diff --git a/Engine/Plugins/Experimental/ChaosClothAsset/Source/ChaosClothAsset/Public/ChaosClothAsset/CollectionClothFacade.h b/Engine/Plugins/Experimental/ChaosClothAsset/Source/ChaosClothAsset/Public/ChaosClothAsset/CollectionClothFacade.h
--- a/Engine/Plugins/Experimental/ChaosClothAsset/Source/ChaosClothAsset/Public/ChaosClothAsset/CollectionClothFacade.h
+++ b/Engine/Plugins/Experimental/ChaosClothAsset/Source/ChaosClothAsset/Public/ChaosClothAsset/CollectionClothFacade.h
@@ -100,7 +100,7 @@ namespace UE::Chaos::ClothAsset
 		//~ Seam Group
 		/** Return the number of seams in this collection. */
 		int32 GetNumSeams() const;
-		/** Return a seam facade for the specified pattern index. */
+		/** Return a seam facade for the specified seam index. */
 		FCollectionClothSeamConstFacade GetSeam(int32 SeamIndex) const;
 
 		//~ Render Vertices Group
diff --git a/Engine/Plugins/Experimental/ChaosClothAssetEditor/Source/ChaosClothAssetDataflowNodes/ChaosClothAssetDataflowNodes.Build.cs b/Engine/Plugins/Experimental/ChaosClothAssetEditor/Source/ChaosClothAssetDataflowNodes/ChaosClothAssetDataflowNodes.Build.cs
--- a/Engine/Plugins/Experimental/ChaosClothAssetEditor/Source/ChaosClothAssetDataflowNodes/ChaosClothAssetDataflowNodes.Build.cs
+++ b/Engine/Plugins/Experimental/ChaosClothAssetEditor/Source/ChaosClothAssetDataflowNodes/ChaosClothAssetDataflowNodes.Build.cs
@@ -39,6 +39,7 @@ public class ChaosClothAssetDataflowNodes : ModuleRules
 				"MeshConversion",
 				"MeshDescription",
 				"MeshUtilitiesCommon",
+				"ModelingOperators",
 				"RenderCore",
 				"SkeletalMeshDescription",
 				"Slate",
diff --git a/Engine/Plugins/Experimental/ChaosClothAssetEditor/Source/ChaosClothAssetDataflowNodes/Private/ChaosClothAsset/ChaosClothAssetDataflowNodesModule.cpp b/Engine/Plugins/Experimental/ChaosClothAssetEditor/Source/ChaosClothAssetDataflowNodes/Private/ChaosClothAsset/ChaosClothAssetDataflowNodesModule.cpp
--- a/Engine/Plugins/Experimental/ChaosClothAssetEditor/Source/ChaosClothAssetDataflowNodes/Private/ChaosClothAsset/ChaosClothAssetDataflowNodesModule.cpp
+++ b/Engine/Plugins/Experimental/ChaosClothAssetEditor/Source/ChaosClothAssetDataflowNodes/Private/ChaosClothAsset/ChaosClothAssetDataflowNodesModule.cpp
@@ -8,6 +8,7 @@
 #include "ChaosClothAsset/DatasmithImportNode.h"
 #include "ChaosClothAsset/DeleteElementNode.h"
 #include "ChaosClothAsset/ImportFilePathCustomization.h"
+#include "ChaosClothAsset/RemeshNode.h"
 #include "ChaosClothAsset/ImportNode.h"
 #include "ChaosClothAsset/MergeClothCollectionsNode.h"
 #include "ChaosClothAsset/ReverseNormalsNode.h"
@@ -64,6 +65,7 @@ namespace UE::Chaos::ClothAsset
 			DATAFLOW_NODE_REGISTER_CREATION_FACTORY(FChaosClothAssetDeleteElementNode);
 			DATAFLOW_NODE_REGISTER_CREATION_FACTORY(FChaosClothAssetImportNode);
 			DATAFLOW_NODE_REGISTER_CREATION_FACTORY(FChaosClothAssetMergeClothCollectionsNode);
+			DATAFLOW_NODE_REGISTER_CREATION_FACTORY(FChaosClothAssetRemeshNode);
 			DATAFLOW_NODE_REGISTER_CREATION_FACTORY(FChaosClothAssetReverseNormalsNode);
 			DATAFLOW_NODE_REGISTER_CREATION_FACTORY(FChaosClothAssetSelectionNode);
 			DATAFLOW_NODE_REGISTER_CREATION_FACTORY(FChaosClothAssetSelectionToWeightMapNode);
diff --git a/Engine/Plugins/Experimental/ChaosClothAssetEditor/Source/ChaosClothAssetDataflowNodes/Private/ChaosClothAsset/RemeshNode.cpp b/Engine/Plugins/Experimental/ChaosClothAssetEditor/Source/ChaosClothAssetDataflowNodes/Private/ChaosClothAsset/RemeshNode.cpp
new file mode 100644
--- /dev/null
+++ b/Engine/Plugins/Experimental/ChaosClothAssetEditor/Source/ChaosClothAssetDataflowNodes/Private/ChaosClothAsset/RemeshNode.cpp
@@ -0,0 +1,1172 @@
+// Copyright Epic Games, Inc. All Rights Reserved.
+
+#include "ChaosClothAsset/RemeshNode.h"
+
+#include "ChaosClothAsset/ClothPatternToDynamicMesh.h"
+#include "ChaosClothAsset/CollectionClothFacade.h"
+#include "ChaosClothAsset/ClothGeometryTools.h"
+#include "Dataflow/DataflowInputOutput.h"
+#include "DynamicMesh/MeshTangents.h"
+#include "DynamicMesh/MeshNormals.h"
+#include "DynamicMesh/DynamicVertexSkinWeightsAttribute.h"
+#include "CleaningOps/RemeshMeshOp.h"
+#include "MeshUVChannelInfo.h"
+#include "MeshBoundaryLoops.h"
+
+#include UE_INLINE_GENERATED_CPP_BY_NAME(RemeshNode)
+
+#define LOCTEXT_NAMESPACE "ChaosClothAssetRemeshNode"
+
+namespace UE::Chaos::ClothAsset::Private
+{
+
+	//
+	// Functions to support seam remeshing
+	//
+
+	struct FSeamCollapseParameters
+	{
+		bool bCanCollapse = false;
+		int32 KeepStitchIndex = -1;
+		int32 DeleteStitchIndex = -1;
+		double NewPositionEdgeParameter = 0.5;
+	};
+
+	/**
+	 * A "Seam Edge" is a pair of consecutive stitches in a seam. Corresponds to two "mesh" edges.
+	 */
+	struct FSeamEdge
+	{
+		FSeamEdge(FIntVector2 A, FIntVector2 B)
+		{
+			Stitches[0] = A;
+			Stitches[1] = B;
+		}
+
+		TStaticArray<FIntVector2, 2> Stitches;
+
+		bool operator==(const FSeamEdge& Other) const
+		{
+			return (Other.Stitches == Stitches);
+		}
+	};
+
+	// Try to identify if the corner A -- B -- C is "sharp".
+	// 
+	// The function computes the dot product of normalized vectors AB and BC, and returns true if the result is less than the given threshold
+	// Also returns true if A \approx B or B \approx C
+	bool IsSharpCorner(const FVector3d& A, const FVector3d& B, const FVector3d& C, double CosAngleThreshold)
+	{
+		const FVector3d AB = B - A;
+		const FVector3d BC = C - B;
+		const double NormAB = AB.Length();
+		const double NormBC = BC.Length();
+
+		if (FMath::Abs(NormAB) < UE_SMALL_NUMBER || FMath::Abs(NormBC) < UE_SMALL_NUMBER)
+		{
+			return true;
+		}
+
+		const double CosAngle = AB.Dot(BC) / (NormAB * NormBC);
+		return (CosAngle < CosAngleThreshold);
+	}
+
+	FSeamCollapseParameters GetSeamEdgeCollapseParameters(int32 SeamID, int32 StitchID, const UE::Geometry::FDynamicMesh3& Mesh, const TArray<TArray<FIntVector2>>& Seams, double CosAngleCornerTheshold = 0.5)
+	{
+		using namespace UE::Geometry;
+
+		const FSeamEdge SeamEdge(Seams[SeamID][StitchID], Seams[SeamID][StitchID + 1]);
+
+		// Check if the pair of mesh edges refer to the same edge
+		const int EdgeA = Mesh.FindEdge(SeamEdge.Stitches[0][0], SeamEdge.Stitches[1][0]);
+		const int EdgeB = Mesh.FindEdge(SeamEdge.Stitches[0][1], SeamEdge.Stitches[1][1]);
+
+		if (EdgeA == EdgeB)
+		{
+			// Don't collapse the same edge twice
+			FSeamCollapseParameters Constraints;
+			Constraints.bCanCollapse = false;
+			return Constraints;
+		}
+
+		// Check if any vertex exists in another stitch somewhere. For now we will skip these operations
+		// TODO: If only one vertex is involved in another stitch somewhere else, we could constrain that vertex to be kept instead
+		TStaticArray<int32, 4> SeamEdgeVertices;
+		SeamEdgeVertices[0] = SeamEdge.Stitches[0][0];
+		SeamEdgeVertices[1] = SeamEdge.Stitches[0][1];
+		SeamEdgeVertices[2] = SeamEdge.Stitches[1][0];
+		SeamEdgeVertices[3] = SeamEdge.Stitches[1][1];
+
+		for (int32 InnerSeamID = 0; InnerSeamID < Seams.Num(); ++InnerSeamID)
+		{
+			const TArray<FIntVector2>& InnerSeam = Seams[InnerSeamID];
+			for (int32 InnerStitchID = 0; InnerStitchID < InnerSeam.Num(); ++InnerStitchID)
+			{
+				if (InnerSeamID == SeamID)
+				{
+					if (InnerStitchID == StitchID || InnerStitchID == StitchID + 1)
+					{
+						// Don't check against adjacent stitches
+						continue;
+					}
+				}
+
+				const FIntVector2& InnerStitch = InnerSeam[InnerStitchID];
+
+				for (const int32& SeamVertex : SeamEdgeVertices)
+				{
+					if (SeamVertex == InnerStitch[0] || SeamVertex == InnerStitch[1])
+					{
+						FSeamCollapseParameters Constraints;
+						Constraints.bCanCollapse = false;
+						return Constraints;
+					}
+				}
+			}
+		}
+
+		// Now check for conditions that might prevent one or the other vertex from being deleted
+
+		bool bStitchIsConstrained = false;
+		bool bNextStitchIsConstrained = false;
+
+		// check if stitch is at the beginning or end of the seam
+		if (StitchID == 0)
+		{
+			bStitchIsConstrained = true;
+		}
+		if (StitchID + 1 == Seams[SeamID].Num() - 1)
+		{
+			bNextStitchIsConstrained = true;
+		}
+
+		// check for vertices that connect two seam sides -- these were previously added by creating a stitch with the same vertex twice
+		if (!bStitchIsConstrained && (SeamEdge.Stitches[0][0] == SeamEdge.Stitches[0][1]))
+		{
+			bStitchIsConstrained = true;
+		}
+		if (!bNextStitchIsConstrained &&(SeamEdge.Stitches[1][0] == SeamEdge.Stitches[1][1]))
+		{
+			bNextStitchIsConstrained = true;
+		}
+
+		// check if any vertex is at a sharp corner
+
+		for (int32 Side = 0; Side < 2; ++Side)
+		{
+			if (StitchID > 0)
+			{
+				const FSeamEdge PrevSeamEdge(Seams[SeamID][StitchID - 1], Seams[SeamID][StitchID]);
+
+				check(PrevSeamEdge.Stitches[1][Side] == SeamEdge.Stitches[0][Side]);
+				check(PrevSeamEdge.Stitches[1][Side] == SeamEdge.Stitches[0][Side])
+
+				const FVector3d A = Mesh.GetVertex(PrevSeamEdge.Stitches[0][Side]);
+				const FVector3d B = Mesh.GetVertex(PrevSeamEdge.Stitches[1][Side]);
+				const FVector3d C = Mesh.GetVertex(SeamEdge.Stitches[1][Side]);
+
+				bStitchIsConstrained = IsSharpCorner(A, B, C, CosAngleCornerTheshold);
+			}
+
+			if (StitchID < Seams[SeamID].Num() - 2)
+			{
+				const FSeamEdge NextSeamEdge(Seams[SeamID][StitchID + 1], Seams[SeamID][StitchID + 2]);
+
+				check(SeamEdge.Stitches[1][Side] == NextSeamEdge.Stitches[0][Side]);
+				check(NextSeamEdge.Stitches[0][Side] == SeamEdge.Stitches[1][Side])
+
+				const FVector3d A = Mesh.GetVertex(SeamEdge.Stitches[0][Side]);
+				const FVector3d B = Mesh.GetVertex(SeamEdge.Stitches[1][Side]);
+				const FVector3d C = Mesh.GetVertex(NextSeamEdge.Stitches[1][Side]);
+
+				bNextStitchIsConstrained = IsSharpCorner(A, B, C, CosAngleCornerTheshold);
+			}
+		}
+
+		if (bStitchIsConstrained && bNextStitchIsConstrained)
+		{
+			FSeamCollapseParameters Constraints;
+			Constraints.bCanCollapse = false;
+			return Constraints;
+		}
+
+		FSeamCollapseParameters Constraints;
+		Constraints.bCanCollapse = true;
+
+		if (bStitchIsConstrained && !bNextStitchIsConstrained)
+		{
+			Constraints.KeepStitchIndex = StitchID;
+			Constraints.DeleteStitchIndex = StitchID+1;
+			Constraints.NewPositionEdgeParameter = 0.0;		// "Keep" vertex should stay where it is
+		}
+		else if (!bStitchIsConstrained && bNextStitchIsConstrained)
+		{
+			Constraints.KeepStitchIndex = StitchID+1;
+			Constraints.DeleteStitchIndex = StitchID;
+			Constraints.NewPositionEdgeParameter = 0.0;	    // "Keep" vertex should stay where it is
+		}
+		else
+		{
+			// unconstrained
+			Constraints.KeepStitchIndex = StitchID;
+			Constraints.DeleteStitchIndex = StitchID + 1;
+			Constraints.NewPositionEdgeParameter = 0.5;		// Collapse to the edge midpoint
+		}
+
+
+		// Check if either collapse would fail in FDynamicMesh3::CollapseEdge
+		for (int32 Side = 0; Side < 2; ++Side)
+		{
+			const int KeepVertexIndex = Seams[SeamID][Constraints.KeepStitchIndex][Side];
+			const int DeleteVertexIndex = Seams[SeamID][Constraints.DeleteStitchIndex][Side];
+
+			const EMeshResult CanCollapsePreview = Mesh.CanCollapseEdge(KeepVertexIndex, DeleteVertexIndex, Constraints.NewPositionEdgeParameter);
+			if (CanCollapsePreview != EMeshResult::Ok)
+			{
+				Constraints.bCanCollapse = false;
+				break;
+			}
+		}
+
+		return Constraints;
+	}
+
+
+	void RemeshSeams(UE::Geometry::FDynamicMesh3& Mesh, TArray<TArray<FIntVector2>>& Seams, double TargetEdgeLength)
+	{
+		using namespace UE::Geometry;
+
+		// constants pulled from FRemesher::SetTargetEdgeLength
+		const double MinLength = 0.66 * TargetEdgeLength;
+		const double MaxLength = 1.33 * TargetEdgeLength;
+
+		for (int32 SeamID = 0; SeamID < Seams.Num(); ++SeamID)
+		{
+			TArray<FIntVector2>& Seam = Seams[SeamID];
+
+			for (int32 StitchID = 0; StitchID < Seam.Num() - 1; ++StitchID)
+			{
+				const int32 SideAVertexA = Seam[StitchID][0];
+				const int32 SideAVertexB = Seam[StitchID + 1][0];
+
+				if (Mesh.FindEdge(SideAVertexA, SideAVertexB) == UE::Geometry::FDynamicMesh3::InvalidID)
+				{
+					continue;
+				}
+
+				const double EdgeALength = (Mesh.GetVertex(SideAVertexA) - Mesh.GetVertex(SideAVertexB)).Length();
+
+				const int32 SideBVertexA = Seam[StitchID][1];
+				const int32 SideBVertexB = Seam[StitchID + 1][1];
+
+				if (Mesh.FindEdge(SideBVertexA, SideBVertexB) == UE::Geometry::FDynamicMesh3::InvalidID)
+				{
+					continue;
+				}
+
+				const double EdgeBLength = (Mesh.GetVertex(SideBVertexA) - Mesh.GetVertex(SideBVertexB)).Length();
+
+				if ((EdgeALength < MinLength) && (EdgeBLength < MinLength))
+				{
+					//
+					// Collapse
+					//
+
+					const FSeamCollapseParameters CollapseConstraints = GetSeamEdgeCollapseParameters(SeamID, StitchID, Mesh, Seams);
+					
+					if (!CollapseConstraints.bCanCollapse)
+					{
+						continue;
+					}
+
+					const int32 PatternAKeepVertex = Seam[CollapseConstraints.KeepStitchIndex][0];
+					const int32 PatternADeleteVertex = Seam[CollapseConstraints.DeleteStitchIndex][0];
+					const int32 PatternBKeepVertex = Seam[CollapseConstraints.KeepStitchIndex][1];
+					const int32 PatternBDeleteVertex = Seam[CollapseConstraints.DeleteStitchIndex][1];
+
+					FDynamicMesh3::FEdgeCollapseInfo CollapseInfoA;
+					const EMeshResult ResultA = Mesh.CollapseEdge(PatternAKeepVertex, PatternADeleteVertex, CollapseConstraints.NewPositionEdgeParameter, CollapseInfoA);
+
+					FDynamicMesh3::FEdgeCollapseInfo CollapseInfoB;
+					const EMeshResult ResultB = Mesh.CollapseEdge(PatternBKeepVertex, PatternBDeleteVertex, CollapseConstraints.NewPositionEdgeParameter, CollapseInfoB);
+
+					check(ResultA == EMeshResult::Ok && ResultB == EMeshResult::Ok)
+
+					Seam.RemoveAt(CollapseConstraints.DeleteStitchIndex);
+				}
+				else if ((EdgeALength > MaxLength) && (EdgeBLength > MaxLength))
+				{
+					//
+					// Split
+					//
+
+					FDynamicMesh3::FEdgeSplitInfo SplitInfoA;
+					const EMeshResult ResultA = Mesh.SplitEdge(SideAVertexA, SideAVertexB, SplitInfoA);
+
+					if (ResultA == EMeshResult::Ok)
+					{
+						if (Mesh.FindEdge(SideBVertexA, SideBVertexB) == UE::Geometry::FDynamicMesh3::InvalidID)
+						{
+							// Don't split the same edge twice
+							continue;
+						}
+
+						FDynamicMesh3::FEdgeSplitInfo SplitInfoB;
+						const EMeshResult ResultB = Mesh.SplitEdge(SideBVertexA, SideBVertexB, SplitInfoB);
+						check(ResultB == EMeshResult::Ok);
+
+						const FIntVector2 NewStitch{ SplitInfoA.NewVertex, SplitInfoB.NewVertex };
+
+						Seam.Insert(NewStitch, StitchID+1);
+					}
+				}
+			}
+		}
+	}
+
+
+	//
+	// Boundary remeshing
+	//
+
+	struct FEdgeCollapseParameters
+	{
+		bool bCanCollapse = false;
+		int32 KeepVertexIndex = -1;
+		int32 DeleteVertexIndex = -1;
+		double NewPositionEdgeParameter = 0.5;
+	};
+
+	FEdgeCollapseParameters GetBoundaryEdgeCollapseParameters(const UE::Geometry::FDynamicMesh3& Mesh, const TArray<TArray<FIntVector2>>& Seams, const UE::Geometry::FIndex2i& EdgeVerts, double CosAngleThreshold = 0.5)
+	{
+		using namespace UE::Geometry;
+
+		auto IsSeamVertex = [](const TArray<TArray<FIntVector2>>& Seams, int VertexIndex) -> bool
+		{
+			for (const TArray<FIntVector2>& Seam : Seams)
+			{
+				for (int32 StitchID = 0; StitchID < Seam.Num(); ++StitchID)
+				{
+					const FIntVector2& Stitch = Seam[StitchID];
+					if ((int)Stitch[0] == VertexIndex || (int)Stitch[1] == VertexIndex)
+					{
+						return true;
+					}
+				}
+			}
+
+			return false;
+		};
+
+		auto FindAdjacentBoundaryVertex = [](const UE::Geometry::FDynamicMesh3& Mesh, int VertexID, int ExcludedVertexID) -> int
+		{
+			for (int EdgeID : Mesh.VtxEdgesItr(VertexID))
+			{
+				if (!Mesh.IsBoundaryEdge(EdgeID))
+				{
+					continue;
+				}
+
+				const UE::Geometry::FIndex2i Edge = Mesh.GetEdge(EdgeID).Vert;
+				const int TestVertexID = Edge[0] == VertexID ? Edge[1] : Edge[0];
+
+				if (TestVertexID != ExcludedVertexID)
+				{
+					return TestVertexID;
+				}
+			}
+
+			return UE::Geometry::FDynamicMesh3::InvalidID;
+		};
+
+
+		const int VertexA = EdgeVerts[0];
+		const int VertexB = EdgeVerts[1];
+
+		// Check for conditions that might prevent one or the other vertex from being deleted
+		bool bVertexIsConstrained = false;
+		bool bNextVertexIsConstrained = false;
+
+		// Check if either vertex is on a seam
+
+		if (IsSeamVertex(Seams, VertexA))
+		{
+			bVertexIsConstrained = true;
+		}
+	
+		if (IsSeamVertex(Seams, VertexB))
+		{
+			bNextVertexIsConstrained = true;
+		}
+
+		// check if either vertex is at a sharp boundary corner
+
+		int OtherVertex = FindAdjacentBoundaryVertex(Mesh, VertexA, VertexB);
+		if ((OtherVertex != FDynamicMesh3::InvalidID) && !bVertexIsConstrained)
+		{
+			const FVector3d A = Mesh.GetVertex(OtherVertex);
+			const FVector3d B = Mesh.GetVertex(VertexA);
+			const FVector3d C = Mesh.GetVertex(VertexB);
+
+			bVertexIsConstrained = IsSharpCorner(A, B, C, CosAngleThreshold);
+		}
+
+		OtherVertex = FindAdjacentBoundaryVertex(Mesh, VertexB, VertexA);
+		if ((OtherVertex != FDynamicMesh3::InvalidID) && !bNextVertexIsConstrained)
+		{
+			const FVector3d A = Mesh.GetVertex(VertexA);
+			const FVector3d B = Mesh.GetVertex(VertexB);
+			const FVector3d C = Mesh.GetVertex(OtherVertex);
+
+			bNextVertexIsConstrained = IsSharpCorner(A, B, C, CosAngleThreshold);
+		}
+
+		if (bVertexIsConstrained && bNextVertexIsConstrained)
+		{
+			FEdgeCollapseParameters Constraints;
+			Constraints.bCanCollapse = false;
+			return Constraints;
+		}
+
+		FEdgeCollapseParameters Constraints;
+		Constraints.bCanCollapse = true;
+
+		if (bVertexIsConstrained && !bNextVertexIsConstrained)
+		{
+			Constraints.KeepVertexIndex = VertexA;
+			Constraints.DeleteVertexIndex = VertexB;
+			Constraints.NewPositionEdgeParameter = 0.0;		// "Keep" vertex should stay where it is
+		}
+		else if (!bVertexIsConstrained && bNextVertexIsConstrained)
+		{
+			Constraints.KeepVertexIndex = VertexB;
+			Constraints.DeleteVertexIndex = VertexA;
+			Constraints.NewPositionEdgeParameter = 0.0;	    // "Keep" vertex should stay where it is
+		}
+		else
+		{
+			// unconstrained
+			Constraints.KeepVertexIndex = VertexA;
+			Constraints.DeleteVertexIndex = VertexB;
+			Constraints.NewPositionEdgeParameter = 0.5;		// Collapse to the edge midpoint
+		}
+
+		return Constraints;
+	}
+
+	void RemeshBoundaries(UE::Geometry::FDynamicMesh3& Mesh, const TArray<TArray<FIntVector2>>& Seams, double TargetEdgeLength)
+	{
+		using namespace UE::Geometry;
+
+		auto IsSeamEdge = [](const TArray<TArray<FIntVector2>>& Seams, int EdgeVertexA, int EdgeVertexB)
+		{
+			for (const TArray<FIntVector2>& Seam : Seams)
+			{
+				for (int32 StitchID = 0; StitchID < Seam.Num() - 1; ++StitchID)
+				{
+					const FIntVector2& Stitch = Seam[StitchID];
+					const FIntVector2& NextStitch = Seam[StitchID + 1];
+
+					for (int32 Side = 0; Side < 2; ++Side)
+					{
+						if ((int)Stitch[Side] == EdgeVertexA && (int)NextStitch[Side] == EdgeVertexB)
+						{
+							return true;
+						}
+
+						if ((int)Stitch[Side] == EdgeVertexB && (int)NextStitch[Side] == EdgeVertexA)
+						{
+							return true;
+						}
+					}
+				}
+			}
+
+			return false;
+		};
+
+		// constants pulled from FRemesher::SetTargetEdgeLength
+		const double MinLength = 0.66 * TargetEdgeLength;
+		const double MaxLength = 1.33 * TargetEdgeLength;
+
+		// Get the set of boundary edges up front and then check if they get invalidated later. 
+		// If we process edges inside this loop it tends to collapse a bunch of sequential edges, which can lead to very high-valence vertices
+		TArray<FIndex2i> BoundaryEdges;
+		for (const int EdgeID : Mesh.BoundaryEdgeIndicesItr())
+		{
+			BoundaryEdges.Add(Mesh.GetEdge(EdgeID).Vert);
+		}
+
+		for (const FIndex2i& EdgeVerts : BoundaryEdges)
+		{
+			const int EdgeID = Mesh.FindEdge(EdgeVerts[0], EdgeVerts[1]);
+			if (EdgeID == FDynamicMesh3::InvalidID || !Mesh.IsBoundaryEdge(EdgeID))
+			{
+				continue;
+			}
+
+			if (IsSeamEdge(Seams, EdgeVerts[0], EdgeVerts[1]))
+			{
+				continue;
+			}
+
+			const double EdgeLength = (Mesh.GetVertex(EdgeVerts[0]) - Mesh.GetVertex(EdgeVerts[1])).Length();
+			if (EdgeLength < MinLength)
+			{
+				const Private::FEdgeCollapseParameters CollapseParams = Private::GetBoundaryEdgeCollapseParameters(Mesh, Seams, EdgeVerts);
+
+				if (CollapseParams.bCanCollapse)
+				{
+					FDynamicMesh3::FEdgeCollapseInfo CollapseInfo;
+					EMeshResult CollapseResult = Mesh.CollapseEdge(CollapseParams.KeepVertexIndex, CollapseParams.DeleteVertexIndex, CollapseParams.NewPositionEdgeParameter, CollapseInfo);
+				}
+			}
+			else if (EdgeLength > MaxLength)
+			{
+				FDynamicMesh3::FEdgeSplitInfo SplitInfo;
+				EMeshResult SplitResult = Mesh.SplitEdge(EdgeVerts[0], EdgeVerts[1], SplitInfo);
+			}
+		};
+	}
+
+
+	//
+	// Remeshing away from Seams/Boundaries
+	// 
+
+	bool Remesh(UE::Geometry::FDynamicMesh3& Mesh, double TargetEdgeLength, int32 Iterations, float SmoothingRate, bool bUniformSmoothing, const TArray<TArray<FIntVector2>>& Seams, UE::Geometry::FCompactMaps* CompactMaps = nullptr)
+	{
+		using namespace UE::Geometry;
+
+		//
+		// These consts control overall remeshing behavior and are analogs of the properties exposed to the user in the actual Remesh tool
+		// (i.e. things we might wish to add to the node properties in the future)
+		//
+
+		constexpr bool bReprojectToInputMesh = false;
+		constexpr bool bDiscardAttributes = false;
+		constexpr bool bUseFullRemeshPasses = true;
+		constexpr bool bAllowFlips = true;
+		constexpr bool bAllowSplits = true;
+		constexpr bool bAllowCollapses = true;
+		constexpr bool bPreventNormalFlips = true;
+		constexpr bool bPreventTinyTriangles = true;
+		constexpr bool bAutoCompact = true;
+		constexpr bool bCoarsenBoundaries = false;
+
+		const ERemeshSmoothingType SmoothingType = bUniformSmoothing ? ERemeshSmoothingType::Uniform : ERemeshSmoothingType::MeanValue;
+
+		// Mesh seam behavior
+		// Here we are talking UV, Normal, and Color seams, not Cloth Seams
+		// This controls bAllowSeamCollapse and bAllowSeamSmoothing on those overlay seams
+		// NOTE: These seams are not affected by bReprojectConstraints
+		constexpr bool bPreserveSharpEdges = false;
+
+		// Mesh boundaries
+		const EEdgeRefineFlags MeshBoundaryConstraint = bCoarsenBoundaries ? EEdgeRefineFlags::NoFlip : EEdgeRefineFlags::FullyConstrained;
+
+		// Group ID boundaries
+		const EEdgeRefineFlags PolyGroupBoundaryConstraint = bCoarsenBoundaries ? EEdgeRefineFlags::NoFlip : EEdgeRefineFlags::FullyConstrained;
+
+		// Material ID boundaries
+		const EEdgeRefineFlags MaterialBoundaryConstraint = bCoarsenBoundaries ? EEdgeRefineFlags::NoFlip : EEdgeRefineFlags::FullyConstrained;
+
+		// Whether to move boundary vertices back onto the poly-line defined by the original boundary in case of collapse
+		const bool bReprojectConstraints = bCoarsenBoundaries;
+
+		// Seam "corners" are held fixed. We use this angle threshold to determine what constitutes a seam corner
+		constexpr float SeamCornerThresholdAngleDegrees = 45.0;
+
+		FRemeshMeshOp RemeshOp;
+
+		TSharedPtr<FDynamicMesh3> SourceMesh = MakeShared<FDynamicMesh3>(MoveTemp(Mesh));
+		TSharedPtr<FDynamicMeshAABBTree3> SourceSpatial;
+		if (bReprojectToInputMesh)
+		{
+			// acceleration structure is only used for reprojecting
+			SourceSpatial = MakeShared<FDynamicMeshAABBTree3>(SourceMesh.Get(), true);
+		}
+
+		RemeshOp.OriginalMesh = SourceMesh;
+		RemeshOp.OriginalMeshSpatial = SourceSpatial;
+
+		RemeshOp.bDiscardAttributes = bDiscardAttributes;
+		RemeshOp.RemeshType = (bUseFullRemeshPasses) ? ERemeshType::FullPass : ERemeshType::Standard;
+		RemeshOp.RemeshIterations = Iterations;
+		RemeshOp.MaxRemeshIterations = Iterations;
+		RemeshOp.ExtraProjectionIterations = 0;		// unused for regular remeshing
+		RemeshOp.TriangleCountHint = 0;				// unused for regular remeshing
+		RemeshOp.SmoothingStrength = FMath::Clamp(SmoothingRate, 0.0f, 1.0f);
+		RemeshOp.SmoothingType = SmoothingType;
+
+		RemeshOp.TargetEdgeLength = TargetEdgeLength;
+		RemeshOp.bPreserveSharpEdges = bPreserveSharpEdges;
+		RemeshOp.bFlips = bAllowFlips;
+		RemeshOp.bSplits = bAllowSplits;
+		RemeshOp.bCollapses = bAllowCollapses;
+		RemeshOp.bPreventNormalFlips = bPreventNormalFlips;
+		RemeshOp.bPreventTinyTriangles = bPreventTinyTriangles;
+		RemeshOp.MeshBoundaryConstraint = MeshBoundaryConstraint;
+		RemeshOp.GroupBoundaryConstraint = PolyGroupBoundaryConstraint;
+		RemeshOp.MaterialBoundaryConstraint = MaterialBoundaryConstraint;
+		RemeshOp.bReproject = bReprojectToInputMesh;
+		RemeshOp.ProjectionTarget = SourceMesh.Get();
+		RemeshOp.ProjectionTargetSpatial = SourceSpatial.Get();
+		RemeshOp.bReprojectConstraints = bReprojectConstraints;
+		RemeshOp.BoundaryCornerAngleThreshold = SeamCornerThresholdAngleDegrees;
+		RemeshOp.TargetMeshLocalToWorld = FTransformSRT3d::Identity();
+		RemeshOp.ToolMeshLocalToWorld = FTransformSRT3d::Identity();
+		RemeshOp.bUseWorldSpace = false;
+		RemeshOp.bParallel = true;
+
+
+		// Set up constraints for Cloth Seam edges
+		// TODO: Add edge or vertex contraints for other cloth sim mesh data (e.g. collision particles, tethers)
+		UE::Geometry::FMeshConstraints Constraints;
+		for (const TArray<FIntVector2>& Seam : Seams)
+		{
+			for (int32 StitchIndex = 0; StitchIndex < Seam.Num() - 1; ++StitchIndex)
+			{
+				for (int32 Side = 0; Side < 2; ++Side)
+				{
+					const int EdgeID = SourceMesh->FindEdge(Seam[StitchIndex][Side], Seam[StitchIndex + 1][Side]);
+					check(EdgeID != FDynamicMesh3::InvalidID);
+					FEdgeConstraint EdgeConstraint(EEdgeRefineFlags::FullyConstrained);
+					Constraints.SetOrUpdateEdgeConstraint(EdgeID, EdgeConstraint);
+
+					constexpr bool bCannotDelete = true;
+					constexpr bool bCanMove = false;
+					FVertexConstraint VertexConstraint(bCannotDelete, bCanMove);
+					Constraints.SetOrCombineVertexConstraint(Seam[StitchIndex][Side], VertexConstraint);
+					Constraints.SetOrCombineVertexConstraint(Seam[StitchIndex + 1][Side], VertexConstraint);
+				}
+			}
+		}
+		RemeshOp.SetUserSpecifiedConstraints(Constraints);
+
+		constexpr FProgressCancel* Progress = nullptr;		// Don't allow cancel or report progress for now
+		RemeshOp.CalculateResult(Progress);
+
+		if (RemeshOp.GetResultInfo().Result == EGeometryResultType::Success)
+		{
+			TUniquePtr<FDynamicMesh3> ResultMesh = RemeshOp.ExtractResult();
+			Mesh = MoveTemp(*ResultMesh);
+		}
+		else
+		{
+			return false;
+		}
+
+		// compact the input mesh if enabled
+		if (bAutoCompact)
+		{
+			Mesh.CompactInPlace(CompactMaps);
+		}
+
+		return true;
+	}
+
+
+	void ProjectTo3D(const UE::Geometry::FDynamicMesh3& Mesh2D,
+		const UE::Geometry::FDynamicMesh3& ProjectionTarget2D,
+		const UE::Geometry::FDynamicMesh3& ProjectionTarget3D,
+		UE::Geometry::FDynamicMesh3& OutMesh3D)
+	{
+		using namespace UE::Geometry;
+
+		OutMesh3D.Copy(Mesh2D);
+
+		TSharedPtr<UE::Geometry::FDynamicMeshAABBTree3> ProjectionTargetSpatial = MakeShared<UE::Geometry::FDynamicMeshAABBTree3>(&ProjectionTarget2D, true);
+
+		for (const int VertexIndex : Mesh2D.VertexIndicesItr())
+		{
+			const FVector3d SrcVert = Mesh2D.GetVertex(VertexIndex);
+
+			double Distance;
+			const int NearestTriangle = ProjectionTargetSpatial->FindNearestTriangle(SrcVert, Distance);
+
+			const FDistPoint3Triangle3d Dist = TMeshQueries<FDynamicMesh3>::TriangleDistance(ProjectionTarget2D, NearestTriangle, SrcVert);
+			const FVector3d Bary = Dist.TriangleBaryCoords;
+			const FVector3d InterpolatedPoint = ProjectionTarget3D.GetTriBaryPoint(NearestTriangle, Bary[0], Bary[1], Bary[2]);
+
+			OutMesh3D.SetVertex(VertexIndex, InterpolatedPoint);
+		}
+	}
+
+}
+
+
+FChaosClothAssetRemeshNode::FChaosClothAssetRemeshNode(const Dataflow::FNodeParameters& InParam, FGuid InGuid)
+	: FDataflowNode(InParam, InGuid)
+{
+	RegisterInputConnection(&Collection);
+	RegisterOutputConnection(&Collection, &Collection);
+}
+
+
+void FChaosClothAssetRemeshNode::RemeshSimMesh(TSharedRef<FManagedArrayCollection> ClothCollection,
+	TSharedRef<FManagedArrayCollection> OutClothCollection) const
+{
+	using namespace UE::Geometry;
+	using namespace UE::Chaos::ClothAsset;
+
+	FCollectionClothConstFacade InClothFacade(ClothCollection);
+
+	if (InClothFacade.GetNumSimPatterns() == 0)
+	{
+		FClothGeometryTools::DeleteSimMesh(OutClothCollection);
+		return;
+	}
+
+	// Convert input patterns to a DynamicMesh
+
+	FClothPatternToDynamicMesh Converter;
+
+	FDynamicMesh3 Mesh2D;
+	Converter.Convert(ClothCollection, INDEX_NONE, EClothPatternVertexType::Sim2D, Mesh2D);
+
+	const double TotalArea = TMeshQueries<FDynamicMesh3>::GetVolumeArea(Mesh2D).Y;
+	const int TriangleCount = Mesh2D.TriangleCount();
+
+	// Copy pattern IDs into polygroup layer
+
+	Mesh2D.EnableAttributes();
+	const int32 PatternIndexLayerID = Mesh2D.Attributes()->NumPolygroupLayers();
+	Mesh2D.Attributes()->SetNumPolygroupLayers(PatternIndexLayerID + 1);
+	FDynamicMeshPolygroupAttribute* const PatternIndexLayer = Mesh2D.Attributes()->GetPolygroupLayer((int)PatternIndexLayerID);
+	check(Mesh2D.TriangleCount() == InClothFacade.GetNumSimFaces());
+	for (int32 FaceIndex = 0; FaceIndex < InClothFacade.GetNumSimFaces(); ++FaceIndex)
+	{
+		const int32 PatternID = InClothFacade.FindSimPatternByFaceIndex(FaceIndex);
+		PatternIndexLayer->SetValue(FaceIndex, PatternID);
+	}
+
+	// Grab the seam information from the input collection
+	
+	// Seams are comprised of a set of Stitches. Each Stitch is simply a pair of vertex indices indicating vertices that should be welded to form the 3D mesh.
+	
+	// Stitches are given in random order within the Seam. To make remeshing them easier, we will find connected strips of stitches and store them in sequential order.
+	// So the vertices in Stitch N are connected to the vertices in Stitch N+1.
+
+	TArray<TArray<FIntVector2>> Seams;
+	for (int32 SeamIndex = 0; SeamIndex < InClothFacade.GetNumSeams(); ++SeamIndex)
+	{
+		FCollectionClothSeamConstFacade SeamFacade = InClothFacade.GetSeam(SeamIndex);
+		FClothGeometryTools::BuildConnectedSeams2D(ClothCollection, SeamIndex, Mesh2D, Seams);
+
+		// Check seams are valid
+		for (int32 SeamID = 0; SeamID < Seams.Num(); ++SeamID)
+		{
+			const TArray<FIntVector2>& SubSeam = Seams[SeamID];
+			for (int32 StitchID = 0; StitchID < SubSeam.Num() - 1; ++StitchID)
+			{
+				const int32 NextStitchID = StitchID + 1;
+				for (int32 Side = 0; Side < 2; ++Side)
+				{
+					const int32 StitchVert = SubSeam[StitchID][Side];
+					const int32 NextStitchVert = SubSeam[NextStitchID][Side];
+
+					if (StitchVert == NextStitchVert)
+					{
+						continue;
+					}
+
+					const int32 FoundEdge = Mesh2D.FindEdge(StitchVert, NextStitchVert);
+
+					// This would indicate a problem in BuildConnectedSeams
+					checkf(FoundEdge != UE::Geometry::FDynamicMesh3::InvalidID, TEXT("Could not find a mesh edge between sequential seam vertices %d, %d"), StitchVert, NextStitchVert);
+				}
+			}
+		}
+	}
+
+
+	// For remeshing purposes, we also want to find and constrain any vertices that connect the two seam sides together
+	// (This is only relevant for "internal" seams, which connect vertices within the same pattern)
+	for (TArray<FIntVector2>& Seam : Seams)
+	{
+		if (Seam.Num() == 0)
+		{
+			continue;
+		}
+
+		for (const FIntVector2& EndStitch : {Seam[0], Seam.Last()})
+		{
+			for (const int NeighborA : Mesh2D.VtxVerticesItr(EndStitch[0]))
+			{
+				bool bMatchFound = false;
+
+				for (const int NeighborB : Mesh2D.VtxVerticesItr(EndStitch[1]))
+				{
+					if (NeighborA == NeighborB)
+					{
+						if (Seam.Contains(FIntVector2(NeighborA, NeighborB)))
+						{
+							continue;
+						}
+
+						if (EndStitch == Seam[0])
+						{
+							Seam.Insert(FIntVector2(NeighborA, NeighborB), 0);
+						}
+						else
+						{
+							Seam.Add(FIntVector2(NeighborA, NeighborB));
+						}
+
+						bMatchFound = true;
+						break;
+					}
+				}
+
+				if (bMatchFound)
+				{
+					break;
+				}
+			}
+		}
+	}
+
+	// Remesh seams
+
+	const int TargetTriangleCount = FMath::RoundToInt(static_cast<float>(TargetPercentSim) / 100.0f * static_cast<float>(TriangleCount));
+	const double TargetEdgeLength = FRemeshMeshOp::CalculateTargetEdgeLength(nullptr, TargetTriangleCount, TotalArea);
+
+	for (int32 ResampleIter = 0; ResampleIter < IterationsSim; ++ResampleIter)
+	{
+		Private::RemeshSeams(Mesh2D, Seams, TargetEdgeLength);
+	}
+
+	// Remesh boundaries
+
+	for (int32 ResampleIter = 0; ResampleIter < IterationsSim; ++ResampleIter)
+	{
+		Private::RemeshBoundaries(Mesh2D, Seams, TargetEdgeLength);
+	}
+
+	// Do the remeshing of the rest of the mesh
+
+	UE::Geometry::FCompactMaps CompactMaps;
+	constexpr bool bUniformSmoothing = true;
+	Private::Remesh(Mesh2D, TargetEdgeLength, IterationsSim, SmoothingSim, bUniformSmoothing, Seams, &CompactMaps);
+
+	// Update stitches
+	for (TArray<FIntVector2>& Seam : Seams)
+	{
+		for (FIntVector2& Stitch : Seam)
+		{
+			Stitch[0] = CompactMaps.GetVertexMapping(Stitch[0]);
+			Stitch[1] = CompactMaps.GetVertexMapping(Stitch[1]);
+			checkf(Stitch[0] != FDynamicMesh3::InvalidID, TEXT("Stitch vertex %d was deleted by remeshing"), Stitch[0]);
+			checkf(Stitch[1] != FDynamicMesh3::InvalidID, TEXT("Stitch vertex %d was deleted by remeshing"), Stitch[1]);
+		}
+	}
+
+
+	// Project the 3D vertices onto the input 3D mesh
+
+	FDynamicMesh3 SourceMesh2D;
+	Converter.Convert(ClothCollection, INDEX_NONE, EClothPatternVertexType::Sim2D, SourceMesh2D);
+	FDynamicMesh3 SourceMesh3D;
+	Converter.Convert(ClothCollection, INDEX_NONE, EClothPatternVertexType::Sim3D, SourceMesh3D);
+
+	FDynamicMesh3 Mesh3D;
+	Private::ProjectTo3D(Mesh2D, SourceMesh2D, SourceMesh3D, Mesh3D);
+
+
+	// Build the output cloth sim mesh
+
+	constexpr bool bAppendToExistingMesh = false;
+	constexpr bool bTransferWeightMaps = true;
+	constexpr bool bTransferSimSkinningData = true;
+	TMap<int, int32> DynamicMeshToClothVertexMap;
+	FClothGeometryTools::BuildSimMeshFromDynamicMeshes(OutClothCollection, Mesh2D, Mesh3D, PatternIndexLayerID, bTransferWeightMaps, bTransferSimSkinningData, bAppendToExistingMesh, DynamicMeshToClothVertexMap);
+
+
+	// Re-apply the seam info from the input sim mesh
+	// TODO: This will create a new Seam for each set of connected stitches. Should it instead be one seam for each pair of patterns?
+	//       E.g. if there are two strips of vertex pairs between a two given patterns, it will generate two seams. Should it just put them all in one seam?
+
+	FCollectionClothFacade OutClothFacade(OutClothCollection);
+	for (int32 SeamIndex = 0; SeamIndex < Seams.Num(); ++SeamIndex)
+	{
+		const TArray<FIntVector2>& Seam = Seams[SeamIndex];
+
+		TArray<FIntVector2> NewSeam;
+		for (const FIntVector2& Stitch : Seam)
+		{
+			if (Stitch[0] == Stitch[1])
+			{
+				continue;
+			}
+
+			FIntVector2 NewStitch;
+			NewStitch[0] = DynamicMeshToClothVertexMap[Stitch[0]];
+			NewStitch[1] = DynamicMeshToClothVertexMap[Stitch[1]];
+			NewSeam.Add(NewStitch);
+		}
+
+		FCollectionClothSeamFacade NewSeamFacade = OutClothFacade.AddGetSeam();
+		NewSeamFacade.Initialize(NewSeam);
+	}
+}
+
+
+void FChaosClothAssetRemeshNode::RemeshRenderPattern(TSharedRef<FManagedArrayCollection> ClothCollection,
+	int32 PatternIndex,
+	TSharedRef<FManagedArrayCollection> OutClothCollection) const
+{
+	using namespace UE::Geometry;
+	using namespace UE::Chaos::ClothAsset;
+
+	// Get the source mesh for the pattern
+	FClothPatternToDynamicMesh Converter;
+
+	FDynamicMesh3 DynamicMesh;
+	Converter.Convert(ClothCollection, PatternIndex, EClothPatternVertexType::Render, DynamicMesh);
+
+	const bool bHasUVs = DynamicMesh.HasAttributes() && DynamicMesh.Attributes() && DynamicMesh.Attributes()->PrimaryUV();
+
+	const double PatternArea = TMeshQueries<FDynamicMesh3>::GetVolumeArea(DynamicMesh).Y;
+	const int PatternTriangleCount = DynamicMesh.TriangleCount();
+	const int TargetTriangleCount = FMath::RoundToInt(static_cast<float>(TargetPercentRender) / 100.0f * static_cast<float>(PatternTriangleCount));
+	const double TargetEdgeLength = FRemeshMeshOp::CalculateTargetEdgeLength(nullptr, TargetTriangleCount, PatternArea);
+
+	TArray<TArray<FIntVector2>> Seams;
+
+	for (int32 ResampleIter = 0; ResampleIter < IterationsRender; ++ResampleIter)
+	{
+		Private::RemeshBoundaries(DynamicMesh, Seams, TargetEdgeLength);
+	}
+
+	constexpr bool bUniformSmoothing = false;	// uniform smoothing can distort the UV layer pretty badly
+	const bool bSuccess = Private::Remesh(DynamicMesh, TargetEdgeLength, IterationsRender, SmoothingRender, bUniformSmoothing, Seams);
+	check(bSuccess);
+
+
+	// Set the output cloth collection
+
+	// Collect info from the remeshed DynamicMesh data to reinitialize the Render Pattern
+
+	TArray<FIndex3i> Indices;
+	for (const FIndex3i& Tri : DynamicMesh.TrianglesItr())
+	{
+		Indices.Add(Tri);
+	}
+
+	TArray<FVector3f> Positions;
+	for (const FVector3d& Vertex : DynamicMesh.VerticesItr())
+	{
+		Positions.Add(FVector3f(Vertex));
+	}
+
+	//
+	// Normals
+	//
+
+	const bool bHasNormals = DynamicMesh.HasAttributes() && DynamicMesh.Attributes()->PrimaryNormals();
+	check(bHasNormals);
+
+	TArray<FVector3f> Normals;
+	if (bHasNormals)
+	{
+		Normals.SetNum(Positions.Num());
+		const FDynamicMeshNormalOverlay* const NormalOverlay = DynamicMesh.Attributes()->PrimaryNormals();
+		for (const int TriangleIndex : DynamicMesh.TriangleIndicesItr())
+		{
+			const FIndex3i Tri = DynamicMesh.GetTriangle(TriangleIndex);
+
+			for (int TriangleVertexIndex = 0; TriangleVertexIndex < 3; ++TriangleVertexIndex)
+			{
+				const int VertexIndex = Tri[TriangleVertexIndex];
+
+				// NOTE: This assumes one normal per vertex in the overlay (i.e. no "hard edges")
+				Normals[VertexIndex] = NormalOverlay->GetElementAtVertex(TriangleIndex, VertexIndex);
+			}
+		}
+	}
+	else
+	{
+		// TODO: Compute normals
+
+	}
+
+	//
+	// Tangents
+	//
+
+	TArray<FVector3f> TangentUs;
+	TArray<FVector3f> TangentVs;
+	{
+		const bool bHasTangentUs = DynamicMesh.HasAttributes() && DynamicMesh.Attributes()->PrimaryTangents();
+		const bool bHasTangentVs = DynamicMesh.HasAttributes() && DynamicMesh.Attributes()->PrimaryBiTangents();
+		if (!bHasTangentUs || bHasTangentVs)
+		{
+			FMeshTangentsf::ComputeDefaultOverlayTangents(DynamicMesh);
+		}
+		TangentUs.SetNumZeroed(Positions.Num());
+		TangentVs.SetNumZeroed(Positions.Num());
+
+		const FDynamicMeshNormalOverlay* const TangentUOverlay = DynamicMesh.Attributes()->PrimaryTangents();
+		const FDynamicMeshNormalOverlay* const TangentVOverlay = DynamicMesh.Attributes()->PrimaryBiTangents();
+
+		for (const int TriangleIndex : DynamicMesh.TriangleIndicesItr())
+		{
+			const FIndex3i Tri = DynamicMesh.GetTriangle(TriangleIndex);
+
+			for (int TriangleVertexIndex = 0; TriangleVertexIndex < 3; ++TriangleVertexIndex)
+			{
+				const int VertexIndex = Tri[TriangleVertexIndex];
+
+				TangentUs[VertexIndex] += TangentUOverlay->GetElementAtVertex(TriangleIndex, VertexIndex);
+				TangentVs[VertexIndex] += TangentUOverlay->GetElementAtVertex(TriangleIndex, VertexIndex);
+			}
+		}
+
+		for (int32 VertexIndex = 0; VertexIndex < Positions.Num(); ++VertexIndex)
+		{
+			TangentUs[VertexIndex].Normalize();
+			TangentVs[VertexIndex].Normalize();
+		}
+	}
+
+	//
+	// UVs
+	//
+
+	// TODO: Account for multiple UV layers
+	TArray<FVector2f> UVs;
+	if (bHasUVs)
+	{
+		UVs.SetNum(Positions.Num());
+
+		const FDynamicMeshUVOverlay* const UVOverlay = DynamicMesh.Attributes()->PrimaryUV();
+
+		// Assume no seams in the dynamic mesh UV overlay
+		checkSlow(!UVOverlay->HasInteriorSeamEdges());
+
+		for (const int TriangleIndex : DynamicMesh.TriangleIndicesItr())
+		{
+			const FIndex3i Tri = DynamicMesh.GetTriangle(TriangleIndex);
+
+			for (int TriangleVertexIndex = 0; TriangleVertexIndex < 3; ++TriangleVertexIndex)
+			{
+				const int VertexIndex = Tri[TriangleVertexIndex];
+				UVs[VertexIndex] = UVOverlay->GetElementAtVertex(TriangleIndex, VertexIndex);
+			}
+		}
+	}
+
+	//
+	// Skin Weights
+	//
+
+	FDynamicMeshAttributeSet* Attributes = DynamicMesh.Attributes();
+	TArray<TArray<int32>> BoneIndices;
+	TArray<TArray<float>> BoneWeights;
+
+	if (Attributes)
+	{
+		BoneIndices.SetNum(Positions.Num());
+		BoneWeights.SetNum(Positions.Num());
+
+		for (const TPair<FName, TUniquePtr<FDynamicMeshVertexSkinWeightsAttribute>>& SkinWeightLayer : Attributes->GetSkinWeightsAttributes())
+		{
+			const FDynamicMeshVertexSkinWeightsAttribute* const SkinWeightAttribute = SkinWeightLayer.Value.Get();
+
+			for (const int TriangleIndex : DynamicMesh.TriangleIndicesItr())
+			{
+				const FIndex3i Tri = DynamicMesh.GetTriangle(TriangleIndex);
+
+				for (int TriangleVertexIndex = 0; TriangleVertexIndex < 3; ++TriangleVertexIndex)
+				{
+					const int VertexIndex = Tri[TriangleVertexIndex];
+					SkinWeightAttribute->GetValue(VertexIndex, BoneIndices[VertexIndex], BoneWeights[VertexIndex]);
+				}
+			}
+		}
+	}
+
+
+	FCollectionClothFacade OutClothFacade(OutClothCollection);
+	FCollectionClothRenderPatternFacade OutClothPatternFacade = OutClothFacade.GetRenderPattern(PatternIndex);
+
+	OutClothPatternFacade.SetNumRenderVertices(Positions.Num());
+	OutClothPatternFacade.SetNumRenderFaces(Indices.Num());
+
+	TArrayView<FVector3f> RenderPosition = OutClothPatternFacade.GetRenderPosition();
+	TArrayView<FVector3f> RenderNormal = OutClothPatternFacade.GetRenderNormal();
+	TArrayView<FVector3f> RenderTangentU = OutClothPatternFacade.GetRenderTangentU();
+	TArrayView<FVector3f> RenderTangentV = OutClothPatternFacade.GetRenderTangentV();
+	TArrayView<TArray<FVector2f>> RenderUVs = OutClothPatternFacade.GetRenderUVs();
+	TArrayView<FLinearColor> RenderColor = OutClothPatternFacade.GetRenderColor();
+	TArrayView<TArray<int32>> RenderBoneIndices = OutClothPatternFacade.GetRenderBoneIndices();
+	TArrayView<TArray<float>> RenderBoneWeights = OutClothPatternFacade.GetRenderBoneWeights();
+
+	for (int32 VInd = 0; VInd < Positions.Num(); ++VInd)
+	{
+		RenderPosition[VInd] = Positions[VInd];
+		if (bHasUVs)
+		{
+			RenderUVs[VInd].SetNum(MAX_TEXCOORDS);
+			RenderUVs[VInd][0] = UVs[VInd];
+		}
+		if (bHasNormals)
+		{
+			RenderNormal[VInd] = Normals[VInd];
+		}
+		RenderTangentU[VInd] = TangentUs[VInd];
+		RenderTangentV[VInd] = TangentVs[VInd];
+		RenderColor[VInd] = FLinearColor::White;
+		RenderBoneIndices[VInd] = BoneIndices[VInd];
+		RenderBoneWeights[VInd] = BoneWeights[VInd];
+	}
+
+	const int32 VertexOffset = OutClothPatternFacade.GetRenderVerticesOffset();
+	TArrayView<FIntVector3> RenderIndices = OutClothPatternFacade.GetRenderIndices();
+	for (int32 TInd = 0; TInd < Indices.Num(); ++TInd)
+	{
+		RenderIndices[TInd][0] = VertexOffset + Indices[TInd][0];
+		RenderIndices[TInd][1] = VertexOffset + Indices[TInd][1];
+		RenderIndices[TInd][2] = VertexOffset + Indices[TInd][2];
+	}
+}
+
+
+void FChaosClothAssetRemeshNode::Evaluate(Dataflow::FContext& Context, const FDataflowOutput* Out) const
+{
+	if (Out->IsA<FManagedArrayCollection>(&Collection))
+	{
+		using namespace UE::Chaos::ClothAsset;
+
+		// Evaluate in collection
+		FManagedArrayCollection InCollection = GetValue<FManagedArrayCollection>(Context, &Collection);
+		const TSharedRef<FManagedArrayCollection> ClothCollection = MakeShared<FManagedArrayCollection>(MoveTemp(InCollection));
+
+		// Copy collection to output
+		TSharedRef<FManagedArrayCollection> OutputClothCollection = MakeShared<FManagedArrayCollection>(*ClothCollection);
+
+		FCollectionClothFacade ClothFacade(ClothCollection);
+		if (ClothFacade.IsValid())  // Can only act on the collection if it is a valid cloth collection
+		{
+			if (bRemeshSim)
+			{
+				RemeshSimMesh(ClothCollection, OutputClothCollection);
+			}
+
+			if (bRemeshRender)
+			{
+				// TODO: Does it make more sense to do the whole thing at once, or pattern by pattern? It's easier to rebuild the patterns if we do it this way
+				const int32 NumPatterns = ClothFacade.GetNumRenderPatterns();
+				for (int32 PatternIndex = 0; PatternIndex < NumPatterns; ++PatternIndex)
+				{
+					RemeshRenderPattern(ClothCollection, PatternIndex, OutputClothCollection);
+				}
+			}
+		}
+
+		SetValue(Context, MoveTemp(*OutputClothCollection), &Collection);
+	}
+}
+
+#undef LOCTEXT_NAMESPACE
diff --git a/Engine/Plugins/Experimental/ChaosClothAssetEditor/Source/ChaosClothAssetDataflowNodes/Public/ChaosClothAsset/RemeshNode.h b/Engine/Plugins/Experimental/ChaosClothAssetEditor/Source/ChaosClothAssetDataflowNodes/Public/ChaosClothAsset/RemeshNode.h
new file mode 100644
--- /dev/null
+++ b/Engine/Plugins/Experimental/ChaosClothAssetEditor/Source/ChaosClothAssetDataflowNodes/Public/ChaosClothAsset/RemeshNode.h
@@ -0,0 +1,60 @@
+// Copyright Epic Games, Inc. All Rights Reserved.
+
+#pragma once 
+
+#include "Dataflow/DataflowNode.h"
+#include "GeometryCollection/ManagedArrayCollection.h"
+#include "RemeshNode.generated.h"
+
+/** Remesh the cloth surface(s) to get the specified mesh resolution(s).
+ *  NOTE: Remeshing will invalidate vertex data such as Self Collision Spheres and Long Range Attachment Constraints, however weight maps and skinning data will be resampled to the new mesh topology.
+ */
+USTRUCT(Meta = (DataflowCloth))
+struct FChaosClothAssetRemeshNode : public FDataflowNode
+{
+	GENERATED_USTRUCT_BODY()
+	DATAFLOW_NODE_DEFINE_INTERNAL(FChaosClothAssetRemeshNode, "Remesh", "Cloth", "Cloth Remesh")
+
+public:
+
+	UPROPERTY(Meta = (Dataflowinput, DataflowOutput, DataflowPassthrough = "Collection"))
+	FManagedArrayCollection Collection;
+
+	UPROPERTY(EditAnywhere, Category = "Sim Mesh")
+	bool bRemeshSim = true;
+
+	UPROPERTY(EditAnywhere, Category = "Sim Mesh", meta=(UIMin = "1", UIMax = "200", ClampMin = "1", EditCondition = "bRemeshSim"))
+	int32 TargetPercentSim = 100;
+
+	UPROPERTY(EditAnywhere, Category = "Sim Mesh", meta = (UIMin = "0", UIMax = "20", ClampMin = "0", ClampMax = "100", EditCondition = "bRemeshSim"))
+	int32 IterationsSim = 10;
+
+	UPROPERTY(EditAnywhere, Category = "Sim Mesh", meta = (UIMin = "0.0", UIMax = "1.0", ClampMin = "0.0", ClampMax = "1.0", EditCondition = "bRemeshSim"))
+	double SmoothingSim = 0.25;
+
+	UPROPERTY(EditAnywhere, Category = "Render Mesh")
+	bool bRemeshRender = false;
+
+	UPROPERTY(EditAnywhere, Category = "Render Mesh", meta=(UIMin = "1", UIMax = "200", ClampMin = "1", EditCondition = "bRemeshRender"))
+	int32 TargetPercentRender = 100;
+
+	UPROPERTY(EditAnywhere, Category = "Render Mesh", meta = (UIMin = "0", UIMax = "20", ClampMin = "0", ClampMax = "100", EditCondition = "bRemeshRender"))
+	int32 IterationsRender = 10;
+
+	UPROPERTY(EditAnywhere, Category = "Render Mesh", meta = (UIMin = "0.0", UIMax = "1.0", ClampMin = "0.0", ClampMax = "1.0", EditCondition = "bRemeshRender"))
+	double SmoothingRender = 0.25;
+
+	FChaosClothAssetRemeshNode(const Dataflow::FNodeParameters& InParam, FGuid InGuid = FGuid::NewGuid());
+
+	virtual void Evaluate(Dataflow::FContext& Context, const FDataflowOutput* Out) const override;
+
+private:
+
+	void RemeshSimMesh(TSharedRef<FManagedArrayCollection> ClothCollection,
+		TSharedRef<FManagedArrayCollection> OutClothCollection) const;
+
+	void RemeshRenderPattern(TSharedRef<FManagedArrayCollection> ClothCollection,
+		int32 PatternIndex,
+		TSharedRef<FManagedArrayCollection> OutClothCollection) const;
+
+};
diff --git a/Engine/Plugins/Runtime/MeshModelingToolset/Source/ModelingOperators/Private/CleaningOps/RemeshMeshOp.cpp b/Engine/Plugins/Runtime/MeshModelingToolset/Source/ModelingOperators/Private/CleaningOps/RemeshMeshOp.cpp
--- a/Engine/Plugins/Runtime/MeshModelingToolset/Source/ModelingOperators/Private/CleaningOps/RemeshMeshOp.cpp
+++ b/Engine/Plugins/Runtime/MeshModelingToolset/Source/ModelingOperators/Private/CleaningOps/RemeshMeshOp.cpp
@@ -59,6 +59,11 @@ void FRemeshMeshOp::SetTransform(const FTransformSRT3d& Transform)
 	ResultTransform = Transform;
 }
 
+void FRemeshMeshOp::SetUserSpecifiedConstraints(const FMeshConstraints& Constraints)
+{
+	UserSpecifiedConstraints = Constraints;
+}
+
 void FRemeshMeshOp::CalculateResult(FProgressCancel* Progress)
 {
 	TRACE_CPUPROFILER_EVENT_SCOPE(RemeshMeshOp);
@@ -122,7 +127,8 @@ void FRemeshMeshOp::CalculateResult(FProgressCancel* Progress)
 
 	Remesher->DEBUG_CHECK_LEVEL = 0;
 
-	FMeshConstraints Constraints;
+	FMeshConstraints Constraints = UserSpecifiedConstraints;		// Initialize with externally supplied constraints
+
 	constexpr bool bAllowSeamSplits = true;
 	const bool bAllowSeamCollapse = !bPreserveSharpEdges;
 	const bool bAllowSeamSmoothing = !bPreserveSharpEdges;
diff --git a/Engine/Plugins/Runtime/MeshModelingToolset/Source/ModelingOperators/Public/CleaningOps/RemeshMeshOp.h b/Engine/Plugins/Runtime/MeshModelingToolset/Source/ModelingOperators/Public/CleaningOps/RemeshMeshOp.h
--- a/Engine/Plugins/Runtime/MeshModelingToolset/Source/ModelingOperators/Public/CleaningOps/RemeshMeshOp.h
+++ b/Engine/Plugins/Runtime/MeshModelingToolset/Source/ModelingOperators/Public/CleaningOps/RemeshMeshOp.h
@@ -114,6 +114,9 @@ public:
 	// End normal flow only
 
 
+	/// Add arbitrary mesh constraints besides the boundary/seam constraints that are computed by this op
+	void SetUserSpecifiedConstraints(const FMeshConstraints& Constraints);
+
 	void SetTransform(const FTransformSRT3d& Transform);
 
 	//
@@ -130,6 +133,10 @@ protected:
 
 	TUniquePtr<FRemesher> CreateRemesher(ERemeshType Type, FDynamicMesh3* TargetMesh);
 
+private:
+
+	FMeshConstraints UserSpecifiedConstraints;
+
 };
 
 } // end namespace UE::Geometry
