Cloth Editor: add a "ping pong" mode to the skeletal mesh preview animation playback

diff --git a/Engine/Plugins/Experimental/ChaosClothAssetEditor/Source/ChaosClothAssetEditor/Private/ChaosClothAsset/ClothEditorPreviewScene.cpp b/Engine/Plugins/Experimental/ChaosClothAssetEditor/Source/ChaosClothAssetEditor/Private/ChaosClothAsset/ClothEditorPreviewScene.cpp
--- a/Engine/Plugins/Experimental/ChaosClothAssetEditor/Source/ChaosClothAssetEditor/Private/ChaosClothAsset/ClothEditorPreviewScene.cpp
+++ b/Engine/Plugins/Experimental/ChaosClothAssetEditor/Source/ChaosClothAssetEditor/Private/ChaosClothAsset/ClothEditorPreviewScene.cpp
@@ -129,6 +129,8 @@ void UChaosClothPreviewSceneDescription::PostEditChangeProperty(FPropertyChanged
 	{
 		PreviewScene->SceneDescriptionPropertyChanged(PropertyChangedEvent.GetMemberPropertyName());
 	}
+
+	ClothPreviewSceneDescriptionChanged.Broadcast();
 }
 
 void UChaosClothPreviewSceneDescription::PostTransacted(const FTransactionObjectEvent& TransactionEvent)
diff --git a/Engine/Plugins/Experimental/ChaosClothAssetEditor/Source/ChaosClothAssetEditor/Private/ChaosClothAsset/ClothEditorPreviewScene.h b/Engine/Plugins/Experimental/ChaosClothAssetEditor/Source/ChaosClothAssetEditor/Private/ChaosClothAsset/ClothEditorPreviewScene.h
--- a/Engine/Plugins/Experimental/ChaosClothAssetEditor/Source/ChaosClothAssetEditor/Private/ChaosClothAsset/ClothEditorPreviewScene.h
+++ b/Engine/Plugins/Experimental/ChaosClothAssetEditor/Source/ChaosClothAssetEditor/Private/ChaosClothAsset/ClothEditorPreviewScene.h
@@ -19,6 +19,8 @@ namespace UE::Chaos::ClothAsset
 class FChaosClothPreviewScene;
 }
 
+DECLARE_EVENT(UChaosClothPreviewSceneDescription, FClothPreviewSceneDescriptionChanged)
+
 ///
 /// The UChaosClothPreviewSceneDescription is a description of the Preview scene contents, intended to be editable in an FAdvancedPreviewSettingsWidget
 /// 
@@ -28,6 +30,8 @@ class CHAOSCLOTHASSETEDITOR_API UChaosClothPreviewSceneDescription : public UObj
 public:
 	GENERATED_BODY()
 
+	FClothPreviewSceneDescriptionChanged ClothPreviewSceneDescriptionChanged;
+
 	UChaosClothPreviewSceneDescription()
 	{
 		SetFlags(RF_Transactional);
diff --git a/Engine/Plugins/Experimental/ChaosClothAssetEditor/Source/ChaosClothAssetEditor/Private/ChaosClothAsset/SClothAnimationScrubPanel.cpp b/Engine/Plugins/Experimental/ChaosClothAssetEditor/Source/ChaosClothAssetEditor/Private/ChaosClothAsset/SClothAnimationScrubPanel.cpp
--- a/Engine/Plugins/Experimental/ChaosClothAssetEditor/Source/ChaosClothAssetEditor/Private/ChaosClothAsset/SClothAnimationScrubPanel.cpp
+++ b/Engine/Plugins/Experimental/ChaosClothAssetEditor/Source/ChaosClothAssetEditor/Private/ChaosClothAsset/SClothAnimationScrubPanel.cpp
@@ -7,6 +7,8 @@
 #include "Animation/AnimSingleNodeInstance.h"
 #include "Animation/AnimSequence.h"
 #include "Animation/BlendSpace.h"
+#include "Widgets/Input/SButton.h"
+#include "Widgets/Images/SImage.h"
 
 #define LOCTEXT_NAMESPACE "ClothAnimationScrubPanel"
 
@@ -14,6 +16,18 @@ void SClothAnimationScrubPanel::Construct( const SClothAnimationScrubPanel::FArg
 {
 	PreviewSceneWeakPtr = InPreviewScene;
 
+	// Skip adding the the Loop button so we can add our own
+	TArray<FTransportControlWidget> TransportControlWidgets;
+	for (const ETransportControlWidgetType Type : TEnumRange<ETransportControlWidgetType>())
+	{
+		if ((Type != ETransportControlWidgetType::Custom) && (Type != ETransportControlWidgetType::Loop))
+		{
+			TransportControlWidgets.Add(FTransportControlWidget(Type));
+		}
+	}
+	const FTransportControlWidget NewWidget(FOnMakeTransportWidget::CreateSP(this, &SClothAnimationScrubPanel::OnCreatePreviewPlaybackModeWidget));
+	TransportControlWidgets.Add(NewWidget);
+
 	this->ChildSlot
 	[
 		SNew(SHorizontalBox)
@@ -38,16 +52,82 @@ void SClothAnimationScrubPanel::Construct( const SClothAnimationScrubPanel::FArg
 			.OnClickedBackwardPlay(this, &SClothAnimationScrubPanel::OnClick_Backward)
 			.OnClickedBackwardStep(this, &SClothAnimationScrubPanel::OnClick_Backward_Step)
 			.OnClickedBackwardEnd(this, &SClothAnimationScrubPanel::OnClick_Backward_End)
-			.OnClickedToggleLoop(this, &SClothAnimationScrubPanel::OnClick_ToggleLoop)
-			.OnGetLooping(this, &SClothAnimationScrubPanel::IsLoopStatusOn)
+			.OnTickPlayback(this, &SClothAnimationScrubPanel::OnTickPlayback)
 			.OnGetPlaybackMode(this, &SClothAnimationScrubPanel::GetPlaybackMode)
 			.ViewInputMin(InArgs._ViewInputMin)
 			.ViewInputMax(InArgs._ViewInputMax)
 			.bDisplayAnimScrubBarEditing(false)
 			.bAllowZoom(false)
 			.IsRealtimeStreamingMode(false)
+			.TransportControlWidgetsToCreate(TransportControlWidgets)
 		]
 	];
+
+
+	if (const TSharedPtr<UE::Chaos::ClothAsset::FChaosClothPreviewScene> PinnedPreviewScene = PreviewSceneWeakPtr.Pin())
+	{
+		if (UChaosClothPreviewSceneDescription* const SceneDescription = PinnedPreviewScene->GetPreviewSceneDescription())
+		{
+			SceneDescription->ClothPreviewSceneDescriptionChanged.AddSP(this, &SClothAnimationScrubPanel::ApplyPlaybackSettings);
+		}
+	}
+
+	ApplyPlaybackSettings();
+}
+
+
+TSharedRef<SWidget> SClothAnimationScrubPanel::OnCreatePreviewPlaybackModeWidget()
+{
+	PreviewPlaybackModeButton = SNew(SButton)
+		.OnClicked(this, &SClothAnimationScrubPanel::OnClick_PreviewPlaybackMode)
+		.ButtonStyle( FAppStyle::Get(), "Animation.PlayControlsButton" )
+		.IsFocusable(false)
+		.ToolTipText_Lambda([&]()
+		{
+			if (PreviewPlaybackMode == EClothPreviewPlaybackMode::Default)
+			{
+				return LOCTEXT("PlaybackModeDefaultTooltip", "Linear playback");
+			}
+			else if (PreviewPlaybackMode == EClothPreviewPlaybackMode::Looping)
+			{
+				return LOCTEXT("PlaybackModeDefaultTooltip", "Looping playback");
+			}
+			else
+			{
+				return LOCTEXT("PlaybackModeDefaultTooltip", "Ping pong playback");
+			}
+		})
+		.ContentPadding(0.0f);
+
+	TWeakPtr<SButton> WeakButton = PreviewPlaybackModeButton;
+
+	PreviewPlaybackModeButton->SetContent(SNew(SImage)
+		.ColorAndOpacity(FSlateColor::UseSubduedForeground())
+		.Image_Lambda([&, WeakButton]()
+		{
+			if (PreviewPlaybackMode == EClothPreviewPlaybackMode::Default)
+			{
+				return FAppStyle::Get().GetBrush("Animation.Loop.Disabled");
+			}
+			else if (PreviewPlaybackMode == EClothPreviewPlaybackMode::Looping)
+			{
+				return FAppStyle::Get().GetBrush("Animation.Loop.Enabled");
+			}
+			else
+			{
+				return FAppStyle::Get().GetBrush("Animation.Loop.SelectionRange");		// TODO: Replace with a back and forth type icon
+			}
+		})
+	);
+
+	TSharedRef<SHorizontalBox> PreviewPlaybackModeBox = SNew(SHorizontalBox);
+	PreviewPlaybackModeBox->AddSlot()
+	.AutoWidth()
+	[
+		PreviewPlaybackModeButton.ToSharedRef()
+	];
+
+	return PreviewPlaybackModeBox;
 }
 
 FReply SClothAnimationScrubPanel::OnClick_Forward_Step()
@@ -160,20 +240,95 @@ FReply SClothAnimationScrubPanel::OnClick_Backward()
 	return FReply::Handled();
 }
 
-FReply SClothAnimationScrubPanel::OnClick_ToggleLoop()
+FReply SClothAnimationScrubPanel::OnClick_PreviewPlaybackMode()
 {
-	if (UAnimSingleNodeInstance* const PreviewInstance = GetPreviewAnimationInstance())
+	if (PreviewPlaybackMode == EClothPreviewPlaybackMode::Default)
+	{
+		PreviewPlaybackMode = EClothPreviewPlaybackMode::Looping;
+		
+		if (UAnimSingleNodeInstance* const PreviewInstance = GetPreviewAnimationInstance())
+		{
+			// If we paused due to hitting the end point, start playing again when entering loop mode
+			const float CurrentTime = PreviewInstance->GetCurrentTime();
+			const float PlayRate = PreviewInstance->GetPlayRate();
+			const float AssetPlayLength = PreviewInstance->CurrentAsset->GetPlayLength();
+			if (PlayRate < 0.0 && CurrentTime <= 0.0)
+			{
+				PreviewInstance->SetPlaying(true);
+			}
+			else if (PlayRate > 0.0 && CurrentTime >= AssetPlayLength)
+			{
+				PreviewInstance->SetPlaying(true);
+			}
+		}
+	}
+	else if (PreviewPlaybackMode == EClothPreviewPlaybackMode::Looping)
 	{
-		bool bIsLooping = PreviewInstance->IsLooping();
-		PreviewInstance->SetLooping(!bIsLooping);
+		PreviewPlaybackMode = EClothPreviewPlaybackMode::PingPong;
 	}
+	else
+	{
+		PreviewPlaybackMode = EClothPreviewPlaybackMode::Default;
+
+		if (UAnimSingleNodeInstance* const PreviewInstance = GetPreviewAnimationInstance())
+		{
+			// If we're switching to linear playback, set it to forward mode
+			PreviewInstance->SetReverse(false);
+		}
+	}
+
+	ApplyPlaybackSettings();
+
 	return FReply::Handled();
 }
 
-bool SClothAnimationScrubPanel::IsLoopStatusOn() const
+void SClothAnimationScrubPanel::ApplyPlaybackSettings()
 {
-	const UAnimSingleNodeInstance* const PreviewInstance = GetPreviewAnimationInstance();
-	return (PreviewInstance && PreviewInstance->IsLooping());
+	UAnimSingleNodeInstance* const PreviewInstance = GetPreviewAnimationInstance();
+
+	switch (PreviewPlaybackMode)
+	{
+	case EClothPreviewPlaybackMode::Default:
+		if (PreviewInstance)
+		{
+			PreviewInstance->SetLooping(false);
+		}
+		break;
+	case EClothPreviewPlaybackMode::Looping:
+		if (PreviewInstance)
+		{
+			PreviewInstance->SetLooping(true);
+		}
+		break;
+	case EClothPreviewPlaybackMode::PingPong:
+		if (PreviewInstance)
+		{
+			PreviewInstance->SetLooping(false);
+		}
+		break;
+	}
+}
+
+void SClothAnimationScrubPanel::OnTickPlayback(double InCurrentTime, float InDeltaTime)
+{
+	if (UAnimSingleNodeInstance* const PreviewInstance = GetPreviewAnimationInstance())
+	{
+		const float CurrentTime = PreviewInstance->GetCurrentTime();
+		const float PlayRate = PreviewInstance->GetPlayRate();
+		const float AssetPlayLength = PreviewInstance->CurrentAsset->GetPlayLength();
+
+		if (PreviewPlaybackMode == EClothPreviewPlaybackMode::PingPong)
+		{
+			if (PlayRate < 0.0 && CurrentTime <= 0.0)
+			{
+				PreviewInstance->SetReverse(!PreviewInstance->IsReverse());
+			}
+			else if (PlayRate > 0.0 && CurrentTime >= AssetPlayLength)
+			{
+				PreviewInstance->SetReverse(!PreviewInstance->IsReverse());
+			}
+		}
+	}
 }
 
 EPlaybackMode::Type SClothAnimationScrubPanel::GetPlaybackMode() const
diff --git a/Engine/Plugins/Experimental/ChaosClothAssetEditor/Source/ChaosClothAssetEditor/Private/ChaosClothAsset/SClothAnimationScrubPanel.h b/Engine/Plugins/Experimental/ChaosClothAssetEditor/Source/ChaosClothAssetEditor/Private/ChaosClothAsset/SClothAnimationScrubPanel.h
--- a/Engine/Plugins/Experimental/ChaosClothAssetEditor/Source/ChaosClothAssetEditor/Private/ChaosClothAsset/SClothAnimationScrubPanel.h
+++ b/Engine/Plugins/Experimental/ChaosClothAssetEditor/Source/ChaosClothAssetEditor/Private/ChaosClothAsset/SClothAnimationScrubPanel.h
@@ -6,6 +6,7 @@
 #include "ITransportControl.h"  // EPlaybackMode::Type
 
 class SScrubControlPanel;
+class SButton;
 class UAnimSingleNodeInstance;
 namespace UE::Chaos::ClothAsset
 {
@@ -19,6 +20,13 @@ class SClothAnimationScrubPanel : public SCompoundWidget
 {
 private:
 
+	enum class EClothPreviewPlaybackMode : int32
+	{
+		Default,
+		Looping,
+		PingPong
+	};
+	
 	SLATE_BEGIN_ARGS(SClothAnimationScrubPanel)	{}
 		SLATE_ATTRIBUTE( float, ViewInputMin )
 		SLATE_ATTRIBUTE( float, ViewInputMax )
@@ -26,6 +34,8 @@ private:
 
 	void Construct( const FArguments& InArgs, const TWeakPtr<UE::Chaos::ClothAsset::FChaosClothPreviewScene> InPreviewScene );
 
+	TSharedRef<SWidget> OnCreatePreviewPlaybackModeWidget();
+
 	// notifiers 
 	FReply OnClick_Forward_Step();
 	FReply OnClick_Forward_End();
@@ -33,13 +43,16 @@ private:
 	FReply OnClick_Backward_End();
 	FReply OnClick_Forward();
 	FReply OnClick_Backward();
-	FReply OnClick_ToggleLoop();
+	FReply OnClick_PreviewPlaybackMode();
+
+	void ApplyPlaybackSettings();
+
+	void OnTickPlayback(double InCurrentTime, float InDeltaTime);
 
 	void OnValueChanged(float NewValue);
 	void OnBeginSliderMovement();
 
 	EPlaybackMode::Type GetPlaybackMode() const;
-	bool IsLoopStatusOn() const;
 	float GetScrubValue() const;
 
 	UAnimSingleNodeInstance* GetPreviewAnimationInstance();
@@ -53,4 +66,7 @@ private:
 	TWeakPtr<UE::Chaos::ClothAsset::FChaosClothPreviewScene> PreviewSceneWeakPtr;
 
 	TSharedPtr<SScrubControlPanel> ScrubControlPanel;
+
+	TSharedPtr<SButton> PreviewPlaybackModeButton;
+	EClothPreviewPlaybackMode PreviewPlaybackMode = EClothPreviewPlaybackMode::Looping;
 };
diff --git a/Engine/Source/Editor/KismetWidgets/Private/SScrubControlPanel.cpp b/Engine/Source/Editor/KismetWidgets/Private/SScrubControlPanel.cpp
--- a/Engine/Source/Editor/KismetWidgets/Private/SScrubControlPanel.cpp
+++ b/Engine/Source/Editor/KismetWidgets/Private/SScrubControlPanel.cpp
@@ -41,7 +41,8 @@ void SScrubControlPanel::Construct( const SScrubControlPanel::FArguments& InArgs
 	TransportControlArgs.OnGetPlaybackMode = InArgs._OnGetPlaybackMode;
 	TransportControlArgs.OnGetRecording = InArgs._OnGetRecording;
 	TransportControlArgs.OnTickPlayback = InArgs._OnTickPlayback;
-	
+	TransportControlArgs.WidgetsToCreate = InArgs._TransportControlWidgetsToCreate.Get(TArray<FTransportControlWidget>());
+
 	FTransportControlArgs TransportControlArgsForRealtimeStreamingMode;
 	TransportControlArgsForRealtimeStreamingMode.OnForwardPlay = TransportControlArgs.OnForwardPlay;
 	TransportControlArgsForRealtimeStreamingMode.OnForwardStep = TransportControlArgs.OnForwardStep;
diff --git a/Engine/Source/Editor/KismetWidgets/Public/SScrubControlPanel.h b/Engine/Source/Editor/KismetWidgets/Public/SScrubControlPanel.h
--- a/Engine/Source/Editor/KismetWidgets/Public/SScrubControlPanel.h
+++ b/Engine/Source/Editor/KismetWidgets/Public/SScrubControlPanel.h
@@ -43,6 +43,7 @@ public:
 		, _OnSetInputViewRange()
 		, _OnCropAnimSequence()
 		, _OnTickPlayback()
+		, _TransportControlWidgetsToCreate()
 	{}
 		SLATE_ATTRIBUTE( float, Value )
 		SLATE_ATTRIBUTE( uint32, NumOfKeys )
@@ -89,6 +90,12 @@ public:
 		SLATE_EVENT( FOnScrubBarCommit, OnBarCommit )
 		/** Called each frame during playback */
 		SLATE_EVENT( FOnTickPlayback, OnTickPlayback )
+		/**
+		 * Array of custom widgets to create in the Transport Control widget. If this array is used the default widget ordering will be
+		 * ignored in favor of this set of widgets.
+		 */
+		SLATE_ARGUMENT(TOptional<TArray<FTransportControlWidget>>, TransportControlWidgetsToCreate);
+
 	SLATE_END_ARGS()
 
 	/**
