Fix duplicate LOCTEXT keys

diff --git a/Engine/Plugins/Experimental/ChaosClothAssetEditor/Source/ChaosClothAssetEditor/Private/ChaosClothAsset/SClothAnimationScrubPanel.cpp b/Engine/Plugins/Experimental/ChaosClothAssetEditor/Source/ChaosClothAssetEditor/Private/ChaosClothAsset/SClothAnimationScrubPanel.cpp
--- a/Engine/Plugins/Experimental/ChaosClothAssetEditor/Source/ChaosClothAssetEditor/Private/ChaosClothAsset/SClothAnimationScrubPanel.cpp
+++ b/Engine/Plugins/Experimental/ChaosClothAssetEditor/Source/ChaosClothAssetEditor/Private/ChaosClothAsset/SClothAnimationScrubPanel.cpp
@@ -90,11 +90,11 @@ TSharedRef<SWidget> SClothAnimationScrubPanel::OnCreatePreviewPlaybackModeWidget
 			}
 			else if (PreviewPlaybackMode == EClothPreviewPlaybackMode::Looping)
 			{
-				return LOCTEXT("PlaybackModeDefaultTooltip", "Looping playback");
+				return LOCTEXT("PlaybackModeLoopingTooltip", "Looping playback");
 			}
 			else
 			{
-				return LOCTEXT("PlaybackModeDefaultTooltip", "Ping pong playback");
+				return LOCTEXT("PlaybackModePingPongTooltip", "Ping pong playback");
 			}
 		})
 		.ContentPadding(0.0f);
