Cloth Remesh Node: regenerate long-range attachments and self-collision spheres after remeshing. Also remove any other Selections if we remesh.

diff --git a/Engine/Plugins/Experimental/ChaosClothAsset/Source/ChaosClothAsset/Private/ChaosClothAsset/ClothGeometryTools.cpp b/Engine/Plugins/Experimental/ChaosClothAsset/Source/ChaosClothAsset/Private/ChaosClothAsset/ClothGeometryTools.cpp
--- a/Engine/Plugins/Experimental/ChaosClothAsset/Source/ChaosClothAsset/Private/ChaosClothAsset/ClothGeometryTools.cpp
+++ b/Engine/Plugins/Experimental/ChaosClothAsset/Source/ChaosClothAsset/Private/ChaosClothAsset/ClothGeometryTools.cpp
@@ -1190,4 +1190,35 @@ namespace UE::Chaos::ClothAsset
 		}
 	}
 
+	void FClothGeometryTools::SampleVertices(const TConstArrayView<FVector3f> VertexPositions, float CullDiameterSq, TSet<int32>& OutVertexSet)
+	{
+		check(CullDiameterSq > 0.0f);
+
+		TArray<bool> VertexIsValid;
+		VertexIsValid.Init(true, VertexPositions.Num());
+
+		for (int32 Index = 0; Index < VertexPositions.Num(); ++Index)
+		{
+			if (!VertexIsValid[Index])
+			{
+				continue;
+			}
+			OutVertexSet.Add(Index);
+
+			const FVector3f& Pos0 = VertexPositions[Index];
+			for (int32 CompareIndex = Index + 1; CompareIndex < VertexPositions.Num(); ++CompareIndex)
+			{
+				if (!VertexIsValid[CompareIndex])
+				{
+					continue;
+				}
+				if (FVector3f::DistSquared(Pos0, VertexPositions[CompareIndex]) < CullDiameterSq)
+				{
+					VertexIsValid[CompareIndex] = false;
+				}
+			}
+		}
+	}
+
+
 }  // End namespace UE::Chaos::ClothAsset
diff --git a/Engine/Plugins/Experimental/ChaosClothAsset/Source/ChaosClothAsset/Private/ChaosClothAsset/CollectionClothSelectionFacade.cpp b/Engine/Plugins/Experimental/ChaosClothAsset/Source/ChaosClothAsset/Private/ChaosClothAsset/CollectionClothSelectionFacade.cpp
--- a/Engine/Plugins/Experimental/ChaosClothAsset/Source/ChaosClothAsset/Private/ChaosClothAsset/CollectionClothSelectionFacade.cpp
+++ b/Engine/Plugins/Experimental/ChaosClothAsset/Source/ChaosClothAsset/Private/ChaosClothAsset/CollectionClothSelectionFacade.cpp
@@ -129,6 +129,12 @@ namespace UE::Chaos::ClothAsset
 		return Selection ? Selection->GetData() : nullptr;
 	}
 
+	void FCollectionClothSelectionFacade::RemoveSelectionSet(const FName& Name)
+	{
+		check(IsValid());
+		ManagedArrayCollection->RemoveAttribute(Name, Private::SelectionGroup);
+	}
+
 	TSet<int32>& FCollectionClothSelectionFacade::FindOrAddSelectionSet(const FName& Name, const FName& GroupName)
 	{
 		check(IsValid());
diff --git a/Engine/Plugins/Experimental/ChaosClothAsset/Source/ChaosClothAsset/Public/ChaosClothAsset/ClothGeometryTools.h b/Engine/Plugins/Experimental/ChaosClothAsset/Source/ChaosClothAsset/Public/ChaosClothAsset/ClothGeometryTools.h
--- a/Engine/Plugins/Experimental/ChaosClothAsset/Source/ChaosClothAsset/Public/ChaosClothAsset/ClothGeometryTools.h
+++ b/Engine/Plugins/Experimental/ChaosClothAsset/Source/ChaosClothAsset/Public/ChaosClothAsset/ClothGeometryTools.h
@@ -99,5 +99,16 @@ namespace UE::Chaos::ClothAsset
 			const UE::Geometry::FDynamicMesh3& Mesh,
 			TArray<TArray<FIntVector2>>& Seams);
 
+
+		/**
+		 * Use Poisson disk sampling to get a set of evenly-spaced vertices
+		 * 
+		 * @param VertexPositions set of vertex points to sample from
+		 * @param CullDiameterSq squared minimum distance between samples
+		 * @param OutVertexSet indices of the sampled subset of VertexPositions
+		 */
+		static void SampleVertices(const TConstArrayView<FVector3f> VertexPositions, float CullDiameterSq, TSet<int32>& OutVertexSet);
+
+
 	};
 }  // End namespace UE::Chaos::ClothAsset
\ No newline at end of file
diff --git a/Engine/Plugins/Experimental/ChaosClothAsset/Source/ChaosClothAsset/Public/ChaosClothAsset/CollectionClothSelectionFacade.h b/Engine/Plugins/Experimental/ChaosClothAsset/Source/ChaosClothAsset/Public/ChaosClothAsset/CollectionClothSelectionFacade.h
--- a/Engine/Plugins/Experimental/ChaosClothAsset/Source/ChaosClothAsset/Public/ChaosClothAsset/CollectionClothSelectionFacade.h
+++ b/Engine/Plugins/Experimental/ChaosClothAsset/Source/ChaosClothAsset/Public/ChaosClothAsset/CollectionClothSelectionFacade.h
@@ -75,6 +75,9 @@ namespace UE::Chaos::ClothAsset
 		/** Find a selection with the given name, or nullptr if no such selection exists. */
 		CHAOSCLOTHASSET_API TSet<int32>* FindSelectionSet(const FName& Name);
 
+		/** Remove the selection with the given name if it exists. */
+		CHAOSCLOTHASSET_API void RemoveSelectionSet(const FName& Name);
+
 		/** 
 		 * Append all sets from an existing collection to this collection.
 		 * Matching sets (i.e., same name and type) will be unioned.
diff --git a/Engine/Plugins/Experimental/ChaosClothAssetEditor/Source/ChaosClothAssetDataflowNodes/Private/ChaosClothAsset/RemeshNode.cpp b/Engine/Plugins/Experimental/ChaosClothAssetEditor/Source/ChaosClothAssetDataflowNodes/Private/ChaosClothAsset/RemeshNode.cpp
--- a/Engine/Plugins/Experimental/ChaosClothAssetEditor/Source/ChaosClothAssetDataflowNodes/Private/ChaosClothAsset/RemeshNode.cpp
+++ b/Engine/Plugins/Experimental/ChaosClothAssetEditor/Source/ChaosClothAssetDataflowNodes/Private/ChaosClothAsset/RemeshNode.cpp
@@ -4,7 +4,10 @@
 
 #include "ChaosClothAsset/ClothPatternToDynamicMesh.h"
 #include "ChaosClothAsset/CollectionClothFacade.h"
+#include "ChaosClothAsset/CollectionClothSelectionFacade.h"
+#include "ChaosClothAsset/ClothCollectionGroup.h"
 #include "ChaosClothAsset/ClothGeometryTools.h"
+#include "ChaosClothAsset/ClothEngineTools.h"
 #include "Dataflow/DataflowInputOutput.h"
 #include "DynamicMesh/MeshTangents.h"
 #include "DynamicMesh/MeshNormals.h"
@@ -12,6 +15,8 @@
 #include "CleaningOps/RemeshMeshOp.h"
 #include "MeshUVChannelInfo.h"
 #include "MeshBoundaryLoops.h"
+#include "Chaos/CollectionPropertyFacade.h"
+#include "Algo/Find.h"
 
 #include UE_INLINE_GENERATED_CPP_BY_NAME(RemeshNode)
 
@@ -621,7 +626,6 @@ namespace UE::Chaos::ClothAsset::Private
 
 
 		// Set up constraints for Cloth Seam edges
-		// TODO: Add edge or vertex contraints for other cloth sim mesh data (e.g. collision particles, tethers)
 		UE::Geometry::FMeshConstraints Constraints;
 		for (const TArray<FIntVector2>& Seam : Seams)
 		{
@@ -704,8 +708,26 @@ FChaosClothAssetRemeshNode::FChaosClothAssetRemeshNode(const Dataflow::FNodePara
 }
 
 
-void FChaosClothAssetRemeshNode::RemeshSimMesh(TSharedRef<FManagedArrayCollection> ClothCollection,
-	TSharedRef<FManagedArrayCollection> OutClothCollection) const
+void FChaosClothAssetRemeshNode::EmptySimSelections(const TSharedRef<FManagedArrayCollection>& ClothCollection) const
+{
+	using namespace UE::Chaos::ClothAsset;
+
+	FCollectionClothSelectionFacade SelectionFacade(ClothCollection);
+
+	const TArray<FName> SelectionNames = SelectionFacade.GetNames();
+	for (const FName& SelectionName : SelectionNames)
+	{
+		const FName GroupName = SelectionFacade.GetSelectionGroup(SelectionName);
+		if (GroupName == ClothCollectionGroup::SimVertices3D || GroupName == ClothCollectionGroup::SimVertices2D || GroupName == ClothCollectionGroup::SimFaces)
+		{
+			TSet<int32>& SelectionSet = SelectionFacade.GetSelectionSet(SelectionName);
+			SelectionSet.Empty();
+		}
+	}
+}
+
+void FChaosClothAssetRemeshNode::RemeshSimMesh(const TSharedRef<const FManagedArrayCollection>& ClothCollection,
+	const TSharedRef<FManagedArrayCollection>& OutClothCollection) const
 {
 	using namespace UE::Geometry;
 	using namespace UE::Chaos::ClothAsset;
@@ -883,9 +905,7 @@ void FChaosClothAssetRemeshNode::RemeshSimMesh(TSharedRef<FManagedArrayCollectio
 	FClothGeometryTools::BuildSimMeshFromDynamicMeshes(OutClothCollection, Mesh2D, Mesh3D, PatternIndexLayerID, bTransferWeightMaps, bTransferSimSkinningData, bAppendToExistingMesh, DynamicMeshToClothVertexMap);
 
 
-	// Re-apply the seam info from the input sim mesh
-	// TODO: This will create a new Seam for each set of connected stitches. Should it instead be one seam for each pair of patterns?
-	//       E.g. if there are two strips of vertex pairs between a two given patterns, it will generate two seams. Should it just put them all in one seam?
+	// Re-apply the seam info from the input sim mesh. This will create a new Seam for each set of connected stitches
 
 	FCollectionClothFacade OutClothFacade(OutClothCollection);
 	for (int32 SeamIndex = 0; SeamIndex < Seams.Num(); ++SeamIndex)
@@ -912,9 +932,85 @@ void FChaosClothAssetRemeshNode::RemeshSimMesh(TSharedRef<FManagedArrayCollectio
 }
 
 
-void FChaosClothAssetRemeshNode::RemeshRenderPattern(TSharedRef<FManagedArrayCollection> ClothCollection,
+void FChaosClothAssetRemeshNode::RebuildTopologyDependentSimData(const TSharedRef<const FManagedArrayCollection>& InClothCollection,
+	const TSharedRef<FManagedArrayCollection>& OutClothCollection) const
+{
+	using namespace UE::Chaos::ClothAsset;
+	using namespace ::Chaos::Softs;
+
+	FCollectionClothConstFacade InClothFacade(InClothCollection);
+	FCollectionClothFacade OutClothFacade(OutClothCollection);
+
+	// Check that weight maps and skinning info have been interpolated over
+	for (const FName& InWeightMapName : InClothFacade.GetWeightMapNames())
+	{
+		const FName* FoundName = Algo::Find(OutClothFacade.GetWeightMapNames(), InWeightMapName);
+		checkf(FoundName, TEXT("Weight map %s was not copied to the output cloth collection"), *InWeightMapName.ToString());
+	}
+	if (InClothFacade.GetSimBoneIndices().Num() > 0 && OutClothFacade.GetNumSimVertices3D() > 0)
+	{
+		checkf(OutClothFacade.GetSimBoneIndices().Num() > 0, TEXT("Skinning bone indices not copied to the sim mesh of the output cloth collection"));
+	}
+	if (InClothFacade.GetSimBoneWeights().Num() > 0 && OutClothFacade.GetNumSimVertices3D() > 0)
+	{
+		checkf(OutClothFacade.GetSimBoneWeights().Num() > 0, TEXT("Skinning bone weights not copied to the sim mesh of the output cloth collection"));
+	}
+
+	Chaos::Softs::FCollectionPropertyConstFacade InProperties(InClothCollection);
+
+	// Reconstruct collision spheres
+	if (InProperties.GetKeyIndex(TEXT("SelfCollisionSphereStiffness")) != INDEX_NONE)
+	{
+		const float SelfCollisionSphereRadius = InProperties.GetValue<float>(TEXT("SelfCollisionSphereRadius"));
+		const float SelfCollisionSphereRadiusCullMultiplier = InProperties.GetValue<float>(TEXT("SelfCollisionSphereRadiusCullMultiplier"));
+		const float CullDiameterSq = FMath::Square(SelfCollisionSphereRadius * SelfCollisionSphereRadiusCullMultiplier * 2.f);
+
+		if (OutClothFacade.IsValid() && CullDiameterSq > 0.f)
+		{
+			TConstArrayView<FVector3f> SimPositions = OutClothFacade.GetSimPosition3D();
+			TSet<int32> VertexSet;
+			FClothGeometryTools::SampleVertices(SimPositions, CullDiameterSq, VertexSet);
+
+			FCollectionClothSelectionFacade Selection(OutClothCollection);
+			Selection.DefineSchema();
+
+			static const FName SelectionSetName(TEXT("_SelfCollisionSpheres"));
+			Selection.FindOrAddSelectionSet(SelectionSetName, ClothCollectionGroup::SimVertices3D) = VertexSet;
+		}
+	}
+
+	// Reconstruct long-range attachments
+	if (InProperties.GetKeyIndex(TEXT("TetherStiffness")) != INDEX_NONE)
+	{
+		const bool bUseGeodesicTethers = InProperties.GetValue<bool>(TEXT("UseGeodesicTethers"));
+		const FName FixedEndWeightMap(InProperties.GetStringValue(TEXT("FixedEndWeightMap")));
+
+		UE::Chaos::ClothAsset::FClothEngineTools::GenerateTethers(OutClothCollection, FixedEndWeightMap, bUseGeodesicTethers);
+	}
+}
+
+
+void FChaosClothAssetRemeshNode::EmptyRenderSelections(const TSharedRef<FManagedArrayCollection>& ClothCollection) const
+{
+	using namespace UE::Chaos::ClothAsset;
+
+	FCollectionClothSelectionFacade SelectionFacade(ClothCollection);
+
+	const TArray<FName> SelectionNames = SelectionFacade.GetNames();
+	for (const FName& SelectionName : SelectionNames)
+	{
+		const FName GroupName = SelectionFacade.GetSelectionGroup(SelectionName);
+		if (GroupName == ClothCollectionGroup::RenderVertices || GroupName == ClothCollectionGroup::RenderFaces)
+		{
+			TSet<int32>& SelectionSet = SelectionFacade.GetSelectionSet(SelectionName);
+			SelectionSet.Empty();
+		}
+	}
+}
+
+void FChaosClothAssetRemeshNode::RemeshRenderPattern(const TSharedRef<const FManagedArrayCollection>& ClothCollection,
 	int32 PatternIndex,
-	TSharedRef<FManagedArrayCollection> OutClothCollection) const
+	const TSharedRef<FManagedArrayCollection>& OutClothCollection) const
 {
 	using namespace UE::Geometry;
 	using namespace UE::Chaos::ClothAsset;
@@ -1151,12 +1247,15 @@ void FChaosClothAssetRemeshNode::Evaluate(Dataflow::FContext& Context, const FDa
 		{
 			if (bRemeshSim)
 			{
+				EmptySimSelections(OutputClothCollection);
 				RemeshSimMesh(ClothCollection, OutputClothCollection);
+				RebuildTopologyDependentSimData(ClothCollection, OutputClothCollection);
 			}
 
 			if (bRemeshRender)
 			{
-				// TODO: Does it make more sense to do the whole thing at once, or pattern by pattern? It's easier to rebuild the patterns if we do it this way
+				EmptyRenderSelections(OutputClothCollection);
+
 				const int32 NumPatterns = ClothFacade.GetNumRenderPatterns();
 				for (int32 PatternIndex = 0; PatternIndex < NumPatterns; ++PatternIndex)
 				{
diff --git a/Engine/Plugins/Experimental/ChaosClothAssetEditor/Source/ChaosClothAssetDataflowNodes/Private/ChaosClothAsset/SimulationLongRangeAttachmentConfigNode.cpp b/Engine/Plugins/Experimental/ChaosClothAssetEditor/Source/ChaosClothAssetDataflowNodes/Private/ChaosClothAsset/SimulationLongRangeAttachmentConfigNode.cpp
--- a/Engine/Plugins/Experimental/ChaosClothAssetEditor/Source/ChaosClothAssetDataflowNodes/Private/ChaosClothAsset/SimulationLongRangeAttachmentConfigNode.cpp
+++ b/Engine/Plugins/Experimental/ChaosClothAssetEditor/Source/ChaosClothAssetDataflowNodes/Private/ChaosClothAsset/SimulationLongRangeAttachmentConfigNode.cpp
@@ -22,6 +22,7 @@ void FChaosClothAssetSimulationLongRangeAttachmentConfigNode::AddProperties(Data
 	UE_CHAOS_CLOTHASSET_SIMULATIONCONFIG_SETPROPERTYWEIGHTED(TetherStiffness);
 	UE_CHAOS_CLOTHASSET_SIMULATIONCONFIG_SETPROPERTYWEIGHTED(TetherScale);
 	UE_CHAOS_CLOTHASSET_SIMULATIONCONFIG_SETPROPERTYBOOL(UseGeodesicTethers);
+	UE_CHAOS_CLOTHASSET_SIMULATIONCONFIG_SETPROPERTYSTRING(FixedEndWeightMap);
 }
 
 void FChaosClothAssetSimulationLongRangeAttachmentConfigNode::EvaluateClothCollection(Dataflow::FContext& Context, const TSharedRef<FManagedArrayCollection>& ClothCollection) const
diff --git a/Engine/Plugins/Experimental/ChaosClothAssetEditor/Source/ChaosClothAssetDataflowNodes/Private/ChaosClothAsset/SimulationSelfCollisionSpheresConfigNode.cpp b/Engine/Plugins/Experimental/ChaosClothAssetEditor/Source/ChaosClothAssetDataflowNodes/Private/ChaosClothAsset/SimulationSelfCollisionSpheresConfigNode.cpp
--- a/Engine/Plugins/Experimental/ChaosClothAssetEditor/Source/ChaosClothAssetDataflowNodes/Private/ChaosClothAsset/SimulationSelfCollisionSpheresConfigNode.cpp
+++ b/Engine/Plugins/Experimental/ChaosClothAssetEditor/Source/ChaosClothAssetDataflowNodes/Private/ChaosClothAsset/SimulationSelfCollisionSpheresConfigNode.cpp
@@ -7,6 +7,7 @@
 #include "ChaosClothAsset/CollectionClothFacade.h"
 #include "ChaosClothAsset/CollectionClothSelectionFacade.h"
 #include "ChaosClothAsset/ClothCollectionGroup.h"
+#include "ChaosClothAsset/ClothGeometryTools.h"
 
 #include UE_INLINE_GENERATED_CPP_BY_NAME(SimulationSelfCollisionSpheresConfigNode)
 
@@ -19,6 +20,7 @@ FChaosClothAssetSimulationSelfCollisionSpheresConfigNode::FChaosClothAssetSimula
 void FChaosClothAssetSimulationSelfCollisionSpheresConfigNode::AddProperties(Dataflow::FContext& Context, ::Chaos::Softs::FCollectionPropertyMutableFacade& Properties) const
 {
 	UE_CHAOS_CLOTHASSET_SIMULATIONCONFIG_SETPROPERTY(SelfCollisionSphereRadius);
+	UE_CHAOS_CLOTHASSET_SIMULATIONCONFIG_SETPROPERTY(SelfCollisionSphereRadiusCullMultiplier);
 	UE_CHAOS_CLOTHASSET_SIMULATIONCONFIG_SETPROPERTY(SelfCollisionSphereStiffness); 
 }
 
@@ -30,36 +32,14 @@ void FChaosClothAssetSimulationSelfCollisionSpheresConfigNode::EvaluateClothColl
 	if (Cloth.IsValid() && CullDiameterSq > 0.f)
 	{
 		TConstArrayView<FVector3f> SimPositions = Cloth.GetSimPosition3D();
-		TArray<bool> VertexIsValid;
-		VertexIsValid.Init(true, SimPositions.Num());
-
 		TSet<int32> VertexSet;
-		for (int32 Index = 0; Index < SimPositions.Num(); ++Index)
-		{
-			if (!VertexIsValid[Index])
-			{
-				continue;
-			}
-			VertexSet.Add(Index);
 
-			const FVector3f& Pos0 = SimPositions[Index];
-			for (int32 CompareIndex = Index + 1; CompareIndex < SimPositions.Num(); ++CompareIndex)
-			{
-				if (!VertexIsValid[CompareIndex])
-				{
-					continue;
-				}
-				if (FVector3f::DistSquared(Pos0, SimPositions[CompareIndex]) < CullDiameterSq)
-				{
-					VertexIsValid[CompareIndex] = false;
-				}
-			}
-		}
+		FClothGeometryTools::SampleVertices(SimPositions, CullDiameterSq, VertexSet);
 
 		FCollectionClothSelectionFacade Selection(ClothCollection);
 		Selection.DefineSchema();
 
-		static const FName SelectionSetName("_SelfCollisionSpheres");
+		static const FName SelectionSetName(TEXT("_SelfCollisionSpheres"));
 		Selection.FindOrAddSelectionSet(SelectionSetName, ClothCollectionGroup::SimVertices3D) = VertexSet;
 	}
 
diff --git a/Engine/Plugins/Experimental/ChaosClothAssetEditor/Source/ChaosClothAssetDataflowNodes/Public/ChaosClothAsset/RemeshNode.h b/Engine/Plugins/Experimental/ChaosClothAssetEditor/Source/ChaosClothAssetDataflowNodes/Public/ChaosClothAsset/RemeshNode.h
--- a/Engine/Plugins/Experimental/ChaosClothAssetEditor/Source/ChaosClothAssetDataflowNodes/Public/ChaosClothAsset/RemeshNode.h
+++ b/Engine/Plugins/Experimental/ChaosClothAssetEditor/Source/ChaosClothAssetDataflowNodes/Public/ChaosClothAsset/RemeshNode.h
@@ -7,7 +7,7 @@
 #include "RemeshNode.generated.h"
 
 /** Remesh the cloth surface(s) to get the specified mesh resolution(s).
- *  NOTE: Remeshing will invalidate vertex data such as Self Collision Spheres and Long Range Attachment Constraints, however weight maps and skinning data will be resampled to the new mesh topology.
+ *  NOTE: Weight Maps, Skinning Data, Self Collision Spheres, and Long Range Attachment Constraints will be reconstructed on the output mesh, however all other Selections will be removed
  */
 USTRUCT(Meta = (DataflowCloth))
 struct FChaosClothAssetRemeshNode : public FDataflowNode
@@ -50,11 +50,18 @@ public:
 
 private:
 
-	void RemeshSimMesh(TSharedRef<FManagedArrayCollection> ClothCollection,
-		TSharedRef<FManagedArrayCollection> OutClothCollection) const;
+	void EmptySimSelections(const TSharedRef<FManagedArrayCollection>& ClothCollection) const;
 
-	void RemeshRenderPattern(TSharedRef<FManagedArrayCollection> ClothCollection,
+	void RemeshSimMesh(const TSharedRef<const FManagedArrayCollection>& ClothCollection,
+		const TSharedRef<FManagedArrayCollection>& OutClothCollection) const;
+
+	void RebuildTopologyDependentSimData(const TSharedRef<const FManagedArrayCollection>& ClothCollection,
+		const TSharedRef<FManagedArrayCollection>& OutClothCollection) const;
+
+	void EmptyRenderSelections(const TSharedRef<FManagedArrayCollection>& ClothCollection) const;
+
+	void RemeshRenderPattern(const TSharedRef<const FManagedArrayCollection>& ClothCollection,
 		int32 PatternIndex,
-		TSharedRef<FManagedArrayCollection> OutClothCollection) const;
+		const TSharedRef<FManagedArrayCollection>& OutClothCollection) const;
 
 };
