Cloth Editor Preview Scene: manually register skeletal and cloth components instead of calling RegisterAllComponents

diff --git a/Engine/Plugins/Experimental/ChaosClothAssetEditor/Source/ChaosClothAssetEditor/Private/ChaosClothAsset/ClothEditorPreviewScene.cpp b/Engine/Plugins/Experimental/ChaosClothAssetEditor/Source/ChaosClothAssetEditor/Private/ChaosClothAsset/ClothEditorPreviewScene.cpp
--- a/Engine/Plugins/Experimental/ChaosClothAssetEditor/Source/ChaosClothAssetEditor/Private/ChaosClothAsset/ClothEditorPreviewScene.cpp
+++ b/Engine/Plugins/Experimental/ChaosClothAssetEditor/Source/ChaosClothAssetEditor/Private/ChaosClothAsset/ClothEditorPreviewScene.cpp
@@ -162,11 +162,12 @@ FChaosClothPreviewScene::FChaosClothPreviewScene(FPreviewScene::ConstructionValu
 	SkeletalMeshComponent = NewObject<USkeletalMeshComponent>(SceneActor);
 	SkeletalMeshComponent->SelectionOverrideDelegate = UPrimitiveComponent::FSelectionOverride::CreateRaw(this, &FChaosClothPreviewScene::IsComponentSelected);
 	SkeletalMeshComponent->SetDisablePostProcessBlueprint(true);
+	SkeletalMeshComponent->RegisterComponentWithWorld(GetWorld());
 
 	ClothComponent = NewObject<UChaosClothComponent>(SceneActor);
 	ClothComponent->SelectionOverrideDelegate = UPrimitiveComponent::FSelectionOverride::CreateRaw(this, &FChaosClothPreviewScene::IsComponentSelected);
-
-	SceneActor->RegisterAllComponents();
+	ClothComponent->RegisterComponentWithWorld(GetWorld());
+	
 }
 
 FChaosClothPreviewScene::~FChaosClothPreviewScene()
