Fix bad access on editor shutdown. The UObject system cannot be used after a certain point.

diff --git a/Engine/Plugins/Experimental/ChaosClothAssetEditor/Source/ChaosClothAssetDataflowNodes/Private/ChaosClothAsset/ChaosClothAssetDataflowNodesModule.cpp b/Engine/Plugins/Experimental/ChaosClothAssetEditor/Source/ChaosClothAssetDataflowNodes/Private/ChaosClothAsset/ChaosClothAssetDataflowNodesModule.cpp
--- a/Engine/Plugins/Experimental/ChaosClothAssetEditor/Source/ChaosClothAssetDataflowNodes/Private/ChaosClothAsset/ChaosClothAssetDataflowNodesModule.cpp
+++ b/Engine/Plugins/Experimental/ChaosClothAssetEditor/Source/ChaosClothAssetDataflowNodes/Private/ChaosClothAsset/ChaosClothAssetDataflowNodesModule.cpp
@@ -132,12 +132,15 @@ namespace UE::Chaos::ClothAsset
 		virtual void ShutdownModule() override
 		{
 			// Unregister type customizations
-			if (FPropertyEditorModule* const PropertyModule = FModuleManager::GetModulePtr<FPropertyEditorModule>("PropertyEditor"))
+			if (UObjectInitialized() && !IsEngineExitRequested())
 			{
-				PropertyModule->UnregisterCustomPropertyTypeLayout(FChaosClothAssetWeightedValue::StaticStruct()->GetFName());
-				PropertyModule->UnregisterCustomPropertyTypeLayout(FChaosClothAssetWeightedValueNonAnimatable::StaticStruct()->GetFName());
-				PropertyModule->UnregisterCustomPropertyTypeLayout(FChaosClothAssetWeightedValueNonAnimatableNoLowHighRange::StaticStruct()->GetFName());
-				PropertyModule->UnregisterCustomPropertyTypeLayout(FChaosClothAssetNodeSelectionGroup::StaticStruct()->GetFName());
+				if (FPropertyEditorModule* const PropertyModule = FModuleManager::GetModulePtr<FPropertyEditorModule>("PropertyEditor"))
+				{
+					PropertyModule->UnregisterCustomPropertyTypeLayout(FChaosClothAssetWeightedValue::StaticStruct()->GetFName());
+					PropertyModule->UnregisterCustomPropertyTypeLayout(FChaosClothAssetWeightedValueNonAnimatable::StaticStruct()->GetFName());
+					PropertyModule->UnregisterCustomPropertyTypeLayout(FChaosClothAssetWeightedValueNonAnimatableNoLowHighRange::StaticStruct()->GetFName());
+					PropertyModule->UnregisterCustomPropertyTypeLayout(FChaosClothAssetNodeSelectionGroup::StaticStruct()->GetFName());
+				}
 			}
 
 			// Unregister modular features
