Chaos Cloth: Add TransformPositions dataflow node.

diff --git a/Engine/Plugins/Experimental/ChaosClothAssetEditor/Source/ChaosClothAssetDataflowNodes/Private/ChaosClothAsset/ChaosClothAssetDataflowNodesModule.cpp b/Engine/Plugins/Experimental/ChaosClothAssetEditor/Source/ChaosClothAssetDataflowNodes/Private/ChaosClothAsset/ChaosClothAssetDataflowNodesModule.cpp
--- a/Engine/Plugins/Experimental/ChaosClothAssetEditor/Source/ChaosClothAssetDataflowNodes/Private/ChaosClothAsset/ChaosClothAssetDataflowNodesModule.cpp
+++ b/Engine/Plugins/Experimental/ChaosClothAssetEditor/Source/ChaosClothAssetDataflowNodes/Private/ChaosClothAsset/ChaosClothAssetDataflowNodesModule.cpp
@@ -45,6 +45,7 @@
 #include "ChaosClothAsset/StaticMeshImportNode.h"
 #include "ChaosClothAsset/TerminalNode.h"
 #include "ChaosClothAsset/TransferSkinWeightsNode.h"
+#include "ChaosClothAsset/TransformPositionsNode.h"
 #include "ChaosClothAsset/TransformUVsNode.h"
 #include "ChaosClothAsset/USDImportNode.h"
 #include "ChaosClothAsset/WeightedValueCustomization.h"
@@ -100,6 +101,7 @@ namespace UE::Chaos::ClothAsset
 			DATAFLOW_NODE_REGISTER_CREATION_FACTORY(FChaosClothAssetStaticMeshImportNode);
 			DATAFLOW_NODE_REGISTER_CREATION_FACTORY(FChaosClothAssetTerminalNode);
 			DATAFLOW_NODE_REGISTER_CREATION_FACTORY(FChaosClothAssetTransferSkinWeightsNode);
+			DATAFLOW_NODE_REGISTER_CREATION_FACTORY(FChaosClothAssetTransformPositionsNode);
 			DATAFLOW_NODE_REGISTER_CREATION_FACTORY(FChaosClothAssetTransformUVsNode);
 			DATAFLOW_NODE_REGISTER_CREATION_FACTORY(FChaosClothAssetUSDImportNode);
 			DATAFLOW_NODE_REGISTER_CREATION_FACTORY_NODE_COLORS_BY_CATEGORY("Cloth", FColorScheme::NodeHeader, FColorScheme::NodeBody);
diff --git a/Engine/Plugins/Experimental/ChaosClothAssetEditor/Source/ChaosClothAssetDataflowNodes/Private/ChaosClothAsset/TransformPositionsNode.cpp b/Engine/Plugins/Experimental/ChaosClothAssetEditor/Source/ChaosClothAssetDataflowNodes/Private/ChaosClothAsset/TransformPositionsNode.cpp
new file mode 100644
--- /dev/null
+++ b/Engine/Plugins/Experimental/ChaosClothAssetEditor/Source/ChaosClothAssetDataflowNodes/Private/ChaosClothAsset/TransformPositionsNode.cpp
@@ -0,0 +1,106 @@
+// Copyright Epic Games, Inc. All Rights Reserved.
+
+#include "ChaosClothAsset/TransformPositionsNode.h"
+#include "ChaosClothAsset/CollectionClothFacade.h"
+#include "Dataflow/DataflowInputOutput.h"
+
+#include UE_INLINE_GENERATED_CPP_BY_NAME(TransformPositionsNode)
+
+#define LOCTEXT_NAMESPACE "ChaosClothAssetTransformPositionsNode"
+
+FChaosClothAssetTransformPositionsNode::FChaosClothAssetTransformPositionsNode(const Dataflow::FNodeParameters& InParam, FGuid InGuid)
+	: FDataflowNode(InParam, InGuid)
+{
+	RegisterInputConnection(&Collection);
+	RegisterOutputConnection(&Collection, &Collection);
+}
+
+void FChaosClothAssetTransformPositionsNode::Evaluate(Dataflow::FContext& Context, const FDataflowOutput* Out) const
+{
+	if (Out->IsA<FManagedArrayCollection>(&Collection))
+	{
+		using namespace UE::Chaos::ClothAsset;
+
+		// Evaluate in collection
+		FManagedArrayCollection InCollection = GetValue<FManagedArrayCollection>(Context, &Collection);
+		const TSharedRef<FManagedArrayCollection> ClothCollection = MakeShared<FManagedArrayCollection>(MoveTemp(InCollection));
+
+		// Always check for a valid cloth collection/facade to avoid processing non cloth collections
+		FCollectionClothFacade ClothFacade(ClothCollection);
+		if (ClothFacade.IsValid())
+		{
+			if (bTransform2DSimPositions)
+			{
+				auto Transform2D = [this](TArrayView<FVector2f>&& SimPositions2D)
+				{
+					for (FVector2f& SimPosition : SimPositions2D)
+					{
+						// Scale
+						SimPosition = SimPosition * Sim2DScale;
+						// Rotate
+						float SinAngle, CosAngle;
+						FMath::SinCos(&SinAngle, &CosAngle, FMath::DegreesToRadians(Sim2DRotation));
+						SimPosition = FVector2f(
+							SimPosition.X * CosAngle - SimPosition.Y * SinAngle,
+							SimPosition.Y * CosAngle + SimPosition.X * SinAngle);  // Clockwise rotation
+						// Translate
+						SimPosition += FVector2f(Sim2DTranslation.X, Sim2DTranslation.Y);
+					}
+				};
+
+				if (Sim2DPattern == INDEX_NONE)
+				{
+					Transform2D(ClothFacade.GetSimPosition2D());
+				}
+				else
+				{
+					const int32 NumPatterns = ClothFacade.GetNumSimPatterns();
+					if (Sim2DPattern >= 0 && Sim2DPattern < NumPatterns)
+					{
+						FCollectionClothSimPatternFacade ClothPatternFacade = ClothFacade.GetSimPattern(Sim2DPattern);
+						Transform2D(ClothPatternFacade.GetSimPosition2D());
+					}
+				}
+			}
+			if (bTransform3DSimPositions)
+			{
+				const FTransform3f Transform(FRotator3f(Sim3DRotation.X, Sim3DRotation.Y, Sim3DRotation.Z), Sim3DTranslation, Sim3DScale);
+
+				TArrayView<FVector3f> SimPositions3D = ClothFacade.GetSimPosition3D();
+				for (FVector3f& SimPosition : SimPositions3D)
+				{
+					SimPosition = Transform.TransformPosition(SimPosition);
+				}
+			}
+			if (bTransformRenderPositions)
+			{
+				const FTransform3f Transform(FRotator3f(RenderRotation.X, RenderRotation.Y, RenderRotation.Z), RenderTranslation, RenderScale);
+				auto Transform3D = [this, &Transform](TArrayView<FVector3f>&& RenderPositions)
+				{
+					for (FVector3f& RenderPosition : RenderPositions)
+					{
+						RenderPosition = Transform.TransformPosition(RenderPosition);
+					}
+				};
+
+				if (RenderPattern == INDEX_NONE)
+				{
+					Transform3D(ClothFacade.GetRenderPosition());
+				}
+				else
+				{
+					const int32 NumPatterns = ClothFacade.GetNumRenderPatterns();
+					if (RenderPattern >= 0 && RenderPattern < NumPatterns)
+					{
+						FCollectionClothRenderPatternFacade RenderPatternFacade = ClothFacade.GetRenderPattern(RenderPattern);
+						Transform3D(RenderPatternFacade.GetRenderPosition());
+					}
+				}
+			}
+		}
+
+		SetValue(Context, MoveTemp(*ClothCollection), &Collection);
+	}
+}
+
+#undef LOCTEXT_NAMESPACE
diff --git a/Engine/Plugins/Experimental/ChaosClothAssetEditor/Source/ChaosClothAssetDataflowNodes/Public/ChaosClothAsset/TransformPositionsNode.h b/Engine/Plugins/Experimental/ChaosClothAssetEditor/Source/ChaosClothAssetDataflowNodes/Public/ChaosClothAsset/TransformPositionsNode.h
new file mode 100644
--- /dev/null
+++ b/Engine/Plugins/Experimental/ChaosClothAssetEditor/Source/ChaosClothAssetDataflowNodes/Public/ChaosClothAsset/TransformPositionsNode.h
@@ -0,0 +1,80 @@
+// Copyright Epic Games, Inc. All Rights Reserved.
+
+#pragma once
+
+#include "Dataflow/DataflowNode.h"
+#include "GeometryCollection/ManagedArrayCollection.h"
+#include "TransformPositionsNode.generated.h"
+
+USTRUCT(Meta = (DataflowCloth))
+struct FChaosClothAssetTransformPositionsNode : public FDataflowNode
+{
+	GENERATED_USTRUCT_BODY()
+	DATAFLOW_NODE_DEFINE_INTERNAL(FChaosClothAssetTransformPositionsNode, "TransformPositions", "Cloth", "Cloth Simulation Transform Positions")
+
+public:
+
+	UPROPERTY(Meta = (Dataflowinput, DataflowOutput, DataflowPassthrough = "Collection"))
+	FManagedArrayCollection Collection;
+
+	/** Enable Transforming 2D Sim Mesh Positions */
+	UPROPERTY(EditAnywhere, Category = "TransformPositions - 2D Sim")
+	bool bTransform2DSimPositions = false;
+
+	/** 2D Sim Transform scale. */
+	UPROPERTY(EditAnywhere, Category = "TransformPositions - 2D Sim", Meta = (AllowPreserveRatio, EditCondition = "bTransform2DSimPositions"))
+	FVector2f Sim2DScale = { 1.f, 1.f };
+
+	/** 2D Sim Transform rotation angle in degrees. */
+	UPROPERTY(EditAnywhere, Category = "TransformPositions - 2D Sim", Meta = (UIMin = -360, UIMax = 360, ClampMin = -360, ClampMax = 360, EditCondition = "bTransform2DSimPositions"))
+	float Sim2DRotation = 0.f;
+
+	/** 2D Sim Transform translation. */
+	UPROPERTY(EditAnywhere, Category = "TransformPositions - 2D Sim", Meta = (EditCondition = "bTransform2DSimPositions"))
+	FVector2f Sim2DTranslation = { 0.f, 0.f };
+
+	/** Sim Pattern to transform. All patterns will be used when set to -1. */
+	UPROPERTY(EditAnywhere, Category = "TransformPositions - 2D Sim", Meta = (UIMax = 10, ClampMin = -1, EditCondition = "bTransform2DSimPositions"))
+	int32 Sim2DPattern = INDEX_NONE;
+
+	/** Enable Transforming 3D Sim Mesh Positions */
+	UPROPERTY(EditAnywhere, Category = "TransformPositions - 3D Sim")
+	bool bTransform3DSimPositions = false;
+
+	/** 3D Sim Transform scale. */
+	UPROPERTY(EditAnywhere, Category = "TransformPositions - 3D Sim", Meta = (AllowPreserveRatio, EditCondition = "bTransform3DSimPositions"))
+	FVector3f Sim3DScale = { 1.f, 1.f, 1.f };
+
+	/** 3D Sim Transform rotation angle in degrees (Euler Angles) */
+	UPROPERTY(EditAnywhere, Category = "TransformPositions - 3D Sim", Meta = (UIMin = -360, UIMax = 360, ClampMin = -360, ClampMax = 360, EditCondition = "bTransform3DSimPositions"))
+	FVector3f Sim3DRotation = { 0.f, 0.f, 0.f };
+
+	/** 3D Sim Transform translation. */
+	UPROPERTY(EditAnywhere, Category = "TransformPositions - 3D Sim", Meta = (EditCondition = "bTransform3DSimPositions"))
+	FVector3f Sim3DTranslation = { 0.f, 0.f, 0.f };
+
+	/** Enable Transforming Render Positions */
+	UPROPERTY(EditAnywhere, Category = "TransformPositions - Render")
+	bool bTransformRenderPositions = false;
+
+	/** Render Transform scale. */
+	UPROPERTY(EditAnywhere, Category = "TransformPositions - Render", Meta = (AllowPreserveRatio, EditCondition = "bTransformRenderPositions"))
+	FVector3f RenderScale = { 1.f, 1.f, 1.f };
+
+	/** Render Transform rotation angle in degrees (Euler Angles) */
+	UPROPERTY(EditAnywhere, Category = "TransformPositions - Render", Meta = (UIMin = -360, UIMax = 360, ClampMin = -360, ClampMax = 360, EditCondition = "bTransformRenderPositions"))
+	FVector3f RenderRotation = { 0.f, 0.f, 0.f };
+
+	/** Render Transform translation. */
+	UPROPERTY(EditAnywhere, Category = "TransformPositions - Render", Meta = (EditCondition = "bTransformRenderPositions"))
+	FVector3f RenderTranslation = { 0.f, 0.f, 0.f };
+
+	/** Render Pattern to transform. All patterns will be used when set to -1. */
+	UPROPERTY(EditAnywhere, Category = "TransformPositions - Render", Meta = (UIMax = 10, ClampMin = -1, EditCondition = "bTransformRenderPositions"))
+	int32 RenderPattern = INDEX_NONE;
+
+	FChaosClothAssetTransformPositionsNode(const Dataflow::FNodeParameters& InParam, FGuid InGuid = FGuid::NewGuid());
+
+private:
+	virtual void Evaluate(Dataflow::FContext& Context, const FDataflowOutput* Out) const override;
+};
