Chaos Cloth Asset - Replaced the set property defines used to add properties in the simulation config nodes by a new FPropertyHelper struct API that allows Intrinsic properties and does additional types, flags, and names checks.

diff --git a/Engine/Plugins/Experimental/ChaosCloth/Source/Chaos/Private/ChaosCloth/ChaosClothingSimulationConfig.cpp b/Engine/Plugins/Experimental/ChaosCloth/Source/Chaos/Private/ChaosCloth/ChaosClothingSimulationConfig.cpp
--- a/Engine/Plugins/Experimental/ChaosCloth/Source/Chaos/Private/ChaosCloth/ChaosClothingSimulationConfig.cpp
+++ b/Engine/Plugins/Experimental/ChaosCloth/Source/Chaos/Private/ChaosCloth/ChaosClothingSimulationConfig.cpp
@@ -27,11 +27,12 @@ namespace Chaos
 
 	void FClothingSimulationConfig::Initialize(const UChaosClothConfig* ClothConfig, const UChaosClothSharedSimConfig* ClothSharedConfig, bool bUseLegacyConfig)
 	{
-		constexpr ::Chaos::Softs::ECollectionPropertyFlags NonAnimatablePropertyFlags =
-			::Chaos::Softs::ECollectionPropertyFlags::Enabled |
-			::Chaos::Softs::ECollectionPropertyFlags::Legacy;  // Indicates a property set from a pre-property collection config (e.g. that can be overriden in Dataflow without warning)
-		constexpr ::Chaos::Softs::ECollectionPropertyFlags AnimatablePropertyFlags = NonAnimatablePropertyFlags |
-			::Chaos::Softs::ECollectionPropertyFlags::Animatable;
+		using namespace ::Chaos::Softs;
+		constexpr ECollectionPropertyFlags NonAnimatablePropertyFlags =
+			ECollectionPropertyFlags::Enabled |
+			ECollectionPropertyFlags::Legacy;  // Indicates a property set from a pre-property collection config (e.g. that can be overriden in Dataflow without warning)
+		constexpr ECollectionPropertyFlags AnimatablePropertyFlags = NonAnimatablePropertyFlags |
+			ECollectionPropertyFlags::Animatable;
 
 		// Clear all properties
 		PropertyCollections.Reset(1);
@@ -69,9 +70,9 @@ namespace Chaos
 					MassValue = ClothConfig->Density;
 					break;
 				}
-				Property->AddValue(TEXT("MassMode"), (int32)ClothConfig->MassMode, NonAnimatablePropertyFlags);
-				Property->AddValue(TEXT("MassValue"), MassValue, NonAnimatablePropertyFlags);
-				Property->AddValue(TEXT("MinPerParticleMass"), ClothConfig->MinPerParticleMass, NonAnimatablePropertyFlags);
+				Property->AddValue(TEXT("MassMode"), (int32)ClothConfig->MassMode, NonAnimatablePropertyFlags | ECollectionPropertyFlags::Intrinsic);
+				Property->AddValue(TEXT("MassValue"), MassValue, NonAnimatablePropertyFlags | ECollectionPropertyFlags::Intrinsic);
+				Property->AddValue(TEXT("MinPerParticleMass"), ClothConfig->MinPerParticleMass, NonAnimatablePropertyFlags | ECollectionPropertyFlags::Intrinsic);
 			}
 
 			// Edge constraint
diff --git a/Engine/Plugins/Experimental/ChaosClothAssetEditor/Source/ChaosClothAssetDataflowNodes/Private/ChaosClothAsset/SimulationAerodynamicsConfigNode.cpp b/Engine/Plugins/Experimental/ChaosClothAssetEditor/Source/ChaosClothAssetDataflowNodes/Private/ChaosClothAsset/SimulationAerodynamicsConfigNode.cpp
--- a/Engine/Plugins/Experimental/ChaosClothAssetEditor/Source/ChaosClothAssetDataflowNodes/Private/ChaosClothAsset/SimulationAerodynamicsConfigNode.cpp
+++ b/Engine/Plugins/Experimental/ChaosClothAssetEditor/Source/ChaosClothAssetDataflowNodes/Private/ChaosClothAsset/SimulationAerodynamicsConfigNode.cpp
@@ -1,7 +1,6 @@
 // Copyright Epic Games, Inc. All Rights Reserved.
 
 #include "ChaosClothAsset/SimulationAerodynamicsConfigNode.h"
-#include "ChaosClothAsset/SimulationBaseConfigNodePrivate.h"
 #include "Chaos/CollectionPropertyFacade.h"
 
 #include UE_INLINE_GENERATED_CPP_BY_NAME(SimulationAerodynamicsConfigNode)
@@ -14,10 +13,10 @@ FChaosClothAssetSimulationAerodynamicsConfigNode::FChaosClothAssetSimulationAero
 	RegisterInputConnection(&Lift.WeightMap);
 }
 
-void FChaosClothAssetSimulationAerodynamicsConfigNode::AddProperties(Dataflow::FContext& Context, ::Chaos::Softs::FCollectionPropertyMutableFacade& Properties) const
+void FChaosClothAssetSimulationAerodynamicsConfigNode::AddProperties(FPropertyHelper& PropertyHelper) const
 {
-	UE_CHAOS_CLOTHASSET_SIMULATIONCONFIG_SETPROPERTY(FluidDensity);
-	UE_CHAOS_CLOTHASSET_SIMULATIONCONFIG_SETPROPERTYWEIGHTED(Drag);
-	UE_CHAOS_CLOTHASSET_SIMULATIONCONFIG_SETPROPERTYWEIGHTED(Lift);
-	UE_CHAOS_CLOTHASSET_SIMULATIONCONFIG_SETPROPERTY(WindVelocity);
+	PropertyHelper.SetProperty(this, &FluidDensity);
+	PropertyHelper.SetPropertyWeighted(this, &Drag);
+	PropertyHelper.SetPropertyWeighted(this, &Lift);
+	PropertyHelper.SetProperty(this, &WindVelocity);
 }
diff --git a/Engine/Plugins/Experimental/ChaosClothAssetEditor/Source/ChaosClothAssetDataflowNodes/Private/ChaosClothAsset/SimulationAnimDriveConfigNode.cpp b/Engine/Plugins/Experimental/ChaosClothAssetEditor/Source/ChaosClothAssetDataflowNodes/Private/ChaosClothAsset/SimulationAnimDriveConfigNode.cpp
--- a/Engine/Plugins/Experimental/ChaosClothAssetEditor/Source/ChaosClothAssetDataflowNodes/Private/ChaosClothAsset/SimulationAnimDriveConfigNode.cpp
+++ b/Engine/Plugins/Experimental/ChaosClothAssetEditor/Source/ChaosClothAssetDataflowNodes/Private/ChaosClothAsset/SimulationAnimDriveConfigNode.cpp
@@ -1,7 +1,6 @@
 // Copyright Epic Games, Inc. All Rights Reserved.
 
 #include "ChaosClothAsset/SimulationAnimDriveConfigNode.h"
-#include "ChaosClothAsset/SimulationBaseConfigNodePrivate.h"
 #include "Chaos/CollectionPropertyFacade.h"
 
 #include UE_INLINE_GENERATED_CPP_BY_NAME(SimulationAnimDriveConfigNode)
@@ -14,8 +13,8 @@ FChaosClothAssetSimulationAnimDriveConfigNode::FChaosClothAssetSimulationAnimDri
 	RegisterInputConnection(&AnimDriveDamping.WeightMap);
 }
 
-void FChaosClothAssetSimulationAnimDriveConfigNode::AddProperties(Dataflow::FContext& Context, ::Chaos::Softs::FCollectionPropertyMutableFacade& Properties) const
+void FChaosClothAssetSimulationAnimDriveConfigNode::AddProperties(FPropertyHelper& PropertyHelper) const
 {
-	UE_CHAOS_CLOTHASSET_SIMULATIONCONFIG_SETPROPERTYWEIGHTED(AnimDriveStiffness);
-	UE_CHAOS_CLOTHASSET_SIMULATIONCONFIG_SETPROPERTYWEIGHTED(AnimDriveDamping);
+	PropertyHelper.SetPropertyWeighted(this, &AnimDriveStiffness);
+	PropertyHelper.SetPropertyWeighted(this, &AnimDriveDamping);
 }
diff --git a/Engine/Plugins/Experimental/ChaosClothAssetEditor/Source/ChaosClothAssetDataflowNodes/Private/ChaosClothAsset/SimulationBackstopConfigNode.cpp b/Engine/Plugins/Experimental/ChaosClothAssetEditor/Source/ChaosClothAssetDataflowNodes/Private/ChaosClothAsset/SimulationBackstopConfigNode.cpp
--- a/Engine/Plugins/Experimental/ChaosClothAssetEditor/Source/ChaosClothAssetDataflowNodes/Private/ChaosClothAsset/SimulationBackstopConfigNode.cpp
+++ b/Engine/Plugins/Experimental/ChaosClothAssetEditor/Source/ChaosClothAssetDataflowNodes/Private/ChaosClothAsset/SimulationBackstopConfigNode.cpp
@@ -1,7 +1,6 @@
 // Copyright Epic Games, Inc. All Rights Reserved.
 
 #include "ChaosClothAsset/SimulationBackstopConfigNode.h"
-#include "ChaosClothAsset/SimulationBaseConfigNodePrivate.h"
 #include "Chaos/CollectionPropertyFacade.h"
 
 #include UE_INLINE_GENERATED_CPP_BY_NAME(SimulationBackstopConfigNode)
@@ -14,8 +13,8 @@ FChaosClothAssetSimulationBackstopConfigNode::FChaosClothAssetSimulationBackstop
 	RegisterInputConnection(&BackstopRadius.WeightMap);
 }
 
-void FChaosClothAssetSimulationBackstopConfigNode::AddProperties(Dataflow::FContext& Context, ::Chaos::Softs::FCollectionPropertyMutableFacade& Properties) const
+void FChaosClothAssetSimulationBackstopConfigNode::AddProperties(FPropertyHelper& PropertyHelper) const
 {
-	UE_CHAOS_CLOTHASSET_SIMULATIONCONFIG_SETPROPERTYWEIGHTED(BackstopDistance);
-	UE_CHAOS_CLOTHASSET_SIMULATIONCONFIG_SETPROPERTYWEIGHTED(BackstopRadius);
+	PropertyHelper.SetPropertyWeighted(this, &BackstopDistance);
+	PropertyHelper.SetPropertyWeighted(this, &BackstopRadius);
 }
diff --git a/Engine/Plugins/Experimental/ChaosClothAssetEditor/Source/ChaosClothAssetDataflowNodes/Private/ChaosClothAsset/SimulationBaseConfigNode.cpp b/Engine/Plugins/Experimental/ChaosClothAssetEditor/Source/ChaosClothAssetDataflowNodes/Private/ChaosClothAsset/SimulationBaseConfigNode.cpp
--- a/Engine/Plugins/Experimental/ChaosClothAssetEditor/Source/ChaosClothAssetDataflowNodes/Private/ChaosClothAsset/SimulationBaseConfigNode.cpp
+++ b/Engine/Plugins/Experimental/ChaosClothAssetEditor/Source/ChaosClothAssetDataflowNodes/Private/ChaosClothAsset/SimulationBaseConfigNode.cpp
@@ -67,7 +67,12 @@ void FChaosClothAssetSimulationBaseConfigNode::Evaluate(Dataflow::FContext& Cont
 		FCollectionPropertyMutableFacade Properties(ClothCollection);
 		Properties.DefineSchema();
 
-		AddProperties(Context, Properties);
+PRAGMA_DISABLE_DEPRECATION_WARNINGS
+		AddProperties(Context, Properties);  // Deprecated 5.4
+PRAGMA_ENABLE_DEPRECATION_WARNINGS
+
+		FPropertyHelper PropertyHelper(*this, Context, Properties);
+		AddProperties(PropertyHelper);
 
 		if (FCollectionClothFacade(ClothCollection).IsValid())  // Can only act on the collection if it is a valid cloth collection
 		{
@@ -81,8 +86,8 @@ void FChaosClothAssetSimulationBaseConfigNode::Evaluate(Dataflow::FContext& Cont
 int32 FChaosClothAssetSimulationBaseConfigNode::AddPropertyHelper(
 	::Chaos::Softs::FCollectionPropertyMutableFacade& Properties,
 	const FName& PropertyName,
-	bool bIsAnimatable,
-	const TArray<FName>& SimilarPropertyNames) const
+	const TArray<FName>& SimilarPropertyNames,
+	ECollectionPropertyFlags PropertyFlags) const
 {
 	using namespace UE::Chaos::ClothAsset;
 
@@ -98,9 +103,7 @@ int32 FChaosClothAssetSimulationBaseConfigNode::AddPropertyHelper(
 	// else don't warn and hope people know what they are doing
 
 	// Set/clear flags
-	Properties.SetEnabled(KeyIndex, true);  // Always enabled when then node is active
-	Properties.SetLegacy(KeyIndex, false);  // No longer a legacy property after beind set by this node
-	Properties.SetAnimatable(KeyIndex, bIsAnimatable);
+	Properties.SetFlags(KeyIndex, ECollectionPropertyFlags::Enabled | PropertyFlags); // Always enabled when then node is active, no longer a legacy property after being set by this node
 
 	// Check for similar properties
 	for (const FName& SimilarPropertyName : SimilarPropertyNames)
@@ -115,4 +118,87 @@ int32 FChaosClothAssetSimulationBaseConfigNode::AddPropertyHelper(
 	return KeyIndex;
 }
 
+FChaosClothAssetSimulationBaseConfigNode::FPropertyHelper::FPropertyHelper(const FChaosClothAssetSimulationBaseConfigNode& InConfigNode, Dataflow::FContext& InContext, ::Chaos::Softs::FCollectionPropertyMutableFacade& InProperties)
+	: ConfigNode(InConfigNode)
+	, Context(InContext)
+	, Properties(InProperties)
+{}
+
+int32 FChaosClothAssetSimulationBaseConfigNode::FPropertyHelper::SetPropertyBool(
+	const FName& PropertyName,
+	bool PropertyValue,
+	const TArray<FName>& SimilarPropertyNames,
+	ECollectionPropertyFlags PropertyFlags)
+{
+	checkf(!PropertyName.ToString().StartsWith(TEXT("b")), TEXT("Unlike its C++ counterpart the Boolean property name should not start with a 'b'."));
+	const int32 PropertyKeyIndex = ConfigNode.AddPropertyHelper(Properties, PropertyName, SimilarPropertyNames, PropertyFlags);
+	Properties.SetValue(PropertyKeyIndex, PropertyValue);
+	return PropertyKeyIndex;
+}
+
+int32 FChaosClothAssetSimulationBaseConfigNode::FPropertyHelper::SetPropertyString(
+	const FName& PropertyName,
+	const FString& PropertyValue,
+	const TArray<FName>& SimilarPropertyNames,
+	ECollectionPropertyFlags PropertyFlags)
+{
+	const int32 PropertyKeyIndex = ConfigNode.AddPropertyHelper(Properties, PropertyName, SimilarPropertyNames, PropertyFlags);
+	Properties.SetStringValue(PropertyKeyIndex, PropertyValue);
+	return PropertyKeyIndex;
+}
+
+int32 FChaosClothAssetSimulationBaseConfigNode::FPropertyHelper::SetPropertyWeighted(
+	const FName& PropertyName,
+	const FChaosClothAssetWeightedValue& PropertyValue,
+	const TArray<FName>& SimilarPropertyNames,
+	ECollectionPropertyFlags PropertyFlags)
+{
+	ensureMsgf(!EnumHasAnyFlags(PropertyFlags, ECollectionPropertyFlags::Animatable), TEXT("Animatable flag ignored. Weighted properties are set animatable through FChaosClothAssetWeightedValue::bIsAnimatable."));
+	if (PropertyValue.bIsAnimatable)
+	{
+		EnumAddFlags(PropertyFlags, ECollectionPropertyFlags::Animatable);  // Animatable
+	}
+	else
+	{
+		EnumRemoveFlags(PropertyFlags, ECollectionPropertyFlags::Animatable);  // Non-animatable
+	}
+	const int32 PropertyKeyIndex = ConfigNode.AddPropertyHelper(Properties, PropertyName, SimilarPropertyNames, PropertyFlags);
+	Properties.SetWeightedValue(PropertyKeyIndex, PropertyValue.Low, PropertyValue.High);
+	Properties.SetStringValue(PropertyKeyIndex, ConfigNode.GetValue<FString>(Context, &PropertyValue.WeightMap));
+	PropertyValue.WeightMap_Override = ConfigNode.GetValue<FString>(Context, &PropertyValue.WeightMap, UE::Chaos::ClothAsset::FWeightMapTools::NotOverridden);
+	return PropertyKeyIndex;
+}
+
+int32 FChaosClothAssetSimulationBaseConfigNode::FPropertyHelper::SetPropertyWeighted(
+	const FName& PropertyName,
+	const FChaosClothAssetWeightedValueNonAnimatable& PropertyValue,
+	const TArray<FName>& SimilarPropertyNames,
+	ECollectionPropertyFlags PropertyFlags)
+{
+	ensureMsgf(!EnumHasAnyFlags(PropertyFlags, ECollectionPropertyFlags::Animatable), TEXT("Animatable flag ignored. This code is for a non animatable weighted property."));
+	EnumRemoveFlags(PropertyFlags, ECollectionPropertyFlags::Animatable);  // Non animatable
+	const int32 PropertyKeyIndex = ConfigNode.AddPropertyHelper(Properties, PropertyName, SimilarPropertyNames, PropertyFlags);
+	Properties.SetWeightedValue(PropertyKeyIndex, PropertyValue.Low, PropertyValue.High);
+	Properties.SetStringValue(PropertyKeyIndex, ConfigNode.GetValue<FString>(Context, &PropertyValue.WeightMap));
+	PropertyValue.WeightMap_Override = ConfigNode.GetValue<FString>(Context, &PropertyValue.WeightMap, UE::Chaos::ClothAsset::FWeightMapTools::NotOverridden);
+	return PropertyKeyIndex;
+}
+
+int32 FChaosClothAssetSimulationBaseConfigNode::FPropertyHelper::SetPropertyWeighted(
+	const FName& PropertyName,
+	const FChaosClothAssetWeightedValueNonAnimatableNoLowHighRange& PropertyValue,
+	const TArray<FName>& SimilarPropertyNames,
+	ECollectionPropertyFlags PropertyFlags)
+{
+	ensureMsgf(!EnumHasAnyFlags(PropertyFlags, ECollectionPropertyFlags::Animatable), TEXT("Animatable flag ignored. This code is for a non animatable weighted property."));
+	EnumRemoveFlags(PropertyFlags, ECollectionPropertyFlags::Animatable);  // Non animatable
+	const int32 PropertyKeyIndex = ConfigNode.AddPropertyHelper(Properties, PropertyName, SimilarPropertyNames, PropertyFlags);
+	constexpr float Low = 0.f;
+	constexpr float High = 1.f;
+	Properties.SetWeightedValue(PropertyKeyIndex, Low, High);
+	Properties.SetStringValue(PropertyKeyIndex, ConfigNode.GetValue<FString>(Context, &PropertyValue.WeightMap));
+	PropertyValue.WeightMap_Override = ConfigNode.GetValue<FString>(Context, &PropertyValue.WeightMap, UE::Chaos::ClothAsset::FWeightMapTools::NotOverridden);
+	return PropertyKeyIndex;
+}
+
 #undef LOCTEXT_NAMESPACE
diff --git a/Engine/Plugins/Experimental/ChaosClothAssetEditor/Source/ChaosClothAssetDataflowNodes/Private/ChaosClothAsset/SimulationBaseConfigNodePrivate.h b/Engine/Plugins/Experimental/ChaosClothAssetEditor/Source/ChaosClothAssetDataflowNodes/Private/ChaosClothAsset/SimulationBaseConfigNodePrivate.h
--- a/Engine/Plugins/Experimental/ChaosClothAssetEditor/Source/ChaosClothAssetDataflowNodes/Private/ChaosClothAsset/SimulationBaseConfigNodePrivate.h
+++ b/Engine/Plugins/Experimental/ChaosClothAssetEditor/Source/ChaosClothAssetDataflowNodes/Private/ChaosClothAsset/SimulationBaseConfigNodePrivate.h
@@ -5,6 +5,8 @@
 #include "Dataflow/DataflowInputOutput.h"
 #include "ChaosClothAsset/WeightedValue.h"
 
+// These macros are now deprecated, use the FChaosClothAssetSimulationBaseConfigNode::FPropertyHelper struct instead.
+
 /** Macro for adding and setting value properties from inside AddProperties() */
 #define UE_CHAOS_CLOTHASSET_SIMULATIONCONFIG_SETPROPERTY(Property) \
 	const int32 Property##KeyIndex = AddPropertyHelper(Properties, FName(TEXT(#Property))); \
diff --git a/Engine/Plugins/Experimental/ChaosClothAssetEditor/Source/ChaosClothAssetDataflowNodes/Private/ChaosClothAsset/SimulationCollisionConfigNode.cpp b/Engine/Plugins/Experimental/ChaosClothAssetEditor/Source/ChaosClothAssetDataflowNodes/Private/ChaosClothAsset/SimulationCollisionConfigNode.cpp
--- a/Engine/Plugins/Experimental/ChaosClothAssetEditor/Source/ChaosClothAssetDataflowNodes/Private/ChaosClothAsset/SimulationCollisionConfigNode.cpp
+++ b/Engine/Plugins/Experimental/ChaosClothAssetEditor/Source/ChaosClothAssetDataflowNodes/Private/ChaosClothAsset/SimulationCollisionConfigNode.cpp
@@ -1,7 +1,6 @@
 // Copyright Epic Games, Inc. All Rights Reserved.
 
 #include "ChaosClothAsset/SimulationCollisionConfigNode.h"
-#include "ChaosClothAsset/SimulationBaseConfigNodePrivate.h"
 #include "Chaos/CollectionPropertyFacade.h"
 
 #include UE_INLINE_GENERATED_CPP_BY_NAME(SimulationCollisionConfigNode)
@@ -12,10 +11,10 @@ FChaosClothAssetSimulationCollisionConfigNode::FChaosClothAssetSimulationCollisi
 	RegisterCollectionConnections();
 }
 
-void FChaosClothAssetSimulationCollisionConfigNode::AddProperties(Dataflow::FContext& Context, ::Chaos::Softs::FCollectionPropertyMutableFacade& Properties) const
+void FChaosClothAssetSimulationCollisionConfigNode::AddProperties(FPropertyHelper& PropertyHelper) const
 {
-	UE_CHAOS_CLOTHASSET_SIMULATIONCONFIG_SETPROPERTY(CollisionThickness);
-	UE_CHAOS_CLOTHASSET_SIMULATIONCONFIG_SETPROPERTY(FrictionCoefficient);
-	UE_CHAOS_CLOTHASSET_SIMULATIONCONFIG_SETPROPERTYBOOL(UseCCD);
-	UE_CHAOS_CLOTHASSET_SIMULATIONCONFIG_SETPROPERTY(ProximityStiffness);
+	PropertyHelper.SetProperty(this, &CollisionThickness);
+	PropertyHelper.SetProperty(this, &FrictionCoefficient);
+	PropertyHelper.SetPropertyBool(this, &bUseCCD);
+	PropertyHelper.SetProperty(this, &ProximityStiffness);
 }
diff --git a/Engine/Plugins/Experimental/ChaosClothAssetEditor/Source/ChaosClothAssetDataflowNodes/Private/ChaosClothAsset/SimulationDampingConfigNode.cpp b/Engine/Plugins/Experimental/ChaosClothAssetEditor/Source/ChaosClothAssetDataflowNodes/Private/ChaosClothAsset/SimulationDampingConfigNode.cpp
--- a/Engine/Plugins/Experimental/ChaosClothAssetEditor/Source/ChaosClothAssetDataflowNodes/Private/ChaosClothAsset/SimulationDampingConfigNode.cpp
+++ b/Engine/Plugins/Experimental/ChaosClothAssetEditor/Source/ChaosClothAssetDataflowNodes/Private/ChaosClothAsset/SimulationDampingConfigNode.cpp
@@ -1,7 +1,6 @@
 // Copyright Epic Games, Inc. All Rights Reserved.
 
 #include "ChaosClothAsset/SimulationDampingConfigNode.h"
-#include "ChaosClothAsset/SimulationBaseConfigNodePrivate.h"
 #include "Chaos/CollectionPropertyFacade.h"
 
 #include UE_INLINE_GENERATED_CPP_BY_NAME(SimulationDampingConfigNode)
@@ -12,8 +11,8 @@ FChaosClothAssetSimulationDampingConfigNode::FChaosClothAssetSimulationDampingCo
 	RegisterCollectionConnections();
 }
 
-void FChaosClothAssetSimulationDampingConfigNode::AddProperties(Dataflow::FContext& Context, ::Chaos::Softs::FCollectionPropertyMutableFacade& Properties) const
+void FChaosClothAssetSimulationDampingConfigNode::AddProperties(FPropertyHelper& PropertyHelper) const
 {
-	UE_CHAOS_CLOTHASSET_SIMULATIONCONFIG_SETPROPERTY(DampingCoefficient);
-	UE_CHAOS_CLOTHASSET_SIMULATIONCONFIG_SETPROPERTY(LocalDampingCoefficient);
+	PropertyHelper.SetProperty(this, &DampingCoefficient);
+	PropertyHelper.SetProperty(this, &LocalDampingCoefficient);
 }
diff --git a/Engine/Plugins/Experimental/ChaosClothAssetEditor/Source/ChaosClothAssetDataflowNodes/Private/ChaosClothAsset/SimulationGravityConfigNode.cpp b/Engine/Plugins/Experimental/ChaosClothAssetEditor/Source/ChaosClothAssetDataflowNodes/Private/ChaosClothAsset/SimulationGravityConfigNode.cpp
--- a/Engine/Plugins/Experimental/ChaosClothAssetEditor/Source/ChaosClothAssetDataflowNodes/Private/ChaosClothAsset/SimulationGravityConfigNode.cpp
+++ b/Engine/Plugins/Experimental/ChaosClothAssetEditor/Source/ChaosClothAssetDataflowNodes/Private/ChaosClothAsset/SimulationGravityConfigNode.cpp
@@ -1,7 +1,6 @@
 // Copyright Epic Games, Inc. All Rights Reserved.
 
 #include "ChaosClothAsset/SimulationGravityConfigNode.h"
-#include "ChaosClothAsset/SimulationBaseConfigNodePrivate.h"
 #include "Chaos/CollectionPropertyFacade.h"
 
 #include UE_INLINE_GENERATED_CPP_BY_NAME(SimulationGravityConfigNode)
@@ -12,9 +11,9 @@ FChaosClothAssetSimulationGravityConfigNode::FChaosClothAssetSimulationGravityCo
 	RegisterCollectionConnections();
 }
 
-void FChaosClothAssetSimulationGravityConfigNode::AddProperties(Dataflow::FContext& Context, ::Chaos::Softs::FCollectionPropertyMutableFacade& Properties) const
+void FChaosClothAssetSimulationGravityConfigNode::AddProperties(FPropertyHelper& PropertyHelper) const
 {
-	UE_CHAOS_CLOTHASSET_SIMULATIONCONFIG_SETPROPERTYBOOL(UseGravityOverride);
-	UE_CHAOS_CLOTHASSET_SIMULATIONCONFIG_SETPROPERTY(GravityScale);
-	UE_CHAOS_CLOTHASSET_SIMULATIONCONFIG_SETPROPERTY(GravityOverride);
+	PropertyHelper.SetPropertyBool(this, &bUseGravityOverride);
+	PropertyHelper.SetProperty(this, &GravityScale);
+	PropertyHelper.SetProperty(this, &GravityOverride);
 }
diff --git a/Engine/Plugins/Experimental/ChaosClothAssetEditor/Source/ChaosClothAssetDataflowNodes/Private/ChaosClothAsset/SimulationLongRangeAttachmentConfigNode.cpp b/Engine/Plugins/Experimental/ChaosClothAssetEditor/Source/ChaosClothAssetDataflowNodes/Private/ChaosClothAsset/SimulationLongRangeAttachmentConfigNode.cpp
--- a/Engine/Plugins/Experimental/ChaosClothAssetEditor/Source/ChaosClothAssetDataflowNodes/Private/ChaosClothAsset/SimulationLongRangeAttachmentConfigNode.cpp
+++ b/Engine/Plugins/Experimental/ChaosClothAssetEditor/Source/ChaosClothAssetDataflowNodes/Private/ChaosClothAsset/SimulationLongRangeAttachmentConfigNode.cpp
@@ -1,7 +1,6 @@
 // Copyright Epic Games, Inc. All Rights Reserved.
 
 #include "ChaosClothAsset/SimulationLongRangeAttachmentConfigNode.h"
-#include "ChaosClothAsset/SimulationBaseConfigNodePrivate.h"
 #include "ChaosClothAsset/ClothEngineTools.h"
 #include "Chaos/CollectionPropertyFacade.h"
 #include "Dataflow/DataflowInputOutput.h"
@@ -17,12 +16,12 @@ FChaosClothAssetSimulationLongRangeAttachmentConfigNode::FChaosClothAssetSimulat
 	RegisterInputConnection(&TetherScale.WeightMap);
 }
 
-void FChaosClothAssetSimulationLongRangeAttachmentConfigNode::AddProperties(Dataflow::FContext& Context, ::Chaos::Softs::FCollectionPropertyMutableFacade& Properties) const
+void FChaosClothAssetSimulationLongRangeAttachmentConfigNode::AddProperties(FPropertyHelper& PropertyHelper) const
 {
-	UE_CHAOS_CLOTHASSET_SIMULATIONCONFIG_SETPROPERTYWEIGHTED(TetherStiffness);
-	UE_CHAOS_CLOTHASSET_SIMULATIONCONFIG_SETPROPERTYWEIGHTED(TetherScale);
-	UE_CHAOS_CLOTHASSET_SIMULATIONCONFIG_SETPROPERTYBOOL(UseGeodesicTethers);
-	UE_CHAOS_CLOTHASSET_SIMULATIONCONFIG_SETPROPERTYSTRING(FixedEndWeightMap);
+	PropertyHelper.SetPropertyWeighted(this, &TetherStiffness);
+	PropertyHelper.SetPropertyWeighted(this, &TetherScale);
+	PropertyHelper.SetPropertyBool(this, &bUseGeodesicTethers);
+	PropertyHelper.SetPropertyString(this, &FixedEndWeightMap);
 }
 
 void FChaosClothAssetSimulationLongRangeAttachmentConfigNode::EvaluateClothCollection(Dataflow::FContext& Context, const TSharedRef<FManagedArrayCollection>& ClothCollection) const
diff --git a/Engine/Plugins/Experimental/ChaosClothAssetEditor/Source/ChaosClothAssetDataflowNodes/Private/ChaosClothAsset/SimulationMassConfigNode.cpp b/Engine/Plugins/Experimental/ChaosClothAssetEditor/Source/ChaosClothAssetDataflowNodes/Private/ChaosClothAsset/SimulationMassConfigNode.cpp
--- a/Engine/Plugins/Experimental/ChaosClothAssetEditor/Source/ChaosClothAssetDataflowNodes/Private/ChaosClothAsset/SimulationMassConfigNode.cpp
+++ b/Engine/Plugins/Experimental/ChaosClothAssetEditor/Source/ChaosClothAssetDataflowNodes/Private/ChaosClothAsset/SimulationMassConfigNode.cpp
@@ -1,7 +1,6 @@
 // Copyright Epic Games, Inc. All Rights Reserved.
 
 #include "ChaosClothAsset/SimulationMassConfigNode.h"
-#include "ChaosClothAsset/SimulationBaseConfigNodePrivate.h"
 #include "Chaos/CollectionPropertyFacade.h"
 
 #include UE_INLINE_GENERATED_CPP_BY_NAME(SimulationMassConfigNode)
@@ -12,7 +11,7 @@ FChaosClothAssetSimulationMassConfigNode::FChaosClothAssetSimulationMassConfigNo
 	RegisterCollectionConnections();
 }
 
-void FChaosClothAssetSimulationMassConfigNode::AddProperties(Dataflow::FContext& Context, ::Chaos::Softs::FCollectionPropertyMutableFacade& Properties) const
+void FChaosClothAssetSimulationMassConfigNode::AddProperties(FPropertyHelper& PropertyHelper) const
 {
 	float MassValue;
 	switch (MassMode)
@@ -23,7 +22,7 @@ void FChaosClothAssetSimulationMassConfigNode::AddProperties(Dataflow::FContext&
 	case EClothMassMode::Density: MassValue = Density; break;
 	}
 
-	UE_CHAOS_CLOTHASSET_SIMULATIONCONFIG_SETPROPERTYENUM(MassMode);
-	UE_CHAOS_CLOTHASSET_SIMULATIONCONFIG_SETPROPERTY(MassValue);
-	UE_CHAOS_CLOTHASSET_SIMULATIONCONFIG_SETPROPERTY(MinPerParticleMass);
+	PropertyHelper.SetPropertyEnum(this, &MassMode, {}, ECollectionPropertyFlags::Intrinsic);
+	PropertyHelper.SetProperty(FName(TEXT("MassValue")), MassValue, {}, ECollectionPropertyFlags::Intrinsic);
+	PropertyHelper.SetProperty(this, &MinPerParticleMass, {}, ECollectionPropertyFlags::Intrinsic);
 }
diff --git a/Engine/Plugins/Experimental/ChaosClothAssetEditor/Source/ChaosClothAssetDataflowNodes/Private/ChaosClothAsset/SimulationMaxDistanceConfigNode.cpp b/Engine/Plugins/Experimental/ChaosClothAssetEditor/Source/ChaosClothAssetDataflowNodes/Private/ChaosClothAsset/SimulationMaxDistanceConfigNode.cpp
--- a/Engine/Plugins/Experimental/ChaosClothAssetEditor/Source/ChaosClothAssetDataflowNodes/Private/ChaosClothAsset/SimulationMaxDistanceConfigNode.cpp
+++ b/Engine/Plugins/Experimental/ChaosClothAssetEditor/Source/ChaosClothAssetDataflowNodes/Private/ChaosClothAsset/SimulationMaxDistanceConfigNode.cpp
@@ -1,7 +1,6 @@
 // Copyright Epic Games, Inc. All Rights Reserved.
 
 #include "ChaosClothAsset/SimulationMaxDistanceConfigNode.h"
-#include "ChaosClothAsset/SimulationBaseConfigNodePrivate.h"
 #include "Chaos/CollectionPropertyFacade.h"
 
 #include UE_INLINE_GENERATED_CPP_BY_NAME(SimulationMaxDistanceConfigNode)
@@ -13,7 +12,7 @@ FChaosClothAssetSimulationMaxDistanceConfigNode::FChaosClothAssetSimulationMaxDi
 	RegisterInputConnection(&MaxDistance.WeightMap);
 }
 
-void FChaosClothAssetSimulationMaxDistanceConfigNode::AddProperties(Dataflow::FContext& Context, ::Chaos::Softs::FCollectionPropertyMutableFacade& Properties) const
+void FChaosClothAssetSimulationMaxDistanceConfigNode::AddProperties(FPropertyHelper& PropertyHelper) const
 {
-	UE_CHAOS_CLOTHASSET_SIMULATIONCONFIG_SETPROPERTYWEIGHTED(MaxDistance);
+	PropertyHelper.SetPropertyWeighted(this, &MaxDistance, {}, ECollectionPropertyFlags::Intrinsic);  // Intrinsic since the deformer weights needs to be recalculated
 }
diff --git a/Engine/Plugins/Experimental/ChaosClothAssetEditor/Source/ChaosClothAssetDataflowNodes/Private/ChaosClothAsset/SimulationPBDAreaSpringConfigNode.cpp b/Engine/Plugins/Experimental/ChaosClothAssetEditor/Source/ChaosClothAssetDataflowNodes/Private/ChaosClothAsset/SimulationPBDAreaSpringConfigNode.cpp
--- a/Engine/Plugins/Experimental/ChaosClothAssetEditor/Source/ChaosClothAssetDataflowNodes/Private/ChaosClothAsset/SimulationPBDAreaSpringConfigNode.cpp
+++ b/Engine/Plugins/Experimental/ChaosClothAssetEditor/Source/ChaosClothAssetDataflowNodes/Private/ChaosClothAsset/SimulationPBDAreaSpringConfigNode.cpp
@@ -1,7 +1,6 @@
 // Copyright Epic Games, Inc. All Rights Reserved.
 
 #include "ChaosClothAsset/SimulationPBDAreaSpringConfigNode.h"
-#include "ChaosClothAsset/SimulationBaseConfigNodePrivate.h"
 #include "Chaos/CollectionPropertyFacade.h"
 
 #include UE_INLINE_GENERATED_CPP_BY_NAME(SimulationPBDAreaSpringConfigNode)
@@ -13,9 +12,11 @@ FChaosClothAssetSimulationPBDAreaSpringConfigNode::FChaosClothAssetSimulationPBD
 	RegisterInputConnection(&AreaSpringStiffness.WeightMap);
 }
 
-void FChaosClothAssetSimulationPBDAreaSpringConfigNode::AddProperties(Dataflow::FContext& Context, ::Chaos::Softs::FCollectionPropertyMutableFacade& Properties) const
+void FChaosClothAssetSimulationPBDAreaSpringConfigNode::AddProperties(FPropertyHelper& PropertyHelper) const
 {
-	UE_CHAOS_CLOTHASSET_SIMULATIONCONFIG_SETPROPERTYWEIGHTEDCHECKED1(
-		AreaSpringStiffness,
-		XPBDAreaSpringStiffness);  // Existing properties to warn against
+	PropertyHelper.SetPropertyWeighted(this,
+		&AreaSpringStiffness,
+		{
+			FName(TEXT("XPBDAreaSpringStiffness"))  // Existing properties to warn against
+		});
 }
diff --git a/Engine/Plugins/Experimental/ChaosClothAssetEditor/Source/ChaosClothAssetDataflowNodes/Private/ChaosClothAsset/SimulationPBDBendingElementConfigNode.cpp b/Engine/Plugins/Experimental/ChaosClothAssetEditor/Source/ChaosClothAssetDataflowNodes/Private/ChaosClothAsset/SimulationPBDBendingElementConfigNode.cpp
--- a/Engine/Plugins/Experimental/ChaosClothAssetEditor/Source/ChaosClothAssetDataflowNodes/Private/ChaosClothAsset/SimulationPBDBendingElementConfigNode.cpp
+++ b/Engine/Plugins/Experimental/ChaosClothAssetEditor/Source/ChaosClothAssetDataflowNodes/Private/ChaosClothAsset/SimulationPBDBendingElementConfigNode.cpp
@@ -1,7 +1,6 @@
 // Copyright Epic Games, Inc. All Rights Reserved.
 
 #include "ChaosClothAsset/SimulationPBDBendingElementConfigNode.h"
-#include "ChaosClothAsset/SimulationBaseConfigNodePrivate.h"
 #include "Chaos/CollectionPropertyFacade.h"
 
 #include UE_INLINE_GENERATED_CPP_BY_NAME(SimulationPBDBendingElementConfigNode)
@@ -16,32 +15,45 @@ FChaosClothAssetSimulationPBDBendingElementConfigNode::FChaosClothAssetSimulatio
 	RegisterInputConnection(&BucklingStiffness.WeightMap);
 }
 
-void FChaosClothAssetSimulationPBDBendingElementConfigNode::AddProperties(Dataflow::FContext& Context, ::Chaos::Softs::FCollectionPropertyMutableFacade& Properties) const
+void FChaosClothAssetSimulationPBDBendingElementConfigNode::AddProperties(FPropertyHelper& PropertyHelper) const
 {
-	UE_CHAOS_CLOTHASSET_SIMULATIONCONFIG_SETPROPERTYENUMCHECKED2(
-		RestAngleType,
-		XPBDRestAngleType,        // Existing properties to warn against
-		XPBDAnisoRestAngleType);  //
-	UE_CHAOS_CLOTHASSET_SIMULATIONCONFIG_SETPROPERTYWEIGHTEDCHECKED2(
-		FlatnessRatio,
-		XPBDFlatnessRatio,        // Existing properties to warn against
-		XPBDAnisoFlatnessRatio);  //
-	UE_CHAOS_CLOTHASSET_SIMULATIONCONFIG_SETPROPERTYWEIGHTEDCHECKED2(
-		RestAngle,
-		XPBDRestAngle,        // Existing properties to warn against
-		XPBDAnisoRestAngle);  //
-	UE_CHAOS_CLOTHASSET_SIMULATIONCONFIG_SETPROPERTYWEIGHTEDCHECKED4(
-		BendingElementStiffness,
-		BendingSpringStiffness,          // Existing properties to warn against
-		XPBDBendingSpringStiffness,      //
-		XPBDBendingElementStiffness,     //
-		XPBDAnisoBendingStiffnessWarp);  //
-	UE_CHAOS_CLOTHASSET_SIMULATIONCONFIG_SETPROPERTYCHECKED2(
-		BucklingRatio,
-		XPBDBucklingRatio,        // Existing properties to warn against
-		XPBDAnisoBucklingRatio);  //
-	UE_CHAOS_CLOTHASSET_SIMULATIONCONFIG_SETPROPERTYWEIGHTEDCHECKED2(
-		BucklingStiffness,
-		XPBDBucklingStiffness,            // Existing properties to warn against
-		XPBDAnisoBucklingStiffnessWarp);  //
+	PropertyHelper.SetPropertyEnum(this,
+		&RestAngleType,
+		{
+			FName(TEXT("XPBDRestAngleType")),      // Existing properties to warn against
+			FName(TEXT("XPBDAnisoRestAngleType"))  //
+		},
+		ECollectionPropertyFlags::None);  // Non animatable
+	PropertyHelper.SetPropertyWeighted(this,
+		&FlatnessRatio,
+		{
+			FName(TEXT("XPBDFlatnessRatio")),      // Existing properties to warn against
+			FName(TEXT("XPBDAnisoFlatnessRatio"))  //
+		});
+	PropertyHelper.SetPropertyWeighted(this,
+		&RestAngle,
+		{
+			FName(TEXT("XPBDRestAngle")),      // Existing properties to warn against
+			FName(TEXT("XPBDAnisoRestAngle"))  //
+		});
+	PropertyHelper.SetPropertyWeighted(this,
+		&BendingElementStiffness,
+		{
+			FName(TEXT("BendingSpringStiffness")),        // Existing properties to warn against
+			FName(TEXT("XPBDBendingSpringStiffness")),    //
+			FName(TEXT("XPBDBendingElementStiffness")),   //
+			FName(TEXT("XPBDAnisoBendingStiffnessWarp"))  //
+		});
+	PropertyHelper.SetProperty(this,
+		&BucklingRatio,
+		{
+			FName(TEXT("XPBDBucklingRatio")),      // Existing properties to warn against
+			FName(TEXT("XPBDAnisoBucklingRatio"))  //
+		});
+	PropertyHelper.SetPropertyWeighted(this,
+		&BucklingStiffness,
+		{
+			FName(TEXT("XPBDBucklingStiffness")),          // Existing properties to warn against
+			FName(TEXT("XPBDAnisoBucklingStiffnessWarp"))  //
+		});
 }
diff --git a/Engine/Plugins/Experimental/ChaosClothAssetEditor/Source/ChaosClothAssetDataflowNodes/Private/ChaosClothAsset/SimulationPBDBendingSpringConfigNode.cpp b/Engine/Plugins/Experimental/ChaosClothAssetEditor/Source/ChaosClothAssetDataflowNodes/Private/ChaosClothAsset/SimulationPBDBendingSpringConfigNode.cpp
--- a/Engine/Plugins/Experimental/ChaosClothAssetEditor/Source/ChaosClothAssetDataflowNodes/Private/ChaosClothAsset/SimulationPBDBendingSpringConfigNode.cpp
+++ b/Engine/Plugins/Experimental/ChaosClothAssetEditor/Source/ChaosClothAssetDataflowNodes/Private/ChaosClothAsset/SimulationPBDBendingSpringConfigNode.cpp
@@ -1,7 +1,6 @@
 // Copyright Epic Games, Inc. All Rights Reserved.
 
 #include "ChaosClothAsset/SimulationPBDBendingSpringConfigNode.h"
-#include "ChaosClothAsset/SimulationBaseConfigNodePrivate.h"
 #include "Chaos/CollectionPropertyFacade.h"
 
 #include UE_INLINE_GENERATED_CPP_BY_NAME(SimulationPBDBendingSpringConfigNode)
@@ -13,12 +12,14 @@ FChaosClothAssetSimulationPBDBendingSpringConfigNode::FChaosClothAssetSimulation
 	RegisterInputConnection(&BendingSpringStiffness.WeightMap);
 }
 
-void FChaosClothAssetSimulationPBDBendingSpringConfigNode::AddProperties(Dataflow::FContext& Context, ::Chaos::Softs::FCollectionPropertyMutableFacade& Properties) const
+void FChaosClothAssetSimulationPBDBendingSpringConfigNode::AddProperties(FPropertyHelper& PropertyHelper) const
 {
-	UE_CHAOS_CLOTHASSET_SIMULATIONCONFIG_SETPROPERTYWEIGHTEDCHECKED4(
-		BendingSpringStiffness,
-		BendingElementStiffness,         // Existing properties to warn against
-		XPBDBendingSpringStiffness,      //
-		XPBDBendingElementStiffness,     //
-		XPBDAnisoBendingStiffnessWarp);  //
+	PropertyHelper.SetPropertyWeighted(this,
+		&BendingSpringStiffness,
+		{
+			FName(TEXT("BendingElementStiffness")),       // Existing properties to warn against
+			FName(TEXT("XPBDBendingSpringStiffness")),    //
+			FName(TEXT("XPBDBendingElementStiffness")),   //
+			FName(TEXT("XPBDAnisoBendingStiffnessWarp"))  //
+		});
 }
diff --git a/Engine/Plugins/Experimental/ChaosClothAssetEditor/Source/ChaosClothAssetDataflowNodes/Private/ChaosClothAsset/SimulationPBDEdgeSpringConfigNode.cpp b/Engine/Plugins/Experimental/ChaosClothAssetEditor/Source/ChaosClothAssetDataflowNodes/Private/ChaosClothAsset/SimulationPBDEdgeSpringConfigNode.cpp
--- a/Engine/Plugins/Experimental/ChaosClothAssetEditor/Source/ChaosClothAssetDataflowNodes/Private/ChaosClothAsset/SimulationPBDEdgeSpringConfigNode.cpp
+++ b/Engine/Plugins/Experimental/ChaosClothAssetEditor/Source/ChaosClothAssetDataflowNodes/Private/ChaosClothAsset/SimulationPBDEdgeSpringConfigNode.cpp
@@ -1,7 +1,6 @@
 // Copyright Epic Games, Inc. All Rights Reserved.
 
 #include "ChaosClothAsset/SimulationPBDEdgeSpringConfigNode.h"
-#include "ChaosClothAsset/SimulationBaseConfigNodePrivate.h"
 #include "Chaos/CollectionPropertyFacade.h"
 
 #include UE_INLINE_GENERATED_CPP_BY_NAME(SimulationPBDEdgeSpringConfigNode)
@@ -13,10 +12,12 @@ FChaosClothAssetSimulationPBDEdgeSpringConfigNode::FChaosClothAssetSimulationPBD
 	RegisterInputConnection(&EdgeSpringStiffness.WeightMap);
 }
 
-void FChaosClothAssetSimulationPBDEdgeSpringConfigNode::AddProperties(Dataflow::FContext& Context, ::Chaos::Softs::FCollectionPropertyMutableFacade& Properties) const
+void FChaosClothAssetSimulationPBDEdgeSpringConfigNode::AddProperties(FPropertyHelper& PropertyHelper) const
 {
-	UE_CHAOS_CLOTHASSET_SIMULATIONCONFIG_SETPROPERTYWEIGHTEDCHECKED2(
-		EdgeSpringStiffness,
-		XPBDEdgeSpringStiffness,         // Existing properties to warn against
-		XPBDAnisoStretchStiffnessWarp);  //
+	PropertyHelper.SetPropertyWeighted(this,
+		&EdgeSpringStiffness,
+		{
+			FName(TEXT("XPBDEdgeSpringStiffness")),       // Existing properties to warn against
+			FName(TEXT("XPBDAnisoStretchStiffnessWarp"))  //
+		});
 }
diff --git a/Engine/Plugins/Experimental/ChaosClothAssetEditor/Source/ChaosClothAssetDataflowNodes/Private/ChaosClothAsset/SimulationPressureConfigNode.cpp b/Engine/Plugins/Experimental/ChaosClothAssetEditor/Source/ChaosClothAssetDataflowNodes/Private/ChaosClothAsset/SimulationPressureConfigNode.cpp
--- a/Engine/Plugins/Experimental/ChaosClothAssetEditor/Source/ChaosClothAssetDataflowNodes/Private/ChaosClothAsset/SimulationPressureConfigNode.cpp
+++ b/Engine/Plugins/Experimental/ChaosClothAssetEditor/Source/ChaosClothAssetDataflowNodes/Private/ChaosClothAsset/SimulationPressureConfigNode.cpp
@@ -1,7 +1,6 @@
 // Copyright Epic Games, Inc. All Rights Reserved.
 
 #include "ChaosClothAsset/SimulationPressureConfigNode.h"
-#include "ChaosClothAsset/SimulationBaseConfigNodePrivate.h"
 #include "Chaos/CollectionPropertyFacade.h"
 
 #include UE_INLINE_GENERATED_CPP_BY_NAME(SimulationPressureConfigNode)
@@ -13,7 +12,7 @@ FChaosClothAssetSimulationPressureConfigNode::FChaosClothAssetSimulationPressure
 	RegisterInputConnection(&Pressure.WeightMap);
 }
 
-void FChaosClothAssetSimulationPressureConfigNode::AddProperties(Dataflow::FContext& Context, ::Chaos::Softs::FCollectionPropertyMutableFacade& Properties) const
+void FChaosClothAssetSimulationPressureConfigNode::AddProperties(FPropertyHelper& PropertyHelper) const
 {
-	UE_CHAOS_CLOTHASSET_SIMULATIONCONFIG_SETPROPERTYWEIGHTED(Pressure);
+	PropertyHelper.SetPropertyWeighted(this, &Pressure);
 }
diff --git a/Engine/Plugins/Experimental/ChaosClothAssetEditor/Source/ChaosClothAssetDataflowNodes/Private/ChaosClothAsset/SimulationSelfCollisionConfigNode.cpp b/Engine/Plugins/Experimental/ChaosClothAssetEditor/Source/ChaosClothAssetDataflowNodes/Private/ChaosClothAsset/SimulationSelfCollisionConfigNode.cpp
--- a/Engine/Plugins/Experimental/ChaosClothAssetEditor/Source/ChaosClothAssetDataflowNodes/Private/ChaosClothAsset/SimulationSelfCollisionConfigNode.cpp
+++ b/Engine/Plugins/Experimental/ChaosClothAssetEditor/Source/ChaosClothAssetDataflowNodes/Private/ChaosClothAsset/SimulationSelfCollisionConfigNode.cpp
@@ -1,7 +1,6 @@
 // Copyright Epic Games, Inc. All Rights Reserved.
 
 #include "ChaosClothAsset/SimulationSelfCollisionConfigNode.h"
-#include "ChaosClothAsset/SimulationBaseConfigNodePrivate.h"
 #include "Chaos/CollectionPropertyFacade.h"
 
 #include UE_INLINE_GENERATED_CPP_BY_NAME(SimulationSelfCollisionConfigNode)
@@ -12,18 +11,17 @@ FChaosClothAssetSimulationSelfCollisionConfigNode::FChaosClothAssetSimulationSel
 	RegisterCollectionConnections();
 }
 
-void FChaosClothAssetSimulationSelfCollisionConfigNode::AddProperties(Dataflow::FContext& Context, ::Chaos::Softs::FCollectionPropertyMutableFacade& Properties) const
+void FChaosClothAssetSimulationSelfCollisionConfigNode::AddProperties(FPropertyHelper& PropertyHelper) const
 {
-	constexpr bool bUseSelfCollisions = true;
-	UE_CHAOS_CLOTHASSET_SIMULATIONCONFIG_SETPROPERTYBOOL(UseSelfCollisions);
-	UE_CHAOS_CLOTHASSET_SIMULATIONCONFIG_SETPROPERTY(SelfCollisionThickness);
-	UE_CHAOS_CLOTHASSET_SIMULATIONCONFIG_SETPROPERTY(SelfCollisionStiffness);
-	UE_CHAOS_CLOTHASSET_SIMULATIONCONFIG_SETPROPERTY(SelfCollisionFriction);
+	PropertyHelper.SetPropertyBool(FName("UseSelfCollisions"), true);
+	PropertyHelper.SetProperty(this, &SelfCollisionThickness);
+	PropertyHelper.SetProperty(this, &SelfCollisionStiffness);
+	PropertyHelper.SetProperty(this, &SelfCollisionFriction);
 
-	UE_CHAOS_CLOTHASSET_SIMULATIONCONFIG_SETPROPERTYBOOL(UseSelfIntersections);
-	UE_CHAOS_CLOTHASSET_SIMULATIONCONFIG_SETPROPERTYBOOL(UseGlobalIntersectionAnalysis);
-	UE_CHAOS_CLOTHASSET_SIMULATIONCONFIG_SETPROPERTYBOOL(UseContourMinimization);
-	UE_CHAOS_CLOTHASSET_SIMULATIONCONFIG_SETPROPERTY(NumContourMinimizationPostSteps);
-	UE_CHAOS_CLOTHASSET_SIMULATIONCONFIG_SETPROPERTYBOOL(UseGlobalPostStepContours);
-	UE_CHAOS_CLOTHASSET_SIMULATIONCONFIG_SETPROPERTY(SelfCollisionProximityStiffness);
+	PropertyHelper.SetPropertyBool(this, &bUseSelfIntersections);
+	PropertyHelper.SetPropertyBool(this, &bUseGlobalIntersectionAnalysis);
+	PropertyHelper.SetPropertyBool(this, &bUseContourMinimization);
+	PropertyHelper.SetProperty(this, &NumContourMinimizationPostSteps);
+	PropertyHelper.SetPropertyBool(this, &bUseGlobalPostStepContours);
+	PropertyHelper.SetProperty(this, &SelfCollisionProximityStiffness);
 }
diff --git a/Engine/Plugins/Experimental/ChaosClothAssetEditor/Source/ChaosClothAssetDataflowNodes/Private/ChaosClothAsset/SimulationSelfCollisionSpheresConfigNode.cpp b/Engine/Plugins/Experimental/ChaosClothAssetEditor/Source/ChaosClothAssetDataflowNodes/Private/ChaosClothAsset/SimulationSelfCollisionSpheresConfigNode.cpp
--- a/Engine/Plugins/Experimental/ChaosClothAssetEditor/Source/ChaosClothAssetDataflowNodes/Private/ChaosClothAsset/SimulationSelfCollisionSpheresConfigNode.cpp
+++ b/Engine/Plugins/Experimental/ChaosClothAssetEditor/Source/ChaosClothAssetDataflowNodes/Private/ChaosClothAsset/SimulationSelfCollisionSpheresConfigNode.cpp
@@ -1,7 +1,6 @@
 // Copyright Epic Games, Inc. All Rights Reserved.
 
 #include "ChaosClothAsset/SimulationSelfCollisionSpheresConfigNode.h"
-#include "ChaosClothAsset/SimulationBaseConfigNodePrivate.h"
 #include "Chaos/CollectionPropertyFacade.h"
 
 #include "ChaosClothAsset/CollectionClothFacade.h"
@@ -17,11 +16,11 @@ FChaosClothAssetSimulationSelfCollisionSpheresConfigNode::FChaosClothAssetSimula
 	RegisterCollectionConnections();
 }
 
-void FChaosClothAssetSimulationSelfCollisionSpheresConfigNode::AddProperties(Dataflow::FContext& Context, ::Chaos::Softs::FCollectionPropertyMutableFacade& Properties) const
+void FChaosClothAssetSimulationSelfCollisionSpheresConfigNode::AddProperties(FPropertyHelper& PropertyHelper) const
 {
-	UE_CHAOS_CLOTHASSET_SIMULATIONCONFIG_SETPROPERTY(SelfCollisionSphereRadius);
-	UE_CHAOS_CLOTHASSET_SIMULATIONCONFIG_SETPROPERTY(SelfCollisionSphereRadiusCullMultiplier);
-	UE_CHAOS_CLOTHASSET_SIMULATIONCONFIG_SETPROPERTY(SelfCollisionSphereStiffness); 
+	PropertyHelper.SetProperty(this, &SelfCollisionSphereRadius);
+	PropertyHelper.SetProperty(this, &SelfCollisionSphereRadiusCullMultiplier);
+	PropertyHelper.SetProperty(this, &SelfCollisionSphereStiffness);
 }
 
 void FChaosClothAssetSimulationSelfCollisionSpheresConfigNode::EvaluateClothCollection(Dataflow::FContext& Context, const TSharedRef<FManagedArrayCollection>& ClothCollection) const
diff --git a/Engine/Plugins/Experimental/ChaosClothAssetEditor/Source/ChaosClothAssetDataflowNodes/Private/ChaosClothAsset/SimulationSolverConfigNode.cpp b/Engine/Plugins/Experimental/ChaosClothAssetEditor/Source/ChaosClothAssetDataflowNodes/Private/ChaosClothAsset/SimulationSolverConfigNode.cpp
--- a/Engine/Plugins/Experimental/ChaosClothAssetEditor/Source/ChaosClothAssetDataflowNodes/Private/ChaosClothAsset/SimulationSolverConfigNode.cpp
+++ b/Engine/Plugins/Experimental/ChaosClothAssetEditor/Source/ChaosClothAssetDataflowNodes/Private/ChaosClothAsset/SimulationSolverConfigNode.cpp
@@ -1,7 +1,6 @@
 // Copyright Epic Games, Inc. All Rights Reserved.
 
 #include "ChaosClothAsset/SimulationSolverConfigNode.h"
-#include "ChaosClothAsset/SimulationBaseConfigNodePrivate.h"
 #include "Chaos/CollectionPropertyFacade.h"
 
 #include UE_INLINE_GENERATED_CPP_BY_NAME(SimulationSolverConfigNode)
@@ -12,16 +11,16 @@ FChaosClothAssetSimulationSolverConfigNode::FChaosClothAssetSimulationSolverConf
 	RegisterCollectionConnections();
 }
 
-void FChaosClothAssetSimulationSolverConfigNode::AddProperties(Dataflow::FContext& Context, ::Chaos::Softs::FCollectionPropertyMutableFacade& Properties) const
+void FChaosClothAssetSimulationSolverConfigNode::AddProperties(FPropertyHelper& PropertyHelper) const
 {
-	UE_CHAOS_CLOTHASSET_SIMULATIONCONFIG_SETPROPERTY(NumIterations);
-	UE_CHAOS_CLOTHASSET_SIMULATIONCONFIG_SETPROPERTY(MaxNumIterations);
-	UE_CHAOS_CLOTHASSET_SIMULATIONCONFIG_SETPROPERTY(NumSubsteps);
-	UE_CHAOS_CLOTHASSET_SIMULATIONCONFIG_SETPROPERTYBOOL(EnableNumSelfCollisionSubsteps);
-	UE_CHAOS_CLOTHASSET_SIMULATIONCONFIG_SETPROPERTY(NumSelfCollisionSubsteps);
-	UE_CHAOS_CLOTHASSET_SIMULATIONCONFIG_SETPROPERTYBOOL(EnableForceBasedSolver);
-	UE_CHAOS_CLOTHASSET_SIMULATIONCONFIG_SETPROPERTY(NumNewtonIterations);
-	UE_CHAOS_CLOTHASSET_SIMULATIONCONFIG_SETPROPERTY(MaxNumCGIterations);
-	UE_CHAOS_CLOTHASSET_SIMULATIONCONFIG_SETPROPERTY(CGResidualTolerance);
-	UE_CHAOS_CLOTHASSET_SIMULATIONCONFIG_SETPROPERTYBOOL(DoQuasistatics);
+	PropertyHelper.SetProperty(this, &NumIterations);
+	PropertyHelper.SetProperty(this, &MaxNumIterations);
+	PropertyHelper.SetProperty(this, &NumSubsteps);
+	PropertyHelper.SetPropertyBool(this, &bEnableNumSelfCollisionSubsteps);
+	PropertyHelper.SetProperty(this, &NumSelfCollisionSubsteps);
+	PropertyHelper.SetPropertyBool(this, &bEnableForceBasedSolver, {}, ECollectionPropertyFlags::Intrinsic);
+	PropertyHelper.SetProperty(this, &NumNewtonIterations);
+	PropertyHelper.SetProperty(this, &MaxNumCGIterations);
+	PropertyHelper.SetProperty(this, &CGResidualTolerance);
+	PropertyHelper.SetPropertyBool(this, &bDoQuasistatics);
 }
diff --git a/Engine/Plugins/Experimental/ChaosClothAssetEditor/Source/ChaosClothAssetDataflowNodes/Private/ChaosClothAsset/SimulationVelocityScaleConfigNode.cpp b/Engine/Plugins/Experimental/ChaosClothAssetEditor/Source/ChaosClothAssetDataflowNodes/Private/ChaosClothAsset/SimulationVelocityScaleConfigNode.cpp
--- a/Engine/Plugins/Experimental/ChaosClothAssetEditor/Source/ChaosClothAssetDataflowNodes/Private/ChaosClothAsset/SimulationVelocityScaleConfigNode.cpp
+++ b/Engine/Plugins/Experimental/ChaosClothAssetEditor/Source/ChaosClothAssetDataflowNodes/Private/ChaosClothAsset/SimulationVelocityScaleConfigNode.cpp
@@ -1,7 +1,6 @@
 // Copyright Epic Games, Inc. All Rights Reserved.
 
 #include "ChaosClothAsset/SimulationVelocityScaleConfigNode.h"
-#include "ChaosClothAsset/SimulationBaseConfigNodePrivate.h"
 #include "Chaos/CollectionPropertyFacade.h"
 
 #include UE_INLINE_GENERATED_CPP_BY_NAME(SimulationVelocityScaleConfigNode)
@@ -12,9 +11,9 @@ FChaosClothAssetSimulationVelocityScaleConfigNode::FChaosClothAssetSimulationVel
 	RegisterCollectionConnections();
 }
 
-void FChaosClothAssetSimulationVelocityScaleConfigNode::AddProperties(Dataflow::FContext& Context, ::Chaos::Softs::FCollectionPropertyMutableFacade& Properties) const
+void FChaosClothAssetSimulationVelocityScaleConfigNode::AddProperties(FPropertyHelper& PropertyHelper) const
 {
-	UE_CHAOS_CLOTHASSET_SIMULATIONCONFIG_SETPROPERTY(LinearVelocityScale);
-	UE_CHAOS_CLOTHASSET_SIMULATIONCONFIG_SETPROPERTY(AngularVelocityScale);
-	UE_CHAOS_CLOTHASSET_SIMULATIONCONFIG_SETPROPERTY(FictitiousAngularScale);
+	PropertyHelper.SetProperty(this, &LinearVelocityScale);
+	PropertyHelper.SetProperty(this, &AngularVelocityScale);
+	PropertyHelper.SetProperty(this, &FictitiousAngularScale);
 }
diff --git a/Engine/Plugins/Experimental/ChaosClothAssetEditor/Source/ChaosClothAssetDataflowNodes/Private/ChaosClothAsset/SimulationXPBDAnisoBendingConfigNode.cpp b/Engine/Plugins/Experimental/ChaosClothAssetEditor/Source/ChaosClothAssetDataflowNodes/Private/ChaosClothAsset/SimulationXPBDAnisoBendingConfigNode.cpp
--- a/Engine/Plugins/Experimental/ChaosClothAssetEditor/Source/ChaosClothAssetDataflowNodes/Private/ChaosClothAsset/SimulationXPBDAnisoBendingConfigNode.cpp
+++ b/Engine/Plugins/Experimental/ChaosClothAssetEditor/Source/ChaosClothAssetDataflowNodes/Private/ChaosClothAsset/SimulationXPBDAnisoBendingConfigNode.cpp
@@ -1,7 +1,6 @@
 // Copyright Epic Games, Inc. All Rights Reserved.
 
 #include "ChaosClothAsset/SimulationXPBDAnisoBendingConfigNode.h"
-#include "ChaosClothAsset/SimulationBaseConfigNodePrivate.h"
 #include "Chaos/CollectionPropertyFacade.h"
 
 #include UE_INLINE_GENERATED_CPP_BY_NAME(SimulationXPBDAnisoBendingConfigNode)
@@ -21,40 +20,55 @@ FChaosClothAssetSimulationXPBDAnisoBendingConfigNode::FChaosClothAssetSimulation
 	RegisterInputConnection(&XPBDAnisoBucklingStiffnessBias.WeightMap);
 }
 
-void FChaosClothAssetSimulationXPBDAnisoBendingConfigNode::AddProperties(Dataflow::FContext& Context, ::Chaos::Softs::FCollectionPropertyMutableFacade& Properties) const
+void FChaosClothAssetSimulationXPBDAnisoBendingConfigNode::AddProperties(FPropertyHelper& PropertyHelper) const
 {
-	UE_CHAOS_CLOTHASSET_SIMULATIONCONFIG_SETPROPERTYENUMCHECKED2(
-		XPBDAnisoRestAngleType,
-		XPBDRestAngleType,  // Existing properties to warn against
-		RestAngleType);     //
-	UE_CHAOS_CLOTHASSET_SIMULATIONCONFIG_SETPROPERTYWEIGHTEDCHECKED2(
-		XPBDAnisoFlatnessRatio,
-		XPBDFlatnessRatio,  // Existing properties to warn against
-		FlatnessRatio);     //
-	UE_CHAOS_CLOTHASSET_SIMULATIONCONFIG_SETPROPERTYWEIGHTEDCHECKED2(
-		XPBDAnisoRestAngle,
-		XPBDRestAngle,  // Existing properties to warn against
-		RestAngle);     //
-	UE_CHAOS_CLOTHASSET_SIMULATIONCONFIG_SETPROPERTYWEIGHTEDCHECKED4(
-		XPBDAnisoBendingStiffnessWarp,
-		BendingSpringStiffness,        // Existing properties to warn against
-		BendingElementStiffness,       //
-		XPBDBendingSpringStiffness,    //
-		XPBDBendingElementStiffness);  //
-	UE_CHAOS_CLOTHASSET_SIMULATIONCONFIG_SETPROPERTYWEIGHTED(XPBDAnisoBendingStiffnessWeft);
-	UE_CHAOS_CLOTHASSET_SIMULATIONCONFIG_SETPROPERTYWEIGHTED(XPBDAnisoBendingStiffnessBias);
-	UE_CHAOS_CLOTHASSET_SIMULATIONCONFIG_SETPROPERTYWEIGHTEDCHECKED2(
-		XPBDAnisoBendingDamping,
-		XPBDBendingSpringDamping,    // Existing properties to warn against
-		XPBDBendingElementDamping);  //
-	UE_CHAOS_CLOTHASSET_SIMULATIONCONFIG_SETPROPERTYCHECKED2(
-		XPBDAnisoBucklingRatio,
-		BucklingRatio,       // Existing properties to warn against
-		XPBDBucklingRatio);  //
-	UE_CHAOS_CLOTHASSET_SIMULATIONCONFIG_SETPROPERTYWEIGHTEDCHECKED2(
-		XPBDAnisoBucklingStiffnessWarp,
-		BucklingStiffness,       // Existing properties to warn against
-		XPBDBucklingStiffness);  //
-	UE_CHAOS_CLOTHASSET_SIMULATIONCONFIG_SETPROPERTYWEIGHTED(XPBDAnisoBucklingStiffnessWeft);
-	UE_CHAOS_CLOTHASSET_SIMULATIONCONFIG_SETPROPERTYWEIGHTED(XPBDAnisoBucklingStiffnessBias);
+	PropertyHelper.SetPropertyEnum(this,
+		&XPBDAnisoRestAngleType,
+		{
+			FName(TEXT("XPBDRestAngleType")),  // Existing properties to warn against
+			FName(TEXT("RestAngleType"))       //
+		},
+		ECollectionPropertyFlags::None);  // Non animatable
+	PropertyHelper.SetPropertyWeighted(this,
+		&XPBDAnisoFlatnessRatio,
+		{
+			FName(TEXT("XPBDFlatnessRatio")),  // Existing properties to warn against
+			FName(TEXT("FlatnessRatio"))       //
+		});
+	PropertyHelper.SetPropertyWeighted(this,
+		&XPBDAnisoRestAngle,
+		{
+			FName(TEXT("XPBDRestAngle")),  // Existing properties to warn against
+			FName(TEXT("RestAngle"))       //
+		});
+	PropertyHelper.SetPropertyWeighted(this,
+		&XPBDAnisoBendingStiffnessWarp,
+		{
+			FName(TEXT("BendingSpringStiffness")),      // Existing properties to warn against
+			FName(TEXT("BendingElementStiffness")),     //
+			FName(TEXT("XPBDBendingSpringStiffness")),  //
+			FName(TEXT("XPBDBendingElementStiffness"))  //
+		});
+	PropertyHelper.SetPropertyWeighted(this, &XPBDAnisoBendingStiffnessWeft);
+	PropertyHelper.SetPropertyWeighted(this, &XPBDAnisoBendingStiffnessBias);
+	PropertyHelper.SetPropertyWeighted(this,
+		&XPBDAnisoBendingDamping,
+		{
+			FName(TEXT("XPBDBendingSpringDamping")),  // Existing properties to warn against
+			FName(TEXT("XPBDBendingElementDamping"))  //
+		});
+	PropertyHelper.SetProperty(this,
+		&XPBDAnisoBucklingRatio,
+		{
+			FName(TEXT("BucklingRatio")),     // Existing properties to warn against
+			FName(TEXT("XPBDBucklingRatio"))  //
+		});
+	PropertyHelper.SetPropertyWeighted(this,
+		&XPBDAnisoBucklingStiffnessWarp,
+		{
+			FName(TEXT("BucklingStiffness")),     // Existing properties to warn against
+			FName(TEXT("XPBDBucklingStiffness"))  //
+		});
+	PropertyHelper.SetPropertyWeighted(this, &XPBDAnisoBucklingStiffnessWeft);
+	PropertyHelper.SetPropertyWeighted(this, &XPBDAnisoBucklingStiffnessBias);
 }
diff --git a/Engine/Plugins/Experimental/ChaosClothAssetEditor/Source/ChaosClothAssetDataflowNodes/Private/ChaosClothAsset/SimulationXPBDAnisoSpringConfigNode.cpp b/Engine/Plugins/Experimental/ChaosClothAssetEditor/Source/ChaosClothAssetDataflowNodes/Private/ChaosClothAsset/SimulationXPBDAnisoSpringConfigNode.cpp
--- a/Engine/Plugins/Experimental/ChaosClothAssetEditor/Source/ChaosClothAssetDataflowNodes/Private/ChaosClothAsset/SimulationXPBDAnisoSpringConfigNode.cpp
+++ b/Engine/Plugins/Experimental/ChaosClothAssetEditor/Source/ChaosClothAssetDataflowNodes/Private/ChaosClothAsset/SimulationXPBDAnisoSpringConfigNode.cpp
@@ -1,7 +1,6 @@
 // Copyright Epic Games, Inc. All Rights Reserved.
 
 #include "ChaosClothAsset/SimulationXPBDAnisoSpringConfigNode.h"
-#include "ChaosClothAsset/SimulationBaseConfigNodePrivate.h"
 #include "Chaos/CollectionPropertyFacade.h"
 
 #include UE_INLINE_GENERATED_CPP_BY_NAME(SimulationXPBDAnisoSpringConfigNode)
@@ -18,13 +17,13 @@ FChaosClothAssetSimulationXPBDAnisoSpringConfigNode::FChaosClothAssetSimulationX
 	RegisterInputConnection(&XPBDAnisoSpringWeftScale.WeightMap);
 }
 
-void FChaosClothAssetSimulationXPBDAnisoSpringConfigNode::AddProperties(Dataflow::FContext& Context, ::Chaos::Softs::FCollectionPropertyMutableFacade& Properties) const
+void FChaosClothAssetSimulationXPBDAnisoSpringConfigNode::AddProperties(FPropertyHelper& PropertyHelper) const
 {
-	UE_CHAOS_CLOTHASSET_SIMULATIONCONFIG_SETPROPERTYBOOL(XPBDAnisoSpringUse3dRestLengths);
-	UE_CHAOS_CLOTHASSET_SIMULATIONCONFIG_SETPROPERTYWEIGHTED(XPBDAnisoSpringStiffnessWarp);
-	UE_CHAOS_CLOTHASSET_SIMULATIONCONFIG_SETPROPERTYWEIGHTED(XPBDAnisoSpringStiffnessWeft);
-	UE_CHAOS_CLOTHASSET_SIMULATIONCONFIG_SETPROPERTYWEIGHTED(XPBDAnisoSpringStiffnessBias);
-	UE_CHAOS_CLOTHASSET_SIMULATIONCONFIG_SETPROPERTYWEIGHTED(XPBDAnisoSpringDamping);
-	UE_CHAOS_CLOTHASSET_SIMULATIONCONFIG_SETPROPERTYWEIGHTED(XPBDAnisoSpringWarpScale);
-	UE_CHAOS_CLOTHASSET_SIMULATIONCONFIG_SETPROPERTYWEIGHTED(XPBDAnisoSpringWeftScale);
+	PropertyHelper.SetPropertyBool(this, &bXPBDAnisoSpringUse3dRestLengths, {}, ECollectionPropertyFlags::None);  // Non animatable
+	PropertyHelper.SetPropertyWeighted(this, &XPBDAnisoSpringStiffnessWarp);
+	PropertyHelper.SetPropertyWeighted(this, &XPBDAnisoSpringStiffnessWeft);
+	PropertyHelper.SetPropertyWeighted(this, &XPBDAnisoSpringStiffnessBias);
+	PropertyHelper.SetPropertyWeighted(this, &XPBDAnisoSpringDamping);
+	PropertyHelper.SetPropertyWeighted(this, &XPBDAnisoSpringWarpScale);
+	PropertyHelper.SetPropertyWeighted(this, &XPBDAnisoSpringWeftScale);
 }
diff --git a/Engine/Plugins/Experimental/ChaosClothAssetEditor/Source/ChaosClothAssetDataflowNodes/Private/ChaosClothAsset/SimulationXPBDAnisoStretchConfigNode.cpp b/Engine/Plugins/Experimental/ChaosClothAssetEditor/Source/ChaosClothAssetDataflowNodes/Private/ChaosClothAsset/SimulationXPBDAnisoStretchConfigNode.cpp
--- a/Engine/Plugins/Experimental/ChaosClothAssetEditor/Source/ChaosClothAssetDataflowNodes/Private/ChaosClothAsset/SimulationXPBDAnisoStretchConfigNode.cpp
+++ b/Engine/Plugins/Experimental/ChaosClothAssetEditor/Source/ChaosClothAssetDataflowNodes/Private/ChaosClothAsset/SimulationXPBDAnisoStretchConfigNode.cpp
@@ -1,7 +1,6 @@
 // Copyright Epic Games, Inc. All Rights Reserved.
 
 #include "ChaosClothAsset/SimulationXPBDAnisoStretchConfigNode.h"
-#include "ChaosClothAsset/SimulationBaseConfigNodePrivate.h"
 #include "Chaos/CollectionPropertyFacade.h"
 
 #include UE_INLINE_GENERATED_CPP_BY_NAME(SimulationXPBDAnisoStretchConfigNode)
@@ -18,18 +17,22 @@ FChaosClothAssetSimulationXPBDAnisoStretchConfigNode::FChaosClothAssetSimulation
 	RegisterInputConnection(&XPBDAnisoStretchWeftScale.WeightMap);
 }
 
-void FChaosClothAssetSimulationXPBDAnisoStretchConfigNode::AddProperties(Dataflow::FContext& Context, ::Chaos::Softs::FCollectionPropertyMutableFacade& Properties) const
+void FChaosClothAssetSimulationXPBDAnisoStretchConfigNode::AddProperties(FPropertyHelper& PropertyHelper) const
 {
-	UE_CHAOS_CLOTHASSET_SIMULATIONCONFIG_SETPROPERTYBOOL(XPBDAnisoStretchUse3dRestLengths);
-	UE_CHAOS_CLOTHASSET_SIMULATIONCONFIG_SETPROPERTYWEIGHTEDCHECKED2(
-		XPBDAnisoStretchStiffnessWarp,
-		EdgeSpringStiffness,           // Existing properties to warn against
-		XPBDEdgeSpringStiffness);      //
-	UE_CHAOS_CLOTHASSET_SIMULATIONCONFIG_SETPROPERTYWEIGHTED(XPBDAnisoStretchStiffnessWeft);
-	UE_CHAOS_CLOTHASSET_SIMULATIONCONFIG_SETPROPERTYWEIGHTED(XPBDAnisoStretchStiffnessBias);
-	UE_CHAOS_CLOTHASSET_SIMULATIONCONFIG_SETPROPERTYWEIGHTEDCHECKED1(
-		XPBDAnisoStretchDamping,
-		XPBDEdgeSpringDamping);  // Existing properties to warn against
-	UE_CHAOS_CLOTHASSET_SIMULATIONCONFIG_SETPROPERTYWEIGHTED(XPBDAnisoStretchWarpScale);
-	UE_CHAOS_CLOTHASSET_SIMULATIONCONFIG_SETPROPERTYWEIGHTED(XPBDAnisoStretchWeftScale);
+	PropertyHelper.SetPropertyBool(this, &bXPBDAnisoStretchUse3dRestLengths, {}, ECollectionPropertyFlags::None);  // Non animatable
+	PropertyHelper.SetPropertyWeighted(this,
+		&XPBDAnisoStretchStiffnessWarp,
+		{
+			FName(TEXT("EdgeSpringStiffness")),     // Existing properties to warn against
+			FName(TEXT("XPBDEdgeSpringStiffness"))  //
+		});
+	PropertyHelper.SetPropertyWeighted(this, &XPBDAnisoStretchStiffnessWeft);
+	PropertyHelper.SetPropertyWeighted(this, &XPBDAnisoStretchStiffnessBias);
+	PropertyHelper.SetPropertyWeighted(this,
+		&XPBDAnisoStretchDamping,
+		{
+			FName(TEXT("XPBDEdgeSpringDamping"))  // Existing properties to warn against
+		});
+	PropertyHelper.SetPropertyWeighted(this, &XPBDAnisoStretchWarpScale);
+	PropertyHelper.SetPropertyWeighted(this, &XPBDAnisoStretchWeftScale);
 }
diff --git a/Engine/Plugins/Experimental/ChaosClothAssetEditor/Source/ChaosClothAssetDataflowNodes/Private/ChaosClothAsset/SimulationXPBDAreaSpringConfigNode.cpp b/Engine/Plugins/Experimental/ChaosClothAssetEditor/Source/ChaosClothAssetDataflowNodes/Private/ChaosClothAsset/SimulationXPBDAreaSpringConfigNode.cpp
--- a/Engine/Plugins/Experimental/ChaosClothAssetEditor/Source/ChaosClothAssetDataflowNodes/Private/ChaosClothAsset/SimulationXPBDAreaSpringConfigNode.cpp
+++ b/Engine/Plugins/Experimental/ChaosClothAssetEditor/Source/ChaosClothAssetDataflowNodes/Private/ChaosClothAsset/SimulationXPBDAreaSpringConfigNode.cpp
@@ -1,7 +1,6 @@
 // Copyright Epic Games, Inc. All Rights Reserved.
 
 #include "ChaosClothAsset/SimulationXPBDAreaSpringConfigNode.h"
-#include "ChaosClothAsset/SimulationBaseConfigNodePrivate.h"
 #include "Chaos/CollectionPropertyFacade.h"
 
 #include UE_INLINE_GENERATED_CPP_BY_NAME(SimulationXPBDAreaSpringConfigNode)
@@ -13,9 +12,11 @@ FChaosClothAssetSimulationXPBDAreaSpringConfigNode::FChaosClothAssetSimulationXP
 	RegisterInputConnection(&XPBDAreaSpringStiffness.WeightMap);
 }
 
-void FChaosClothAssetSimulationXPBDAreaSpringConfigNode::AddProperties(Dataflow::FContext& Context, ::Chaos::Softs::FCollectionPropertyMutableFacade& Properties) const
+void FChaosClothAssetSimulationXPBDAreaSpringConfigNode::AddProperties(FPropertyHelper& PropertyHelper) const
 {
-	UE_CHAOS_CLOTHASSET_SIMULATIONCONFIG_SETPROPERTYWEIGHTEDCHECKED1(
-		XPBDAreaSpringStiffness,
-		AreaSpringStiffness);  // Existing properties to warn against
+	PropertyHelper.SetPropertyWeighted(this,
+		&XPBDAreaSpringStiffness,
+		{
+			FName(TEXT("AreaSpringStiffness"))  // Existing properties to warn against
+		});
 }
diff --git a/Engine/Plugins/Experimental/ChaosClothAssetEditor/Source/ChaosClothAssetDataflowNodes/Private/ChaosClothAsset/SimulationXPBDBendingElementConfigNode.cpp b/Engine/Plugins/Experimental/ChaosClothAssetEditor/Source/ChaosClothAssetDataflowNodes/Private/ChaosClothAsset/SimulationXPBDBendingElementConfigNode.cpp
--- a/Engine/Plugins/Experimental/ChaosClothAssetEditor/Source/ChaosClothAssetDataflowNodes/Private/ChaosClothAsset/SimulationXPBDBendingElementConfigNode.cpp
+++ b/Engine/Plugins/Experimental/ChaosClothAssetEditor/Source/ChaosClothAssetDataflowNodes/Private/ChaosClothAsset/SimulationXPBDBendingElementConfigNode.cpp
@@ -1,7 +1,6 @@
 // Copyright Epic Games, Inc. All Rights Reserved.
 
 #include "ChaosClothAsset/SimulationXPBDBendingElementConfigNode.h"
-#include "ChaosClothAsset/SimulationBaseConfigNodePrivate.h"
 #include "Chaos/CollectionPropertyFacade.h"
 
 #include UE_INLINE_GENERATED_CPP_BY_NAME(SimulationXPBDBendingElementConfigNode)
@@ -17,36 +16,51 @@ FChaosClothAssetSimulationXPBDBendingElementConfigNode::FChaosClothAssetSimulati
 	RegisterInputConnection(&XPBDBucklingStiffness.WeightMap);
 }
 
-void FChaosClothAssetSimulationXPBDBendingElementConfigNode::AddProperties(Dataflow::FContext& Context, ::Chaos::Softs::FCollectionPropertyMutableFacade& Properties) const
+void FChaosClothAssetSimulationXPBDBendingElementConfigNode::AddProperties(FPropertyHelper& PropertyHelper) const
 {
-	UE_CHAOS_CLOTHASSET_SIMULATIONCONFIG_SETPROPERTYENUMCHECKED2(
-		XPBDRestAngleType,
-		RestAngleType,            // Existing properties to warn against
-		XPBDAnisoRestAngleType);  //
-	UE_CHAOS_CLOTHASSET_SIMULATIONCONFIG_SETPROPERTYWEIGHTEDCHECKED2(
-		XPBDFlatnessRatio,
-		FlatnessRatio,            // Existing properties to warn against
-		XPBDAnisoFlatnessRatio);  //
-	UE_CHAOS_CLOTHASSET_SIMULATIONCONFIG_SETPROPERTYWEIGHTEDCHECKED2(
-		XPBDRestAngle,
-		RestAngle,            // Existing properties to warn against
-		XPBDAnisoRestAngle);  //
-	UE_CHAOS_CLOTHASSET_SIMULATIONCONFIG_SETPROPERTYWEIGHTEDCHECKED4(
-		XPBDBendingElementStiffness,
-		BendingSpringStiffness,          // Existing properties to warn against
-		BendingElementStiffness,         //
-		XPBDBendingSpringStiffness,      //
-		XPBDAnisoBendingStiffnessWarp);  //
-	UE_CHAOS_CLOTHASSET_SIMULATIONCONFIG_SETPROPERTYWEIGHTEDCHECKED2(
-		XPBDBendingElementDamping,
-		XPBDBendingSpringDamping,  // Existing properties to warn against
-		XPBDAnisoBendingDamping);  //
-	UE_CHAOS_CLOTHASSET_SIMULATIONCONFIG_SETPROPERTYCHECKED2(
-		XPBDBucklingRatio,
-		BucklingRatio,            // Existing properties to warn against
-		XPBDAnisoBucklingRatio);  //
-	UE_CHAOS_CLOTHASSET_SIMULATIONCONFIG_SETPROPERTYWEIGHTEDCHECKED2(
-		XPBDBucklingStiffness,
-		BucklingStiffness,                // Existing properties to warn against
-		XPBDAnisoBucklingStiffnessWarp);  //
+	PropertyHelper.SetPropertyEnum(this,
+		&XPBDRestAngleType,
+		{
+			FName(TEXT("RestAngleType")),          // Existing properties to warn against
+			FName(TEXT("XPBDAnisoRestAngleType"))  //
+		},
+		ECollectionPropertyFlags::None);  // Non animatable
+	PropertyHelper.SetPropertyWeighted(this,
+		&XPBDFlatnessRatio,
+		{
+			FName(TEXT("FlatnessRatio")),          // Existing properties to warn against
+			FName(TEXT("XPBDAnisoFlatnessRatio"))  //
+		});
+	PropertyHelper.SetPropertyWeighted(this,
+		&XPBDRestAngle,
+		{
+			FName(TEXT("RestAngle")),          // Existing properties to warn against
+			FName(TEXT("XPBDAnisoRestAngle"))  //
+		});
+	PropertyHelper.SetPropertyWeighted(this,
+		&XPBDBendingElementStiffness,
+		{
+			FName(TEXT("BendingSpringStiffness")),        // Existing properties to warn against
+			FName(TEXT("BendingElementStiffness")),       //
+			FName(TEXT("XPBDBendingSpringStiffness")),    //
+			FName(TEXT("XPBDAnisoBendingStiffnessWarp"))  //
+		});
+	PropertyHelper.SetPropertyWeighted(this,
+		&XPBDBendingElementDamping,
+		{
+			FName(TEXT("XPBDBendingSpringDamping")),  // Existing properties to warn against
+			FName(TEXT("XPBDAnisoBendingDamping"))    //
+		});
+	PropertyHelper.SetProperty(this,
+		&XPBDBucklingRatio,
+		{
+			FName(TEXT("BucklingRatio")),          // Existing properties to warn against
+			FName(TEXT("XPBDAnisoBucklingRatio"))  //
+		});
+	PropertyHelper.SetPropertyWeighted(this,
+		&XPBDBucklingStiffness,
+		{
+			FName(TEXT("BucklingStiffness")),              // Existing properties to warn against
+			FName(TEXT("XPBDAnisoBucklingStiffnessWarp"))  //
+		});
 }
diff --git a/Engine/Plugins/Experimental/ChaosClothAssetEditor/Source/ChaosClothAssetDataflowNodes/Private/ChaosClothAsset/SimulationXPBDBendingSpringConfigNode.cpp b/Engine/Plugins/Experimental/ChaosClothAssetEditor/Source/ChaosClothAssetDataflowNodes/Private/ChaosClothAsset/SimulationXPBDBendingSpringConfigNode.cpp
--- a/Engine/Plugins/Experimental/ChaosClothAssetEditor/Source/ChaosClothAssetDataflowNodes/Private/ChaosClothAsset/SimulationXPBDBendingSpringConfigNode.cpp
+++ b/Engine/Plugins/Experimental/ChaosClothAssetEditor/Source/ChaosClothAssetDataflowNodes/Private/ChaosClothAsset/SimulationXPBDBendingSpringConfigNode.cpp
@@ -1,7 +1,6 @@
 // Copyright Epic Games, Inc. All Rights Reserved.
 
 #include "ChaosClothAsset/SimulationXPBDBendingSpringConfigNode.h"
-#include "ChaosClothAsset/SimulationBaseConfigNodePrivate.h"
 #include "Chaos/CollectionPropertyFacade.h"
 
 #include UE_INLINE_GENERATED_CPP_BY_NAME(SimulationXPBDBendingSpringConfigNode)
@@ -14,16 +13,20 @@ FChaosClothAssetSimulationXPBDBendingSpringConfigNode::FChaosClothAssetSimulatio
 	RegisterInputConnection(&XPBDBendingSpringDamping.WeightMap);
 }
 
-void FChaosClothAssetSimulationXPBDBendingSpringConfigNode::AddProperties(Dataflow::FContext& Context, ::Chaos::Softs::FCollectionPropertyMutableFacade& Properties) const
+void FChaosClothAssetSimulationXPBDBendingSpringConfigNode::AddProperties(FPropertyHelper& PropertyHelper) const
 {
-	UE_CHAOS_CLOTHASSET_SIMULATIONCONFIG_SETPROPERTYWEIGHTEDCHECKED4(
-		XPBDBendingSpringStiffness,
-		BendingSpringStiffness,          // Existing properties to warn against
-		BendingElementStiffness,         //
-		XPBDBendingElementStiffness,     //
-		XPBDAnisoBendingStiffnessWarp);  //
-	UE_CHAOS_CLOTHASSET_SIMULATIONCONFIG_SETPROPERTYWEIGHTEDCHECKED2(
-		XPBDBendingSpringDamping,
-		XPBDBendingElementDamping,  // Existing properties to warn against
-		XPBDAnisoBendingDamping);   //
+	PropertyHelper.SetPropertyWeighted(this,
+		&XPBDBendingSpringStiffness,
+		{
+			FName(TEXT("BendingSpringStiffness")),        // Existing properties to warn against
+			FName(TEXT("BendingElementStiffness")),       //
+			FName(TEXT("XPBDBendingElementStiffness")),   //
+			FName(TEXT("XPBDAnisoBendingStiffnessWarp"))  //
+		});
+	PropertyHelper.SetPropertyWeighted(this,
+		&XPBDBendingSpringDamping,
+		{
+			FName(TEXT("XPBDBendingElementDamping")),  // Existing properties to warn against
+			FName(TEXT("XPBDAnisoBendingDamping"))     //
+		});
 }
diff --git a/Engine/Plugins/Experimental/ChaosClothAssetEditor/Source/ChaosClothAssetDataflowNodes/Private/ChaosClothAsset/SimulationXPBDEdgeSpringConfigNode.cpp b/Engine/Plugins/Experimental/ChaosClothAssetEditor/Source/ChaosClothAssetDataflowNodes/Private/ChaosClothAsset/SimulationXPBDEdgeSpringConfigNode.cpp
--- a/Engine/Plugins/Experimental/ChaosClothAssetEditor/Source/ChaosClothAssetDataflowNodes/Private/ChaosClothAsset/SimulationXPBDEdgeSpringConfigNode.cpp
+++ b/Engine/Plugins/Experimental/ChaosClothAssetEditor/Source/ChaosClothAssetDataflowNodes/Private/ChaosClothAsset/SimulationXPBDEdgeSpringConfigNode.cpp
@@ -1,7 +1,6 @@
 // Copyright Epic Games, Inc. All Rights Reserved.
 
 #include "ChaosClothAsset/SimulationXPBDEdgeSpringConfigNode.h"
-#include "ChaosClothAsset/SimulationBaseConfigNodePrivate.h"
 #include "Chaos/CollectionPropertyFacade.h"
 
 #include UE_INLINE_GENERATED_CPP_BY_NAME(SimulationXPBDEdgeSpringConfigNode)
@@ -14,13 +13,17 @@ FChaosClothAssetSimulationXPBDEdgeSpringConfigNode::FChaosClothAssetSimulationXP
 	RegisterInputConnection(&XPBDEdgeSpringDamping.WeightMap);
 }
 
-void FChaosClothAssetSimulationXPBDEdgeSpringConfigNode::AddProperties(Dataflow::FContext& Context, ::Chaos::Softs::FCollectionPropertyMutableFacade& Properties) const
+void FChaosClothAssetSimulationXPBDEdgeSpringConfigNode::AddProperties(FPropertyHelper& PropertyHelper) const
 {
-	UE_CHAOS_CLOTHASSET_SIMULATIONCONFIG_SETPROPERTYWEIGHTEDCHECKED2(
-		XPBDEdgeSpringStiffness,
-		EdgeSpringStiffness,             // Existing properties to warn against
-		XPBDAnisoStretchStiffnessWarp);  //
-	UE_CHAOS_CLOTHASSET_SIMULATIONCONFIG_SETPROPERTYWEIGHTEDCHECKED1(
-		XPBDEdgeSpringDamping,
-		XPBDAnisoStretchDamping);  // Existing properties to warn against
+	PropertyHelper.SetPropertyWeighted(this,
+		&XPBDEdgeSpringStiffness,
+		{
+			FName(TEXT("EdgeSpringStiffness")),           // Existing properties to warn against
+			FName(TEXT("XPBDAnisoStretchStiffnessWarp"))  //
+		});
+	PropertyHelper.SetPropertyWeighted(this,
+		&XPBDEdgeSpringDamping,
+		{
+			FName(TEXT("XPBDAnisoStretchDamping"))  // Existing properties to warn against
+		});
 }
diff --git a/Engine/Plugins/Experimental/ChaosClothAssetEditor/Source/ChaosClothAssetDataflowNodes/Public/ChaosClothAsset/SimulationAerodynamicsConfigNode.h b/Engine/Plugins/Experimental/ChaosClothAssetEditor/Source/ChaosClothAssetDataflowNodes/Public/ChaosClothAsset/SimulationAerodynamicsConfigNode.h
--- a/Engine/Plugins/Experimental/ChaosClothAssetEditor/Source/ChaosClothAssetDataflowNodes/Public/ChaosClothAsset/SimulationAerodynamicsConfigNode.h
+++ b/Engine/Plugins/Experimental/ChaosClothAssetEditor/Source/ChaosClothAssetDataflowNodes/Public/ChaosClothAsset/SimulationAerodynamicsConfigNode.h
@@ -50,5 +50,5 @@ public:
 	FChaosClothAssetSimulationAerodynamicsConfigNode(const Dataflow::FNodeParameters& InParam, FGuid InGuid = FGuid::NewGuid());
 
 private:
-	virtual void AddProperties(Dataflow::FContext& Context, ::Chaos::Softs::FCollectionPropertyMutableFacade& Properties) const override;
+	virtual void AddProperties(FPropertyHelper& PropertyHelper) const override;
 };
diff --git a/Engine/Plugins/Experimental/ChaosClothAssetEditor/Source/ChaosClothAssetDataflowNodes/Public/ChaosClothAsset/SimulationAnimDriveConfigNode.h b/Engine/Plugins/Experimental/ChaosClothAssetEditor/Source/ChaosClothAssetDataflowNodes/Public/ChaosClothAsset/SimulationAnimDriveConfigNode.h
--- a/Engine/Plugins/Experimental/ChaosClothAssetEditor/Source/ChaosClothAssetDataflowNodes/Public/ChaosClothAsset/SimulationAnimDriveConfigNode.h
+++ b/Engine/Plugins/Experimental/ChaosClothAssetEditor/Source/ChaosClothAssetDataflowNodes/Public/ChaosClothAsset/SimulationAnimDriveConfigNode.h
@@ -35,5 +35,5 @@ public:
 	FChaosClothAssetSimulationAnimDriveConfigNode(const Dataflow::FNodeParameters& InParam, FGuid InGuid = FGuid::NewGuid());
 
 private:
-	virtual void AddProperties(Dataflow::FContext& Context, ::Chaos::Softs::FCollectionPropertyMutableFacade& Properties) const override;
+	virtual void AddProperties(FPropertyHelper& PropertyHelper) const override;
 };
diff --git a/Engine/Plugins/Experimental/ChaosClothAssetEditor/Source/ChaosClothAssetDataflowNodes/Public/ChaosClothAsset/SimulationBackstopConfigNode.h b/Engine/Plugins/Experimental/ChaosClothAssetEditor/Source/ChaosClothAssetDataflowNodes/Public/ChaosClothAsset/SimulationBackstopConfigNode.h
--- a/Engine/Plugins/Experimental/ChaosClothAssetEditor/Source/ChaosClothAssetDataflowNodes/Public/ChaosClothAsset/SimulationBackstopConfigNode.h
+++ b/Engine/Plugins/Experimental/ChaosClothAssetEditor/Source/ChaosClothAssetDataflowNodes/Public/ChaosClothAsset/SimulationBackstopConfigNode.h
@@ -36,5 +36,5 @@ public:
 	FChaosClothAssetSimulationBackstopConfigNode(const Dataflow::FNodeParameters& InParam, FGuid InGuid = FGuid::NewGuid());
 
 private:
-	virtual void AddProperties(Dataflow::FContext& Context, ::Chaos::Softs::FCollectionPropertyMutableFacade& Properties) const override;
+	virtual void AddProperties(FPropertyHelper& PropertyHelper) const override;
 };
diff --git a/Engine/Plugins/Experimental/ChaosClothAssetEditor/Source/ChaosClothAssetDataflowNodes/Public/ChaosClothAsset/SimulationBaseConfigNode.h b/Engine/Plugins/Experimental/ChaosClothAssetEditor/Source/ChaosClothAssetDataflowNodes/Public/ChaosClothAsset/SimulationBaseConfigNode.h
--- a/Engine/Plugins/Experimental/ChaosClothAssetEditor/Source/ChaosClothAssetDataflowNodes/Public/ChaosClothAsset/SimulationBaseConfigNode.h
+++ b/Engine/Plugins/Experimental/ChaosClothAssetEditor/Source/ChaosClothAssetDataflowNodes/Public/ChaosClothAsset/SimulationBaseConfigNode.h
@@ -2,15 +2,12 @@
 
 #pragma once 
 
+#include "Chaos/CollectionPropertyFacade.h"
+#include "ChaosClothAsset/WeightedValue.h"
 #include "Dataflow/DataflowNode.h"
 #include "GeometryCollection/ManagedArrayCollection.h"
 #include "SimulationBaseConfigNode.generated.h"
 
-namespace Chaos::Softs
-{
-	class FCollectionPropertyMutableFacade;
-}
-
 /**
  * Base abstract class for all cloth asset config nodes.
  * Inherited class must call RegisterCollectionConnections() in constructor to use this base class Collection.
@@ -29,9 +26,57 @@ public:
 	FChaosClothAssetSimulationBaseConfigNode(const Dataflow::FNodeParameters& InParam, FGuid InGuid = FGuid::NewGuid());
 
 protected:
+	using ECollectionPropertyFlags = ::Chaos::Softs::ECollectionPropertyFlags;  // Short hand for property flag enums
+
+	struct FPropertyHelper
+	{
+	public:
+		FPropertyHelper(const FChaosClothAssetSimulationBaseConfigNode& InConfigNode, Dataflow::FContext& InContext, ::Chaos::Softs::FCollectionPropertyMutableFacade& InProperties);
+
+		template<typename PropertyType, TEMPLATE_REQUIRES(::Chaos::Softs::TIsWeightedType<PropertyType>::Value && !std::is_same<PropertyType, bool>::value)>
+		inline int32 SetProperty(const FName& PropertyName, const PropertyType& PropertyValue, const TArray<FName>& SimilarPropertyNames = {}, ECollectionPropertyFlags PropertyFlags = ECollectionPropertyFlags::Animatable);
+
+		template<typename T, typename PropertyType, TEMPLATE_REQUIRES(TIsDerivedFrom<T, FChaosClothAssetSimulationBaseConfigNode>::Value && ::Chaos::Softs::TIsWeightedType<PropertyType>::Value && !std::is_same<PropertyType, bool>::value)>
+		inline int32 SetProperty(const T* ConfigStruct, const PropertyType* PropertyValue, const TArray<FName>& SimilarPropertyNames = {}, ECollectionPropertyFlags PropertyFlags = ECollectionPropertyFlags::Animatable);
+
+		int32 SetPropertyBool(const FName& PropertyName, bool PropertyValue, const TArray<FName>& SimilarPropertyNames = {}, ECollectionPropertyFlags PropertyFlags = ECollectionPropertyFlags::Animatable);
+
+		template<typename T, TEMPLATE_REQUIRES(TIsDerivedFrom<T, FChaosClothAssetSimulationBaseConfigNode>::Value)>
+		inline int32 SetPropertyBool(const T* ConfigStruct, const bool* PropertyValue, const TArray<FName>& SimilarPropertyNames = {}, ECollectionPropertyFlags PropertyFlags = ECollectionPropertyFlags::Animatable);
+
+		template<typename EnumType, TEMPLATE_REQUIRES(TIsEnumClass<EnumType>::Value)>
+		inline int32 SetPropertyEnum(const FName& PropertyName, const EnumType& PropertyValue, const TArray<FName>& SimilarPropertyNames = {}, ECollectionPropertyFlags PropertyFlags = ECollectionPropertyFlags::Animatable);
+
+		template<typename T, typename EnumType, TEMPLATE_REQUIRES(TIsDerivedFrom<T, FChaosClothAssetSimulationBaseConfigNode>::Value && TIsEnumClass<EnumType>::Value)>
+		inline int32 SetPropertyEnum(const T* ConfigStruct, const EnumType* PropertyValue, const TArray<FName>& SimilarPropertyNames = {}, ECollectionPropertyFlags PropertyFlags = ECollectionPropertyFlags::Animatable);
+
+		int32 SetPropertyString(const FName& PropertyName, const FString& PropertyValue, const TArray<FName>& SimilarPropertyNames = {}, ECollectionPropertyFlags PropertyFlags = ECollectionPropertyFlags::Animatable);
+
+		template<typename T, TEMPLATE_REQUIRES(TIsDerivedFrom<T, FChaosClothAssetSimulationBaseConfigNode>::Value)>
+		inline int32 SetPropertyString(const T* ConfigStruct, const FString* PropertyValue, const TArray<FName>& SimilarPropertyNames = {}, ECollectionPropertyFlags PropertyFlags = ECollectionPropertyFlags::Animatable);
+
+		int32 SetPropertyWeighted(const FName& PropertyName, const FChaosClothAssetWeightedValue& PropertyValue, const TArray<FName>& SimilarPropertyNames = {}, ECollectionPropertyFlags PropertyFlags = ECollectionPropertyFlags::None);
+		int32 SetPropertyWeighted(const FName& PropertyName, const FChaosClothAssetWeightedValueNonAnimatable& PropertyValue, const TArray<FName>& SimilarPropertyNames = {}, ECollectionPropertyFlags PropertyFlags = ECollectionPropertyFlags::None);
+		int32 SetPropertyWeighted(const FName& PropertyName, const FChaosClothAssetWeightedValueNonAnimatableNoLowHighRange& PropertyValue, const TArray<FName>& SimilarPropertyNames = {}, ECollectionPropertyFlags PropertyFlags = ECollectionPropertyFlags::None);
+
+		template<typename T, typename WeightedValueType, TEMPLATE_REQUIRES(TIsDerivedFrom<T, FChaosClothAssetSimulationBaseConfigNode>::Value)>
+		inline int32 SetPropertyWeighted(const T* ConfigStruct, const WeightedValueType* PropertyValue, const TArray<FName>& SimilarPropertyNames = {}, ECollectionPropertyFlags PropertyFlags = ECollectionPropertyFlags::None);
+
+	private:
+		template<typename T, TEMPLATE_REQUIRES(TIsDerivedFrom<T, FChaosClothAssetSimulationBaseConfigNode>::Value)>
+		inline FName FindPropertyNameByAddress(const T* ConfigStruct, const void* PropertyAddress);
+
+		const FChaosClothAssetSimulationBaseConfigNode& ConfigNode;
+		Dataflow::FContext& Context;
+		::Chaos::Softs::FCollectionPropertyMutableFacade& Properties;
+	};
+
 	virtual void Evaluate(Dataflow::FContext& Context, const FDataflowOutput* Out) const override;
 
-	virtual void AddProperties(Dataflow::FContext& Context, ::Chaos::Softs::FCollectionPropertyMutableFacade& Properties) const
+	UE_DEPRECATED(5.4, "Use AddProperties(FPropertyHelper&) instead.")
+	virtual void AddProperties(Dataflow::FContext& Context, ::Chaos::Softs::FCollectionPropertyMutableFacade& Properties) const {}
+
+	virtual void AddProperties(struct FPropertyHelper& PropertyHelper) const
 	PURE_VIRTUAL(FChaosClothAssetSimulationBaseConfigNode::AddProperties, );
 
 	/* Override this to do additional node-specific evaluate on the cloth collection output. AddProperties has already been called when this is called. */
@@ -39,11 +84,23 @@ protected:
 
 	void RegisterCollectionConnections();
 
+	UE_DEPRECATED(5.4, "Use FPropertyHelper instead.")
+	int32 AddPropertyHelper(
+		::Chaos::Softs::FCollectionPropertyMutableFacade& Properties,
+		const FName& PropertyName,
+		bool bIsAnimatable,
+		const TArray<FName>& SimilarPropertyNames) const
+	{
+		using namespace Chaos::Softs;
+		const ECollectionPropertyFlags Flags = bIsAnimatable ? ECollectionPropertyFlags::Animatable : ECollectionPropertyFlags::None;
+		AddPropertyHelper(Properties, PropertyName, SimilarPropertyNames, Flags);
+	}
+
 	int32 AddPropertyHelper(
 		::Chaos::Softs::FCollectionPropertyMutableFacade& Properties,
 		const FName& PropertyName,
-		bool bIsAnimatable = true,
-		const TArray<FName>& SimilarPropertyNames = TArray<FName>()) const;
+		const TArray<FName>& SimilarPropertyNames = {},
+		ECollectionPropertyFlags PropertyFlags = ECollectionPropertyFlags::Animatable) const;
 };
 
 template<>
@@ -54,3 +111,110 @@ struct TStructOpsTypeTraits<FChaosClothAssetSimulationBaseConfigNode> : public T
 		WithPureVirtual = true,
 	};
 };
+
+template<typename PropertyType, typename TEnableIf<::Chaos::Softs::TIsWeightedType<PropertyType>::Value && !std::is_same<PropertyType, bool>::value, int>::type>
+inline int32 FChaosClothAssetSimulationBaseConfigNode::FPropertyHelper::SetProperty(
+	const FName& PropertyName,
+	const PropertyType& PropertyValue,
+	const TArray<FName>& SimilarPropertyNames,
+	ECollectionPropertyFlags PropertyFlags)
+{
+	const int32 PropertyKeyIndex = ConfigNode.AddPropertyHelper(Properties, PropertyName, SimilarPropertyNames, PropertyFlags);
+	Properties.SetValue(PropertyKeyIndex, PropertyValue);
+	return PropertyKeyIndex;
+}
+
+template<typename T, typename PropertyType, typename TEnableIf<TIsDerivedFrom<T, FChaosClothAssetSimulationBaseConfigNode>::Value && ::Chaos::Softs::TIsWeightedType<PropertyType>::Value && !std::is_same<PropertyType, bool>::value, int>::type>
+inline int32 FChaosClothAssetSimulationBaseConfigNode::FPropertyHelper::SetProperty(
+	const T* ConfigStruct,
+	const PropertyType* PropertyValue,
+	const TArray<FName>& SimilarPropertyNames,
+	ECollectionPropertyFlags PropertyFlags)
+{
+	check(PropertyValue);
+	const FName PropertyName = FindPropertyNameByAddress(ConfigStruct, PropertyValue);
+	checkf(PropertyName != NAME_None, TEXT("Unknown property."));
+	return SetProperty(PropertyName, *PropertyValue, SimilarPropertyNames, PropertyFlags);
+}
+
+template<typename T, typename TEnableIf<TIsDerivedFrom<T, FChaosClothAssetSimulationBaseConfigNode>::Value, int>::type>
+inline int32 FChaosClothAssetSimulationBaseConfigNode::FPropertyHelper::SetPropertyBool(
+	const T* ConfigStruct,
+	const bool* PropertyValue,
+	const TArray<FName>& SimilarPropertyNames,
+	ECollectionPropertyFlags PropertyFlags)
+{
+	check(PropertyValue);
+	FName PropertyName = FindPropertyNameByAddress(ConfigStruct, PropertyValue);
+	checkf(PropertyName != NAME_None, TEXT("Unknown property."));
+	FString PropertyNameString = PropertyName.ToString();
+	const bool bStartsWithLowerCaseB = PropertyNameString.RemoveFromStart(TEXT("b"), ESearchCase::CaseSensitive);
+	PropertyName = FName(*PropertyNameString);
+	checkf(bStartsWithLowerCaseB, TEXT("A C++ Boolean property name must start with a 'b'."));
+	return SetPropertyBool(PropertyName, *PropertyValue, SimilarPropertyNames, PropertyFlags);
+}
+
+template<typename EnumType, typename TEnableIf<TIsEnumClass<EnumType>::Value, int>::type>
+inline int32 FChaosClothAssetSimulationBaseConfigNode::FPropertyHelper::SetPropertyEnum(
+	const FName& PropertyName,
+	const EnumType& PropertyValue,
+	const TArray<FName>& SimilarPropertyNames,
+	ECollectionPropertyFlags PropertyFlags)
+{
+	const int32 PropertyKeyIndex = ConfigNode.AddPropertyHelper(Properties, PropertyName, SimilarPropertyNames, PropertyFlags);
+	Properties.SetValue(PropertyKeyIndex, (int32)PropertyValue);
+	return PropertyKeyIndex;
+}
+
+template<typename T, typename EnumType, typename TEnableIf<TIsDerivedFrom<T, FChaosClothAssetSimulationBaseConfigNode>::Value && TIsEnumClass<EnumType>::Value, int>::type>
+inline int32 FChaosClothAssetSimulationBaseConfigNode::FPropertyHelper::SetPropertyEnum(
+	const T* ConfigStruct,
+	const EnumType* PropertyValue,
+	const TArray<FName>& SimilarPropertyNames,
+	ECollectionPropertyFlags PropertyFlags)
+{
+	check(PropertyValue);
+	const FName PropertyName = FindPropertyNameByAddress(ConfigStruct, PropertyValue);
+	checkf(PropertyName != NAME_None, TEXT("Unknown property."));
+	return SetPropertyEnum(PropertyName, *PropertyValue, SimilarPropertyNames, PropertyFlags);
+}
+
+template<typename T, typename TEnableIf<TIsDerivedFrom<T, FChaosClothAssetSimulationBaseConfigNode>::Value, int>::type>
+inline int32 FChaosClothAssetSimulationBaseConfigNode::FPropertyHelper::SetPropertyString(
+	const T* ConfigStruct,
+	const FString* PropertyValue,
+	const TArray<FName>& SimilarPropertyNames,
+	ECollectionPropertyFlags PropertyFlags)
+{
+	check(PropertyValue);
+	const FName PropertyName = FindPropertyNameByAddress(ConfigStruct, PropertyValue);
+	checkf(PropertyName != NAME_None, TEXT("Unknown property."));
+	return SetPropertyString(PropertyName, *PropertyValue, SimilarPropertyNames, PropertyFlags);
+}
+
+template<typename T, typename WeightedValueType, typename TEnableIf<TIsDerivedFrom<T, FChaosClothAssetSimulationBaseConfigNode>::Value, int>::type>
+inline int32 FChaosClothAssetSimulationBaseConfigNode::FPropertyHelper::SetPropertyWeighted(
+	const T* ConfigStruct,
+	const WeightedValueType* PropertyValue,
+	const TArray<FName>& SimilarPropertyNames,
+	ECollectionPropertyFlags PropertyFlags)
+{
+	check(PropertyValue);
+	const FName PropertyName = FindPropertyNameByAddress(ConfigStruct, PropertyValue);
+	checkf(PropertyName != NAME_None, TEXT("Unknown property."));
+	return SetPropertyWeighted(PropertyName, *PropertyValue, SimilarPropertyNames, PropertyFlags);
+}
+
+template<typename T, typename TEnableIf<TIsDerivedFrom<T, FChaosClothAssetSimulationBaseConfigNode>::Value, int>::type>
+inline FName FChaosClothAssetSimulationBaseConfigNode::FPropertyHelper::FindPropertyNameByAddress(const T* ConfigStruct, const void* PropertyAddress)
+{
+	for (TFieldIterator<FProperty> FieldIt(T::StaticStruct()); FieldIt; ++FieldIt)
+	{
+		const FProperty* const Property = *FieldIt;
+		if (PropertyAddress == Property->ContainerPtrToValuePtr<void>(ConfigStruct))
+		{
+			return Property->GetFName();
+		}
+	}
+	return NAME_None;
+}
diff --git a/Engine/Plugins/Experimental/ChaosClothAssetEditor/Source/ChaosClothAssetDataflowNodes/Public/ChaosClothAsset/SimulationCollisionConfigNode.h b/Engine/Plugins/Experimental/ChaosClothAssetEditor/Source/ChaosClothAssetDataflowNodes/Public/ChaosClothAsset/SimulationCollisionConfigNode.h
--- a/Engine/Plugins/Experimental/ChaosClothAssetEditor/Source/ChaosClothAssetDataflowNodes/Public/ChaosClothAsset/SimulationCollisionConfigNode.h
+++ b/Engine/Plugins/Experimental/ChaosClothAssetEditor/Source/ChaosClothAssetDataflowNodes/Public/ChaosClothAsset/SimulationCollisionConfigNode.h
@@ -35,5 +35,5 @@ public:
 	FChaosClothAssetSimulationCollisionConfigNode(const Dataflow::FNodeParameters& InParam, FGuid InGuid = FGuid::NewGuid());
 
 private:
-	virtual void AddProperties(Dataflow::FContext& Context, ::Chaos::Softs::FCollectionPropertyMutableFacade& Properties) const override;
+	virtual void AddProperties(FPropertyHelper& PropertyHelper) const override;
 };
diff --git a/Engine/Plugins/Experimental/ChaosClothAssetEditor/Source/ChaosClothAssetDataflowNodes/Public/ChaosClothAsset/SimulationDampingConfigNode.h b/Engine/Plugins/Experimental/ChaosClothAssetEditor/Source/ChaosClothAssetDataflowNodes/Public/ChaosClothAsset/SimulationDampingConfigNode.h
--- a/Engine/Plugins/Experimental/ChaosClothAssetEditor/Source/ChaosClothAssetDataflowNodes/Public/ChaosClothAsset/SimulationDampingConfigNode.h
+++ b/Engine/Plugins/Experimental/ChaosClothAssetEditor/Source/ChaosClothAssetDataflowNodes/Public/ChaosClothAsset/SimulationDampingConfigNode.h
@@ -32,5 +32,5 @@ public:
 	FChaosClothAssetSimulationDampingConfigNode(const Dataflow::FNodeParameters& InParam, FGuid InGuid = FGuid::NewGuid());
 
 private:
-	virtual void AddProperties(Dataflow::FContext& Context, ::Chaos::Softs::FCollectionPropertyMutableFacade& Properties) const override;
+	virtual void AddProperties(FPropertyHelper& PropertyHelper) const override;
 };
diff --git a/Engine/Plugins/Experimental/ChaosClothAssetEditor/Source/ChaosClothAssetDataflowNodes/Public/ChaosClothAsset/SimulationGravityConfigNode.h b/Engine/Plugins/Experimental/ChaosClothAssetEditor/Source/ChaosClothAssetDataflowNodes/Public/ChaosClothAsset/SimulationGravityConfigNode.h
--- a/Engine/Plugins/Experimental/ChaosClothAssetEditor/Source/ChaosClothAssetDataflowNodes/Public/ChaosClothAsset/SimulationGravityConfigNode.h
+++ b/Engine/Plugins/Experimental/ChaosClothAssetEditor/Source/ChaosClothAssetDataflowNodes/Public/ChaosClothAsset/SimulationGravityConfigNode.h
@@ -28,5 +28,5 @@ public:
 	FChaosClothAssetSimulationGravityConfigNode(const Dataflow::FNodeParameters& InParam, FGuid InGuid = FGuid::NewGuid());
 
 private:
-	virtual void AddProperties(Dataflow::FContext& Context, ::Chaos::Softs::FCollectionPropertyMutableFacade& Properties) const override;
+	virtual void AddProperties(FPropertyHelper& PropertyHelper) const override;
 };
diff --git a/Engine/Plugins/Experimental/ChaosClothAssetEditor/Source/ChaosClothAssetDataflowNodes/Public/ChaosClothAsset/SimulationLongRangeAttachmentConfigNode.h b/Engine/Plugins/Experimental/ChaosClothAssetEditor/Source/ChaosClothAssetDataflowNodes/Public/ChaosClothAsset/SimulationLongRangeAttachmentConfigNode.h
--- a/Engine/Plugins/Experimental/ChaosClothAssetEditor/Source/ChaosClothAssetDataflowNodes/Public/ChaosClothAsset/SimulationLongRangeAttachmentConfigNode.h
+++ b/Engine/Plugins/Experimental/ChaosClothAssetEditor/Source/ChaosClothAssetDataflowNodes/Public/ChaosClothAsset/SimulationLongRangeAttachmentConfigNode.h
@@ -51,7 +51,7 @@ public:
 	FChaosClothAssetSimulationLongRangeAttachmentConfigNode(const Dataflow::FNodeParameters& InParam, FGuid InGuid = FGuid::NewGuid());
 
 private:
-	virtual void AddProperties(Dataflow::FContext& Context, ::Chaos::Softs::FCollectionPropertyMutableFacade& Properties) const override;
+	virtual void AddProperties(FPropertyHelper& PropertyHelper) const override;
 
 	virtual void EvaluateClothCollection(Dataflow::FContext& Context, const TSharedRef<FManagedArrayCollection>& ClothCollection) const override;
 };
diff --git a/Engine/Plugins/Experimental/ChaosClothAssetEditor/Source/ChaosClothAssetDataflowNodes/Public/ChaosClothAsset/SimulationMassConfigNode.h b/Engine/Plugins/Experimental/ChaosClothAssetEditor/Source/ChaosClothAssetDataflowNodes/Public/ChaosClothAsset/SimulationMassConfigNode.h
--- a/Engine/Plugins/Experimental/ChaosClothAssetEditor/Source/ChaosClothAssetDataflowNodes/Public/ChaosClothAsset/SimulationMassConfigNode.h
+++ b/Engine/Plugins/Experimental/ChaosClothAssetEditor/Source/ChaosClothAssetDataflowNodes/Public/ChaosClothAsset/SimulationMassConfigNode.h
@@ -51,5 +51,5 @@ public:
 	FChaosClothAssetSimulationMassConfigNode(const Dataflow::FNodeParameters& InParam, FGuid InGuid = FGuid::NewGuid());
 
 private:
-	virtual void AddProperties(Dataflow::FContext& Context, ::Chaos::Softs::FCollectionPropertyMutableFacade& Properties) const override;
+	virtual void AddProperties(FPropertyHelper& PropertyHelper) const override;
 };
diff --git a/Engine/Plugins/Experimental/ChaosClothAssetEditor/Source/ChaosClothAssetDataflowNodes/Public/ChaosClothAsset/SimulationMaxDistanceConfigNode.h b/Engine/Plugins/Experimental/ChaosClothAssetEditor/Source/ChaosClothAssetDataflowNodes/Public/ChaosClothAsset/SimulationMaxDistanceConfigNode.h
--- a/Engine/Plugins/Experimental/ChaosClothAssetEditor/Source/ChaosClothAssetDataflowNodes/Public/ChaosClothAsset/SimulationMaxDistanceConfigNode.h
+++ b/Engine/Plugins/Experimental/ChaosClothAssetEditor/Source/ChaosClothAssetDataflowNodes/Public/ChaosClothAsset/SimulationMaxDistanceConfigNode.h
@@ -27,5 +27,5 @@ public:
 	FChaosClothAssetSimulationMaxDistanceConfigNode(const Dataflow::FNodeParameters& InParam, FGuid InGuid = FGuid::NewGuid());
 
 private:
-	virtual void AddProperties(Dataflow::FContext& Context, ::Chaos::Softs::FCollectionPropertyMutableFacade& Properties) const override;
+	virtual void AddProperties(FPropertyHelper& PropertyHelper) const override;
 };
diff --git a/Engine/Plugins/Experimental/ChaosClothAssetEditor/Source/ChaosClothAssetDataflowNodes/Public/ChaosClothAsset/SimulationPBDAreaSpringConfigNode.h b/Engine/Plugins/Experimental/ChaosClothAssetEditor/Source/ChaosClothAssetDataflowNodes/Public/ChaosClothAsset/SimulationPBDAreaSpringConfigNode.h
--- a/Engine/Plugins/Experimental/ChaosClothAssetEditor/Source/ChaosClothAssetDataflowNodes/Public/ChaosClothAsset/SimulationPBDAreaSpringConfigNode.h
+++ b/Engine/Plugins/Experimental/ChaosClothAssetEditor/Source/ChaosClothAssetDataflowNodes/Public/ChaosClothAsset/SimulationPBDAreaSpringConfigNode.h
@@ -27,5 +27,5 @@ public:
 	FChaosClothAssetSimulationPBDAreaSpringConfigNode(const Dataflow::FNodeParameters& InParam, FGuid InGuid = FGuid::NewGuid());
 
 private:
-	virtual void AddProperties(Dataflow::FContext& Context, ::Chaos::Softs::FCollectionPropertyMutableFacade& Properties) const override;
+	virtual void AddProperties(FPropertyHelper& PropertyHelper) const override;
 };
diff --git a/Engine/Plugins/Experimental/ChaosClothAssetEditor/Source/ChaosClothAssetDataflowNodes/Public/ChaosClothAsset/SimulationPBDBendingElementConfigNode.h b/Engine/Plugins/Experimental/ChaosClothAssetEditor/Source/ChaosClothAssetDataflowNodes/Public/ChaosClothAsset/SimulationPBDBendingElementConfigNode.h
--- a/Engine/Plugins/Experimental/ChaosClothAssetEditor/Source/ChaosClothAssetDataflowNodes/Public/ChaosClothAsset/SimulationPBDBendingElementConfigNode.h
+++ b/Engine/Plugins/Experimental/ChaosClothAssetEditor/Source/ChaosClothAssetDataflowNodes/Public/ChaosClothAsset/SimulationPBDBendingElementConfigNode.h
@@ -85,5 +85,5 @@ public:
 	FChaosClothAssetSimulationPBDBendingElementConfigNode(const Dataflow::FNodeParameters& InParam, FGuid InGuid = FGuid::NewGuid());
 
 private:
-	virtual void AddProperties(Dataflow::FContext& Context, ::Chaos::Softs::FCollectionPropertyMutableFacade& Properties) const override;
+	virtual void AddProperties(FPropertyHelper& PropertyHelper) const override;
 };
diff --git a/Engine/Plugins/Experimental/ChaosClothAssetEditor/Source/ChaosClothAssetDataflowNodes/Public/ChaosClothAsset/SimulationPBDBendingSpringConfigNode.h b/Engine/Plugins/Experimental/ChaosClothAssetEditor/Source/ChaosClothAssetDataflowNodes/Public/ChaosClothAsset/SimulationPBDBendingSpringConfigNode.h
--- a/Engine/Plugins/Experimental/ChaosClothAssetEditor/Source/ChaosClothAssetDataflowNodes/Public/ChaosClothAsset/SimulationPBDBendingSpringConfigNode.h
+++ b/Engine/Plugins/Experimental/ChaosClothAssetEditor/Source/ChaosClothAssetDataflowNodes/Public/ChaosClothAsset/SimulationPBDBendingSpringConfigNode.h
@@ -27,5 +27,5 @@ public:
 	FChaosClothAssetSimulationPBDBendingSpringConfigNode(const Dataflow::FNodeParameters& InParam, FGuid InGuid = FGuid::NewGuid());
 
 private:
-	virtual void AddProperties(Dataflow::FContext& Context, ::Chaos::Softs::FCollectionPropertyMutableFacade& Properties) const override;
+	virtual void AddProperties(FPropertyHelper& PropertyHelper) const override;
 };
diff --git a/Engine/Plugins/Experimental/ChaosClothAssetEditor/Source/ChaosClothAssetDataflowNodes/Public/ChaosClothAsset/SimulationPBDEdgeSpringConfigNode.h b/Engine/Plugins/Experimental/ChaosClothAssetEditor/Source/ChaosClothAssetDataflowNodes/Public/ChaosClothAsset/SimulationPBDEdgeSpringConfigNode.h
--- a/Engine/Plugins/Experimental/ChaosClothAssetEditor/Source/ChaosClothAssetDataflowNodes/Public/ChaosClothAsset/SimulationPBDEdgeSpringConfigNode.h
+++ b/Engine/Plugins/Experimental/ChaosClothAssetEditor/Source/ChaosClothAssetDataflowNodes/Public/ChaosClothAsset/SimulationPBDEdgeSpringConfigNode.h
@@ -27,5 +27,5 @@ public:
 	FChaosClothAssetSimulationPBDEdgeSpringConfigNode(const Dataflow::FNodeParameters& InParam, FGuid InGuid = FGuid::NewGuid());
 
 private:
-	virtual void AddProperties(Dataflow::FContext& Context, ::Chaos::Softs::FCollectionPropertyMutableFacade& Properties) const override;
+	virtual void AddProperties(FPropertyHelper& PropertyHelper) const override;
 };
diff --git a/Engine/Plugins/Experimental/ChaosClothAssetEditor/Source/ChaosClothAssetDataflowNodes/Public/ChaosClothAsset/SimulationPressureConfigNode.h b/Engine/Plugins/Experimental/ChaosClothAssetEditor/Source/ChaosClothAssetDataflowNodes/Public/ChaosClothAsset/SimulationPressureConfigNode.h
--- a/Engine/Plugins/Experimental/ChaosClothAssetEditor/Source/ChaosClothAssetDataflowNodes/Public/ChaosClothAsset/SimulationPressureConfigNode.h
+++ b/Engine/Plugins/Experimental/ChaosClothAssetEditor/Source/ChaosClothAssetDataflowNodes/Public/ChaosClothAsset/SimulationPressureConfigNode.h
@@ -26,5 +26,5 @@ public:
 	FChaosClothAssetSimulationPressureConfigNode(const Dataflow::FNodeParameters& InParam, FGuid InGuid = FGuid::NewGuid());
 
 private:
-	virtual void AddProperties(Dataflow::FContext& Context, ::Chaos::Softs::FCollectionPropertyMutableFacade& Properties) const override;
+	virtual void AddProperties(FPropertyHelper& PropertyHelper) const override;
 };
diff --git a/Engine/Plugins/Experimental/ChaosClothAssetEditor/Source/ChaosClothAssetDataflowNodes/Public/ChaosClothAsset/SimulationSelfCollisionConfigNode.h b/Engine/Plugins/Experimental/ChaosClothAssetEditor/Source/ChaosClothAssetDataflowNodes/Public/ChaosClothAsset/SimulationSelfCollisionConfigNode.h
--- a/Engine/Plugins/Experimental/ChaosClothAssetEditor/Source/ChaosClothAssetDataflowNodes/Public/ChaosClothAsset/SimulationSelfCollisionConfigNode.h
+++ b/Engine/Plugins/Experimental/ChaosClothAssetEditor/Source/ChaosClothAssetDataflowNodes/Public/ChaosClothAsset/SimulationSelfCollisionConfigNode.h
@@ -52,5 +52,5 @@ public:
 	FChaosClothAssetSimulationSelfCollisionConfigNode(const Dataflow::FNodeParameters& InParam, FGuid InGuid = FGuid::NewGuid());
 
 private:
-	virtual void AddProperties(Dataflow::FContext& Context, ::Chaos::Softs::FCollectionPropertyMutableFacade& Properties) const override;
+	virtual void AddProperties(FPropertyHelper& PropertyHelper) const override;
 };
diff --git a/Engine/Plugins/Experimental/ChaosClothAssetEditor/Source/ChaosClothAssetDataflowNodes/Public/ChaosClothAsset/SimulationSelfCollisionSpheresConfigNode.h b/Engine/Plugins/Experimental/ChaosClothAssetEditor/Source/ChaosClothAssetDataflowNodes/Public/ChaosClothAsset/SimulationSelfCollisionSpheresConfigNode.h
--- a/Engine/Plugins/Experimental/ChaosClothAssetEditor/Source/ChaosClothAssetDataflowNodes/Public/ChaosClothAsset/SimulationSelfCollisionSpheresConfigNode.h
+++ b/Engine/Plugins/Experimental/ChaosClothAssetEditor/Source/ChaosClothAssetDataflowNodes/Public/ChaosClothAsset/SimulationSelfCollisionSpheresConfigNode.h
@@ -31,7 +31,7 @@ public:
 	FChaosClothAssetSimulationSelfCollisionSpheresConfigNode(const Dataflow::FNodeParameters& InParam, FGuid InGuid = FGuid::NewGuid());
 
 private:
-	virtual void AddProperties(Dataflow::FContext& Context, ::Chaos::Softs::FCollectionPropertyMutableFacade& Properties) const override;
+	virtual void AddProperties(FPropertyHelper& PropertyHelper) const override;
 
 	virtual void EvaluateClothCollection(Dataflow::FContext& Context, const TSharedRef<FManagedArrayCollection>& ClothCollection) const override;
 };
diff --git a/Engine/Plugins/Experimental/ChaosClothAssetEditor/Source/ChaosClothAssetDataflowNodes/Public/ChaosClothAsset/SimulationSolverConfigNode.h b/Engine/Plugins/Experimental/ChaosClothAssetEditor/Source/ChaosClothAssetDataflowNodes/Public/ChaosClothAsset/SimulationSolverConfigNode.h
--- a/Engine/Plugins/Experimental/ChaosClothAssetEditor/Source/ChaosClothAssetDataflowNodes/Public/ChaosClothAsset/SimulationSolverConfigNode.h
+++ b/Engine/Plugins/Experimental/ChaosClothAssetEditor/Source/ChaosClothAssetDataflowNodes/Public/ChaosClothAsset/SimulationSolverConfigNode.h
@@ -83,5 +83,5 @@ public:
 	FChaosClothAssetSimulationSolverConfigNode(const Dataflow::FNodeParameters& InParam, FGuid InGuid = FGuid::NewGuid());
 
 private:
-	virtual void AddProperties(Dataflow::FContext& Context, ::Chaos::Softs::FCollectionPropertyMutableFacade& Properties) const override;
+	virtual void AddProperties(FPropertyHelper& PropertyHelper) const override;
 };
diff --git a/Engine/Plugins/Experimental/ChaosClothAssetEditor/Source/ChaosClothAssetDataflowNodes/Public/ChaosClothAsset/SimulationVelocityScaleConfigNode.h b/Engine/Plugins/Experimental/ChaosClothAssetEditor/Source/ChaosClothAssetDataflowNodes/Public/ChaosClothAsset/SimulationVelocityScaleConfigNode.h
--- a/Engine/Plugins/Experimental/ChaosClothAssetEditor/Source/ChaosClothAssetDataflowNodes/Public/ChaosClothAsset/SimulationVelocityScaleConfigNode.h
+++ b/Engine/Plugins/Experimental/ChaosClothAssetEditor/Source/ChaosClothAssetDataflowNodes/Public/ChaosClothAsset/SimulationVelocityScaleConfigNode.h
@@ -41,5 +41,5 @@ public:
 	FChaosClothAssetSimulationVelocityScaleConfigNode(const Dataflow::FNodeParameters& InParam, FGuid InGuid = FGuid::NewGuid());
 
 private:
-	virtual void AddProperties(Dataflow::FContext& Context, ::Chaos::Softs::FCollectionPropertyMutableFacade& Properties) const override;
+	virtual void AddProperties(FPropertyHelper& PropertyHelper) const override;
 };
diff --git a/Engine/Plugins/Experimental/ChaosClothAssetEditor/Source/ChaosClothAssetDataflowNodes/Public/ChaosClothAsset/SimulationXPBDAnisoBendingConfigNode.h b/Engine/Plugins/Experimental/ChaosClothAssetEditor/Source/ChaosClothAssetDataflowNodes/Public/ChaosClothAsset/SimulationXPBDAnisoBendingConfigNode.h
--- a/Engine/Plugins/Experimental/ChaosClothAssetEditor/Source/ChaosClothAssetDataflowNodes/Public/ChaosClothAsset/SimulationXPBDAnisoBendingConfigNode.h
+++ b/Engine/Plugins/Experimental/ChaosClothAssetEditor/Source/ChaosClothAssetDataflowNodes/Public/ChaosClothAsset/SimulationXPBDAnisoBendingConfigNode.h
@@ -125,5 +125,5 @@ public:
 	FChaosClothAssetSimulationXPBDAnisoBendingConfigNode(const Dataflow::FNodeParameters& InParam, FGuid InGuid = FGuid::NewGuid());
 
 private:
-	virtual void AddProperties(Dataflow::FContext& Context, ::Chaos::Softs::FCollectionPropertyMutableFacade& Properties) const override;
+	virtual void AddProperties(FPropertyHelper& PropertyHelper) const override;
 };
diff --git a/Engine/Plugins/Experimental/ChaosClothAssetEditor/Source/ChaosClothAssetDataflowNodes/Public/ChaosClothAsset/SimulationXPBDAnisoSpringConfigNode.h b/Engine/Plugins/Experimental/ChaosClothAssetEditor/Source/ChaosClothAssetDataflowNodes/Public/ChaosClothAsset/SimulationXPBDAnisoSpringConfigNode.h
--- a/Engine/Plugins/Experimental/ChaosClothAssetEditor/Source/ChaosClothAssetDataflowNodes/Public/ChaosClothAsset/SimulationXPBDAnisoSpringConfigNode.h
+++ b/Engine/Plugins/Experimental/ChaosClothAssetEditor/Source/ChaosClothAssetDataflowNodes/Public/ChaosClothAsset/SimulationXPBDAnisoSpringConfigNode.h
@@ -75,5 +75,5 @@ public:
 	FChaosClothAssetSimulationXPBDAnisoSpringConfigNode(const Dataflow::FNodeParameters& InParam, FGuid InGuid = FGuid::NewGuid());
 
 private:
-	virtual void AddProperties(Dataflow::FContext& Context, ::Chaos::Softs::FCollectionPropertyMutableFacade& Properties) const override;
+	virtual void AddProperties(FPropertyHelper& PropertyHelper) const override;
 };
diff --git a/Engine/Plugins/Experimental/ChaosClothAssetEditor/Source/ChaosClothAssetDataflowNodes/Public/ChaosClothAsset/SimulationXPBDAnisoStretchConfigNode.h b/Engine/Plugins/Experimental/ChaosClothAssetEditor/Source/ChaosClothAssetDataflowNodes/Public/ChaosClothAsset/SimulationXPBDAnisoStretchConfigNode.h
--- a/Engine/Plugins/Experimental/ChaosClothAssetEditor/Source/ChaosClothAssetDataflowNodes/Public/ChaosClothAsset/SimulationXPBDAnisoStretchConfigNode.h
+++ b/Engine/Plugins/Experimental/ChaosClothAssetEditor/Source/ChaosClothAssetDataflowNodes/Public/ChaosClothAsset/SimulationXPBDAnisoStretchConfigNode.h
@@ -75,5 +75,5 @@ public:
 	FChaosClothAssetSimulationXPBDAnisoStretchConfigNode(const Dataflow::FNodeParameters& InParam, FGuid InGuid = FGuid::NewGuid());
 
 private:
-	virtual void AddProperties(Dataflow::FContext& Context, ::Chaos::Softs::FCollectionPropertyMutableFacade& Properties) const override;
+	virtual void AddProperties(FPropertyHelper& PropertyHelper) const override;
 };
diff --git a/Engine/Plugins/Experimental/ChaosClothAssetEditor/Source/ChaosClothAssetDataflowNodes/Public/ChaosClothAsset/SimulationXPBDAreaSpringConfigNode.h b/Engine/Plugins/Experimental/ChaosClothAssetEditor/Source/ChaosClothAssetDataflowNodes/Public/ChaosClothAsset/SimulationXPBDAreaSpringConfigNode.h
--- a/Engine/Plugins/Experimental/ChaosClothAssetEditor/Source/ChaosClothAssetDataflowNodes/Public/ChaosClothAsset/SimulationXPBDAreaSpringConfigNode.h
+++ b/Engine/Plugins/Experimental/ChaosClothAssetEditor/Source/ChaosClothAssetDataflowNodes/Public/ChaosClothAsset/SimulationXPBDAreaSpringConfigNode.h
@@ -27,5 +27,5 @@ public:
 	FChaosClothAssetSimulationXPBDAreaSpringConfigNode(const Dataflow::FNodeParameters& InParam, FGuid InGuid = FGuid::NewGuid());
 
 private:
-	virtual void AddProperties(Dataflow::FContext& Context, ::Chaos::Softs::FCollectionPropertyMutableFacade& Properties) const override;
+	virtual void AddProperties(FPropertyHelper& PropertyHelper) const override;
 };
diff --git a/Engine/Plugins/Experimental/ChaosClothAssetEditor/Source/ChaosClothAssetDataflowNodes/Public/ChaosClothAsset/SimulationXPBDBendingElementConfigNode.h b/Engine/Plugins/Experimental/ChaosClothAssetEditor/Source/ChaosClothAssetDataflowNodes/Public/ChaosClothAsset/SimulationXPBDBendingElementConfigNode.h
--- a/Engine/Plugins/Experimental/ChaosClothAssetEditor/Source/ChaosClothAssetDataflowNodes/Public/ChaosClothAsset/SimulationXPBDBendingElementConfigNode.h
+++ b/Engine/Plugins/Experimental/ChaosClothAssetEditor/Source/ChaosClothAssetDataflowNodes/Public/ChaosClothAsset/SimulationXPBDBendingElementConfigNode.h
@@ -83,5 +83,5 @@ public:
 	FChaosClothAssetSimulationXPBDBendingElementConfigNode(const Dataflow::FNodeParameters& InParam, FGuid InGuid = FGuid::NewGuid());
 
 private:
-	virtual void AddProperties(Dataflow::FContext& Context, ::Chaos::Softs::FCollectionPropertyMutableFacade& Properties) const override;
+	virtual void AddProperties(FPropertyHelper& PropertyHelper) const override;
 };
diff --git a/Engine/Plugins/Experimental/ChaosClothAssetEditor/Source/ChaosClothAssetDataflowNodes/Public/ChaosClothAsset/SimulationXPBDBendingSpringConfigNode.h b/Engine/Plugins/Experimental/ChaosClothAssetEditor/Source/ChaosClothAssetDataflowNodes/Public/ChaosClothAsset/SimulationXPBDBendingSpringConfigNode.h
--- a/Engine/Plugins/Experimental/ChaosClothAssetEditor/Source/ChaosClothAssetDataflowNodes/Public/ChaosClothAsset/SimulationXPBDBendingSpringConfigNode.h
+++ b/Engine/Plugins/Experimental/ChaosClothAssetEditor/Source/ChaosClothAssetDataflowNodes/Public/ChaosClothAsset/SimulationXPBDBendingSpringConfigNode.h
@@ -35,5 +35,5 @@ public:
 	FChaosClothAssetSimulationXPBDBendingSpringConfigNode(const Dataflow::FNodeParameters& InParam, FGuid InGuid = FGuid::NewGuid());
 
 private:
-	virtual void AddProperties(Dataflow::FContext& Context, ::Chaos::Softs::FCollectionPropertyMutableFacade& Properties) const override;
+	virtual void AddProperties(FPropertyHelper& PropertyHelper) const override;
 };
diff --git a/Engine/Plugins/Experimental/ChaosClothAssetEditor/Source/ChaosClothAssetDataflowNodes/Public/ChaosClothAsset/SimulationXPBDEdgeSpringConfigNode.h b/Engine/Plugins/Experimental/ChaosClothAssetEditor/Source/ChaosClothAssetDataflowNodes/Public/ChaosClothAsset/SimulationXPBDEdgeSpringConfigNode.h
--- a/Engine/Plugins/Experimental/ChaosClothAssetEditor/Source/ChaosClothAssetDataflowNodes/Public/ChaosClothAsset/SimulationXPBDEdgeSpringConfigNode.h
+++ b/Engine/Plugins/Experimental/ChaosClothAssetEditor/Source/ChaosClothAssetDataflowNodes/Public/ChaosClothAsset/SimulationXPBDEdgeSpringConfigNode.h
@@ -35,5 +35,5 @@ public:
 	FChaosClothAssetSimulationXPBDEdgeSpringConfigNode(const Dataflow::FNodeParameters& InParam, FGuid InGuid = FGuid::NewGuid());
 
 private:
-	virtual void AddProperties(Dataflow::FContext& Context, ::Chaos::Softs::FCollectionPropertyMutableFacade& Properties) const override;
+	virtual void AddProperties(FPropertyHelper& PropertyHelper) const override;
 };
diff --git a/Engine/Source/Runtime/Experimental/Chaos/Public/Chaos/CollectionPropertyFacade.h b/Engine/Source/Runtime/Experimental/Chaos/Public/Chaos/CollectionPropertyFacade.h
--- a/Engine/Source/Runtime/Experimental/Chaos/Public/Chaos/CollectionPropertyFacade.h
+++ b/Engine/Source/Runtime/Experimental/Chaos/Public/Chaos/CollectionPropertyFacade.h
@@ -277,7 +277,7 @@ PRAGMA_ENABLE_DEPRECATION_WARNINGS
 		void SetStringValue(int32 KeyIndex, const FString& Value) { if (GetStringValueArray()[KeyIndex] != Value) { GetStringValueArray()[KeyIndex] = Value; SetStringDirty(KeyIndex); } }
 
 		/** SetFlags cannot be used to remove Dirty, StringDirty, or Intrinsic flags. Use ClearDirtyFlags to remove dirty flags. */
-		void SetFlags(int32 KeyIndex, ECollectionPropertyFlags Flags);
+		CHAOS_API void SetFlags(int32 KeyIndex, ECollectionPropertyFlags Flags);
 		UE_DEPRECATED(5.3, "Use SetFlags(int32, ECollectionPropertyFlags) instead.")
 		void SetFlags(int32 KeyIndex, uint8 Flags) { return SetFlags(KeyIndex, (ECollectionPropertyFlags)Flags); }
 
