CIS - Fixed non void function should return a value warning in SimulationBaseConfigNode.h.

diff --git a/Engine/Plugins/Experimental/ChaosClothAssetEditor/Source/ChaosClothAssetDataflowNodes/Public/ChaosClothAsset/SimulationBaseConfigNode.h b/Engine/Plugins/Experimental/ChaosClothAssetEditor/Source/ChaosClothAssetDataflowNodes/Public/ChaosClothAsset/SimulationBaseConfigNode.h
--- a/Engine/Plugins/Experimental/ChaosClothAssetEditor/Source/ChaosClothAssetDataflowNodes/Public/ChaosClothAsset/SimulationBaseConfigNode.h
+++ b/Engine/Plugins/Experimental/ChaosClothAssetEditor/Source/ChaosClothAssetDataflowNodes/Public/ChaosClothAsset/SimulationBaseConfigNode.h
@@ -93,7 +93,7 @@ protected:
 	{
 		using namespace Chaos::Softs;
 		const ECollectionPropertyFlags Flags = bIsAnimatable ? ECollectionPropertyFlags::Animatable : ECollectionPropertyFlags::None;
-		AddPropertyHelper(Properties, PropertyName, SimilarPropertyNames, Flags);
+		return AddPropertyHelper(Properties, PropertyName, SimilarPropertyNames, Flags);
 	}
 
 	int32 AddPropertyHelper(
