Remesh node fixes
- don't split mesh edges that participate in more than one seam edges (for now)
- add a remesh vertex constraint for seams with a single stitch

diff --git a/Engine/Plugins/Experimental/ChaosClothAssetEditor/Source/ChaosClothAssetDataflowNodes/Private/ChaosClothAsset/RemeshNode.cpp b/Engine/Plugins/Experimental/ChaosClothAssetEditor/Source/ChaosClothAssetDataflowNodes/Private/ChaosClothAsset/RemeshNode.cpp
--- a/Engine/Plugins/Experimental/ChaosClothAssetEditor/Source/ChaosClothAssetDataflowNodes/Private/ChaosClothAsset/RemeshNode.cpp
+++ b/Engine/Plugins/Experimental/ChaosClothAssetEditor/Source/ChaosClothAssetDataflowNodes/Private/ChaosClothAsset/RemeshNode.cpp
@@ -236,6 +236,53 @@ namespace UE::Chaos::ClothAsset::Private
 		return Constraints;
 	}
 
+	bool CanSplitSeamEdge(int32 SeamID, int32 StitchID, const TArray<TArray<FIntVector2>>& Seams)
+	{
+		using namespace UE::Geometry;
+
+		bool bCanSplit = true;
+
+		const FSeamEdge SeamEdge(Seams[SeamID][StitchID], Seams[SeamID][StitchID + 1]);
+
+		// Check if any vertex exists in another stitch somewhere. For now we will skip these operations
+		// TODO: We could probably enable splits if we are very careful about handling mesh edges that are in more than one seam
+		TStaticArray<int32, 4> SeamEdgeVertices;
+		SeamEdgeVertices[0] = SeamEdge.Stitches[0][0];
+		SeamEdgeVertices[1] = SeamEdge.Stitches[0][1];
+		SeamEdgeVertices[2] = SeamEdge.Stitches[1][0];
+		SeamEdgeVertices[3] = SeamEdge.Stitches[1][1];
+
+		for (int32 InnerSeamID = 0; InnerSeamID < Seams.Num() && bCanSplit; ++InnerSeamID)
+		{
+			const TArray<FIntVector2>& InnerSeam = Seams[InnerSeamID];
+			for (int32 InnerStitchID = 0; InnerStitchID < InnerSeam.Num() && bCanSplit; ++InnerStitchID)
+			{
+				if (InnerSeamID == SeamID)
+				{
+					if (InnerStitchID == StitchID || InnerStitchID == StitchID + 1)
+					{
+						// Don't check against adjacent stitches
+						continue;
+					}
+				}
+
+				const FIntVector2& InnerStitch = InnerSeam[InnerStitchID];
+
+				for (const int32& SeamVertex : SeamEdgeVertices)
+				{
+					if (SeamVertex == InnerStitch[0] || SeamVertex == InnerStitch[1])
+					{
+						bCanSplit = false;
+						break;
+					}
+				}
+			}
+		}
+
+
+		return bCanSplit;
+	}
+
 
 	void RemeshSeams(UE::Geometry::FDynamicMesh3& Mesh, TArray<TArray<FIntVector2>>& Seams, double TargetEdgeLength)
 	{
@@ -305,6 +352,14 @@ namespace UE::Chaos::ClothAsset::Private
 					// Split
 					//
 
+
+					const bool bCanSplit = CanSplitSeamEdge(SeamID, StitchID, Seams);
+
+					if (!bCanSplit)
+					{
+						continue;
+					}
+
 					FDynamicMesh3::FEdgeSplitInfo SplitInfoA;
 					const EMeshResult ResultA = Mesh.SplitEdge(SideAVertexA, SideAVertexB, SplitInfoA);
 
@@ -629,20 +684,33 @@ namespace UE::Chaos::ClothAsset::Private
 		UE::Geometry::FMeshConstraints Constraints;
 		for (const TArray<FIntVector2>& Seam : Seams)
 		{
-			for (int32 StitchIndex = 0; StitchIndex < Seam.Num() - 1; ++StitchIndex)
+			if (Seam.Num() == 1)
 			{
 				for (int32 Side = 0; Side < 2; ++Side)
 				{
-					const int EdgeID = SourceMesh->FindEdge(Seam[StitchIndex][Side], Seam[StitchIndex + 1][Side]);
-					check(EdgeID != FDynamicMesh3::InvalidID);
-					FEdgeConstraint EdgeConstraint(EEdgeRefineFlags::FullyConstrained);
-					Constraints.SetOrUpdateEdgeConstraint(EdgeID, EdgeConstraint);
-
 					constexpr bool bCannotDelete = true;
 					constexpr bool bCanMove = false;
 					FVertexConstraint VertexConstraint(bCannotDelete, bCanMove);
-					Constraints.SetOrCombineVertexConstraint(Seam[StitchIndex][Side], VertexConstraint);
-					Constraints.SetOrCombineVertexConstraint(Seam[StitchIndex + 1][Side], VertexConstraint);
+					Constraints.SetOrCombineVertexConstraint(Seam[0][Side], VertexConstraint);
+				}
+			}
+			else
+			{
+				for (int32 StitchIndex = 0; StitchIndex < Seam.Num() - 1; ++StitchIndex)
+				{
+					for (int32 Side = 0; Side < 2; ++Side)
+					{
+						const int EdgeID = SourceMesh->FindEdge(Seam[StitchIndex][Side], Seam[StitchIndex + 1][Side]);
+						check(EdgeID != FDynamicMesh3::InvalidID);
+						FEdgeConstraint EdgeConstraint(EEdgeRefineFlags::FullyConstrained);
+						Constraints.SetOrUpdateEdgeConstraint(EdgeID, EdgeConstraint);
+
+						constexpr bool bCannotDelete = true;
+						constexpr bool bCanMove = false;
+						FVertexConstraint VertexConstraint(bCannotDelete, bCanMove);
+						Constraints.SetOrCombineVertexConstraint(Seam[StitchIndex][Side], VertexConstraint);
+						Constraints.SetOrCombineVertexConstraint(Seam[StitchIndex + 1][Side], VertexConstraint);
+					}
 				}
 			}
 		}
