Cloth paint tool: add an Invert action to flip paint values. Also, when painting in 2D update any additional 2D vertices that are mapped to the same 3D vertex via welding.

diff --git a/Engine/Plugins/Experimental/ChaosClothAssetEditor/Source/ChaosClothAssetEditorTools/Private/ChaosClothAsset/ClothWeightMapPaintTool.cpp b/Engine/Plugins/Experimental/ChaosClothAssetEditor/Source/ChaosClothAssetEditorTools/Private/ChaosClothAsset/ClothWeightMapPaintTool.cpp
--- a/Engine/Plugins/Experimental/ChaosClothAssetEditor/Source/ChaosClothAssetEditorTools/Private/ChaosClothAsset/ClothWeightMapPaintTool.cpp
+++ b/Engine/Plugins/Experimental/ChaosClothAssetEditor/Source/ChaosClothAssetEditorTools/Private/ChaosClothAsset/ClothWeightMapPaintTool.cpp
@@ -649,6 +649,26 @@ void UClothEditorWeightMapPaintTool::UpdateROI(const FSculptBrushStamp& BrushSta
 		ApplyVisibilityFilter(VertexSetBuffer, TempROIBuffer, ResultBuffer);
 	}
 
+	if (bHaveDynamicMeshToWeightConversion)
+	{
+		// Find triangles whose vertices map to the same welded vertex as any vertex in VertexSetBuffer and add them to TriangleROI
+
+		for (const int VertexID : VertexSetBuffer)
+		{
+			for (const int OtherVertexID : WeightToDynamicMesh[DynamicMeshToWeight[VertexID]])
+			{
+				if (OtherVertexID != VertexID)
+				{
+					Mesh->EnumerateVertexTriangles(OtherVertexID, [this](int32 AdjacentTri)
+					{
+						TriangleROI.Add(AdjacentTri);
+					});
+				}
+			}
+		}
+	}
+
+
 	VertexROI.SetNum(0, false);
 	//TODO: If we paint a 2D projection of UVs, these will need to be the 2D vertices not the 3D original mesh vertices
 	BufferUtil::AppendElements(VertexROI, VertexSetBuffer);
@@ -1638,6 +1658,33 @@ void UClothEditorWeightMapPaintTool::ClearAllWeightsAction()
 	EndChange();
 }
 
+void UClothEditorWeightMapPaintTool::InvertWeightsAction()
+{
+	if (!ActiveWeightMap)
+	{
+		return;
+	}
+	BeginChange();
+
+	const FDynamicMesh3* Mesh = DynamicMeshComponent->GetMesh();
+	checkf(Mesh, TEXT("Paint Tool's DynamicMeshComponent has no FDynamicMesh"));
+
+	for (const int32 VertexID : Mesh->VertexIndicesItr())
+	{
+		ActiveWeightEditChangeTracker->SaveVertexOneRingTriangles(VertexID, true);
+
+		float WeightValue;
+		ActiveWeightMap->GetValue(VertexID, &WeightValue);
+		WeightValue = 1.0f - WeightValue;
+		ActiveWeightMap->SetValue(VertexID, &WeightValue);
+	}
+
+	// update colors
+	UpdateVertexColorOverlay();
+	DynamicMeshComponent->FastNotifyVertexAttributesUpdated(EMeshRenderAttributeFlags::VertexColors);
+	GetToolManager()->PostInvalidation();
+	EndChange();
+}
 
 void UClothEditorWeightMapPaintTool::UpdateSelectedNode()
 {
@@ -1933,6 +1980,11 @@ void UClothEditorWeightMapPaintTool::ApplyAction(EClothEditorWeightMapPaintToolA
 	case EClothEditorWeightMapPaintToolActions::ClearAll:
 		ClearAllWeightsAction();
 		break;
+
+	case EClothEditorWeightMapPaintToolActions::Invert:
+		InvertWeightsAction();
+		break;
+
 	}
 }
 
diff --git a/Engine/Plugins/Experimental/ChaosClothAssetEditor/Source/ChaosClothAssetEditorTools/Public/ChaosClothAsset/ClothWeightMapPaintTool.h b/Engine/Plugins/Experimental/ChaosClothAssetEditor/Source/ChaosClothAssetEditorTools/Public/ChaosClothAsset/ClothWeightMapPaintTool.h
--- a/Engine/Plugins/Experimental/ChaosClothAssetEditor/Source/ChaosClothAssetEditorTools/Public/ChaosClothAsset/ClothWeightMapPaintTool.h
+++ b/Engine/Plugins/Experimental/ChaosClothAssetEditor/Source/ChaosClothAssetEditorTools/Public/ChaosClothAsset/ClothWeightMapPaintTool.h
@@ -203,6 +203,7 @@ enum class EClothEditorWeightMapPaintToolActions
 
 	FloodFillCurrent,
 	ClearAll,
+	Invert
 };
 
 
@@ -231,6 +232,12 @@ public:
 		PostAction(EClothEditorWeightMapPaintToolActions::FloodFillCurrent);
 	}
 
+	UFUNCTION(CallInEditor, Category = Operations, meta = (DisplayPriority = 13))
+	void Invert()
+	{
+		PostAction(EClothEditorWeightMapPaintToolActions::Invert);
+	}
+
 };
 
 UCLASS()
@@ -301,10 +308,11 @@ private:
 	UPROPERTY()
 	TObjectPtr<UWeightMapEraseBrushOpProps> EraseBrushOpProperties;
 
-public:
 	void FloodFillCurrentWeightAction();
 	void ClearAllWeightsAction();
+	void InvertWeightsAction();
 
+public:
 	void SetVerticesToWeightMap(const TSet<int32>& Vertices, double WeightValue, bool bIsErase);
 
 	bool HaveVisibilityFilter() const;
