Chaos Cloth Asset:
- Added remap function to the AddWeightMap node.
- Fixed property detail view refresh when the node's connections changed.
- Deactivate the active tool when a node is invalidated.

diff --git a/Engine/Plugins/Experimental/ChaosClothAssetEditor/Source/ChaosClothAssetDataflowNodes/Private/ChaosClothAsset/AddWeightMapNode.cpp b/Engine/Plugins/Experimental/ChaosClothAssetEditor/Source/ChaosClothAssetDataflowNodes/Private/ChaosClothAsset/AddWeightMapNode.cpp
--- a/Engine/Plugins/Experimental/ChaosClothAssetEditor/Source/ChaosClothAssetDataflowNodes/Private/ChaosClothAsset/AddWeightMapNode.cpp
+++ b/Engine/Plugins/Experimental/ChaosClothAssetEditor/Source/ChaosClothAssetDataflowNodes/Private/ChaosClothAsset/AddWeightMapNode.cpp
@@ -1,33 +1,227 @@
 // Copyright Epic Games, Inc. All Rights Reserved.
 
 #include "ChaosClothAsset/AddWeightMapNode.h"
-#include "ChaosClothAsset/CollectionClothFacade.h"
+#include "ChaosClothAsset/ClothAsset.h"
 #include "ChaosClothAsset/ClothDataflowTools.h"
 #include "ChaosClothAsset/ClothGeometryTools.h"
+#include "ChaosClothAsset/CollectionClothFacade.h"
 #include "ChaosClothAsset/WeightedValue.h"
 #include "Dataflow/DataflowInputOutput.h"
+#include "Dataflow/DataflowObject.h"
+#include "Utils/ClothingMeshUtils.h"
 
 #include UE_INLINE_GENERATED_CPP_BY_NAME(AddWeightMapNode)
 
 #define LOCTEXT_NAMESPACE "ChaosClothAssetAddWeightMapNode"
 
+namespace UE::Chaos::ClothAsset::Private
+{
+	void TransferWeightMap(
+		const TConstArrayView<FVector3f>& SourcePositions,
+		const TConstArrayView<FIntVector3>& InSourceIndices,
+		const TConstArrayView<float>& SourceWeights,
+		const TConstArrayView<FVector3f>& TargetPositions,
+		const TConstArrayView<FVector3f>& TargetNormals,
+		const TConstArrayView<FIntVector3>& InTargetIndices,
+		TArray<float>& TargetWeights)
+	{
+		check(TargetWeights.Num() == TargetPositions.Num());
+		if (!ensure(SourcePositions.Num() <= 65536))
+		{
+			return;  // MeshToMeshVertData below is limited to 16bit unsigned int indexes
+		}
+
+		TArray<uint32> SourceIndices;
+		SourceIndices.Reserve(InSourceIndices.Num() * 3);
+		for (const FIntVector3& InSourceIndex : InSourceIndices)
+		{
+			SourceIndices.Add(InSourceIndex[0]);
+			SourceIndices.Add(InSourceIndex[1]);
+			SourceIndices.Add(InSourceIndex[2]);
+		}
+		TArray<uint32> TargetIndices;
+		TargetIndices.Reserve(InTargetIndices.Num() * 3);
+		for (const FIntVector3& InTargetIndex : InTargetIndices)
+		{
+			TargetIndices.Add(InTargetIndex[0]);
+			TargetIndices.Add(InTargetIndex[1]);
+			TargetIndices.Add(InTargetIndex[2]);
+		}
+
+		const ClothingMeshUtils::ClothMeshDesc SourceMeshDesc(SourcePositions, SourceIndices);
+		const ClothingMeshUtils::ClothMeshDesc TargetMeshDesc(TargetPositions, TargetNormals, TargetIndices);
+
+		TArray<FMeshToMeshVertData> MeshToMeshVertData;
+		const FPointWeightMap* const MaxDistances = nullptr; // No need to update the vertex contribution on the transition maps
+		constexpr bool bUseSmoothTransitions = false;  // Smooth transitions are only used at rendering for now and not during LOD transitions
+		constexpr bool bUseMultipleInfluences = false;  // Multiple influences must not be used for LOD transitions
+		constexpr float SkinningKernelRadius = 0.f;  // KernelRadius is only required when using multiple influences
+
+		ClothingMeshUtils::GenerateMeshToMeshVertData(
+			MeshToMeshVertData,
+			TargetMeshDesc,
+			SourceMeshDesc,
+			MaxDistances,
+			bUseSmoothTransitions,
+			bUseMultipleInfluences,
+			SkinningKernelRadius);
+
+		check(MeshToMeshVertData.Num() == TargetWeights.Num());
+		for (int32 Index = 0; Index < TargetWeights.Num(); ++Index)
+		{
+			const FMeshToMeshVertData& MeshToMeshVertDatum = MeshToMeshVertData[Index];
+			
+			const uint16 VertIndex0 = MeshToMeshVertDatum.SourceMeshVertIndices[0];
+			const uint16 VertIndex1 = MeshToMeshVertDatum.SourceMeshVertIndices[1];
+			const uint16 VertIndex2 = MeshToMeshVertDatum.SourceMeshVertIndices[2];
+
+			TargetWeights[Index] = FMath::Clamp(
+				SourceWeights[VertIndex0] * MeshToMeshVertDatum.PositionBaryCoordsAndDist[0] +
+				SourceWeights[VertIndex1] * MeshToMeshVertDatum.PositionBaryCoordsAndDist[1] +
+				SourceWeights[VertIndex2] * MeshToMeshVertDatum.PositionBaryCoordsAndDist[2], 0.f, 1.f);
+		}
+	}
+
+	void TransferWeightMap(
+		const TConstArrayView<FVector2f>& InSourcePositions,
+		const TConstArrayView<FIntVector3>& SourceIndices,
+		const TConstArrayView<int32> SourceWeightsLookup,
+		const TConstArrayView<float>& InSourceWeights,
+		const TConstArrayView<FVector2f>& InTargetPositions,
+		const TConstArrayView<FIntVector3>& TargetIndices,
+		const TConstArrayView<int32> TargetWeightsLookup,
+		TArray<float>& OutTargetWeights)
+	{
+		TArray<FVector3f> SourcePositions;
+		SourcePositions.SetNumUninitialized(InSourcePositions.Num());
+		for (int32 Index = 0; Index < SourcePositions.Num(); ++Index)
+		{
+			SourcePositions[Index] = FVector3f(InSourcePositions[Index], 0.f);
+		}
+
+		TArray<float> SourceWeights;
+		SourceWeights.SetNumUninitialized(InSourcePositions.Num());
+		for (int32 Index = 0; Index < SourceWeights.Num(); ++Index)
+		{
+			SourceWeights[Index] = InSourceWeights[SourceWeightsLookup[Index]];
+		}
+
+		TArray<FVector3f> TargetPositions;
+		TArray<FVector3f> TargetNormals;
+		TargetPositions.SetNumUninitialized(InTargetPositions.Num());
+		TargetNormals.SetNumUninitialized(InTargetPositions.Num());
+		for (int32 Index = 0; Index < TargetPositions.Num(); ++Index)
+		{
+			TargetPositions[Index] = FVector3f(InTargetPositions[Index], 0.f);
+			TargetNormals[Index] = FVector3f::ZAxisVector;
+		}
+
+		TArray<float> TargetWeights;
+		TargetWeights.SetNumUninitialized(TargetPositions.Num());
+
+		TransferWeightMap(SourcePositions, SourceIndices, SourceWeights, TargetPositions, TargetNormals, TargetIndices, TargetWeights);
+
+		for (int32 Index = 0; Index < TargetWeights.Num(); ++Index)
+		{
+			OutTargetWeights[TargetWeightsLookup[Index]] = TargetWeights[Index];
+		}
+	}
+}
+
 FChaosClothAssetAddWeightMapNode::FChaosClothAssetAddWeightMapNode(const Dataflow::FNodeParameters& InParam, FGuid InGuid)
-	: FDataflowNode(InParam, InGuid)
+	: FDataflowTerminalNode(InParam, InGuid)
 {
 	RegisterInputConnection(&Collection);
+	RegisterInputConnection(&TransferCollection);
 	RegisterOutputConnection(&Collection, &Collection);
 	RegisterOutputConnection(&Name);
 }
 
+void FChaosClothAssetAddWeightMapNode::SetAssetValue(TObjectPtr<UObject> Asset, Dataflow::FContext& Context) const
+{
+	using namespace UE::Chaos::ClothAsset;
+
+	if (UChaosClothAsset* const ClothAsset = Cast<UChaosClothAsset>(Asset.Get()))
+	{
+PRAGMA_DISABLE_DEPRECATION_WARNINGS  // TODO: Remove after adding the getter function (currently shelved!)  
+		if (UDataflow* const DataflowAsset = ClothAsset->DataflowAsset)
+PRAGMA_ENABLE_DEPRECATION_WARNINGS
+		{
+			const TSharedPtr<Dataflow::FGraph, ESPMode::ThreadSafe> Dataflow = DataflowAsset->GetDataflow();
+			if (const TSharedPtr<FDataflowNode> BaseNode = Dataflow->FindBaseNode(this->GetGuid()))  // This is basically a safe const_cast
+			{
+				FChaosClothAssetAddWeightMapNode* const MutableThis = static_cast<FChaosClothAssetAddWeightMapNode*>(BaseNode.Get());
+				check(MutableThis == this);
+
+				// Make the name a valid attribute name, and replace the value in the UI
+				FWeightMapTools::MakeWeightMapName(MutableThis->Name);
+
+				// Transfer weight map if the transfer collection input has changed and is valid
+				FManagedArrayCollection InCollection = GetValue<FManagedArrayCollection>(Context, &Collection);
+				const TSharedRef<FManagedArrayCollection> ClothCollection = MakeShared<FManagedArrayCollection>(MoveTemp(InCollection));
+				FCollectionClothConstFacade ClothFacade(ClothCollection);
+				if (ClothFacade.IsValid())  // Can only act on the collection if it is a valid cloth collection
+				{
+					FManagedArrayCollection InTransferCollection = GetValue<FManagedArrayCollection>(Context, &TransferCollection);
+					const TSharedRef<const FManagedArrayCollection> TransferClothCollection = MakeShared<const FManagedArrayCollection>(MoveTemp(InTransferCollection));
+					FCollectionClothConstFacade TransferClothFacade(TransferClothCollection);
+
+					const FName InName(Name);
+					const uint32 NameTypeHash = HashCombineFast(GetTypeHash(InName), (uint32)TransferType);
+					const uint32 InTransferCollectionHash = (TransferClothFacade.IsValid() && InName != NAME_None) ?
+						HashCombineFast(TransferClothFacade.CalculateWeightMapTypeHash(), NameTypeHash) : 0;  // TODO: Remove after adding the function (currently shelved!)  
+				
+					if (TransferCollectionHash != InTransferCollectionHash)
+					{
+						MutableThis->TransferCollectionHash = InTransferCollectionHash;
+
+						if (TransferCollectionHash)
+						{
+							if (TransferClothFacade.HasWeightMap(InName))
+							{
+								// Remap the weights
+								MutableThis->GetVertexWeights().SetNumZeroed(ClothFacade.GetNumSimVertices3D());
+
+								switch (TransferType)
+								{
+								case EChaosClothAssetWeightMapTransferType::Use2DSimMesh:
+									Private::TransferWeightMap(
+										TransferClothFacade.GetSimPosition2D(),
+										TransferClothFacade.GetSimIndices2D(),
+										TransferClothFacade.GetSimVertex3DLookup(),
+										TransferClothFacade.GetWeightMap(InName),
+										ClothFacade.GetSimPosition2D(),
+										ClothFacade.GetSimIndices2D(),
+										ClothFacade.GetSimVertex3DLookup(),
+										MutableThis->GetVertexWeights());
+										break;
+								case EChaosClothAssetWeightMapTransferType::Use3DSimMesh:
+									Private::TransferWeightMap(
+										TransferClothFacade.GetSimPosition3D(),
+										TransferClothFacade.GetSimIndices3D(),
+										TransferClothFacade.GetWeightMap(InName),
+										ClothFacade.GetSimPosition3D(),
+										ClothFacade.GetSimNormal(),
+										ClothFacade.GetSimIndices3D(),
+										MutableThis->GetVertexWeights());
+										break;
+								default: unimplemented();
+								}
+							}
+						}
+					}
+				}
+			}
+		}
+	}
+}
+
 void FChaosClothAssetAddWeightMapNode::Evaluate(Dataflow::FContext& Context, const FDataflowOutput* Out) const
 {
 	using namespace UE::Chaos::ClothAsset;
 
 	if (Out->IsA<FManagedArrayCollection>(&Collection))
 	{
-		// Make the name a valid attribute name, and replace the value in the UI
-		FWeightMapTools::MakeWeightMapName(Name);
-
 		// Evaluate in collection
 		FManagedArrayCollection InCollection = GetValue<FManagedArrayCollection>(Context, &Collection);
 		const TSharedRef<FManagedArrayCollection> ClothCollection = MakeShared<FManagedArrayCollection>(MoveTemp(InCollection));
@@ -38,28 +232,25 @@ void FChaosClothAssetAddWeightMapNode::Evaluate(Dataflow::FContext& Context, con
 			ClothFacade.AddWeightMap(InName);		// Does nothing if weight map already exists
 
 			TArrayView<float> ClothWeights = ClothFacade.GetWeightMap(InName);
-			const int32 MaxWeightIndex = FMath::Min(VertexWeights.Num(), ClothWeights.Num());
-			if (VertexWeights.Num() > 0 && VertexWeights.Num() != ClothWeights.Num())
+			const int32 MaxWeightIndex = FMath::Min(GetVertexWeights().Num(), ClothWeights.Num());
+			if (GetVertexWeights().Num() > 0 && GetVertexWeights().Num() != ClothWeights.Num())
 			{
 				FClothDataflowTools::LogAndToastWarning(*this,
 					LOCTEXT("VertexCountMismatchHeadline", "Vertex count mismatch."),
 					FText::Format(LOCTEXT("VertexCountMismatchDetails", "Vertex weights in the node: {0}\n3D vertices in the cloth: {1}"),
-						VertexWeights.Num(),
+						GetVertexWeights().Num(),
 						ClothWeights.Num()));
 			}
 
 			for (int32 VertexID = 0; VertexID < MaxWeightIndex; ++VertexID)
 			{
-				ClothWeights[VertexID] = VertexWeights[VertexID];
+				ClothWeights[VertexID] = GetVertexWeights()[VertexID];
 			}
 		}
 		SetValue(Context, MoveTemp(*ClothCollection), &Collection);
 	}
 	else if (Out->IsA<FString>(&Name))
 	{
-		// Make the name a valid attribute name, and replace the value in the UI  // TODO: Remove all mutables from evaluation
-		FWeightMapTools::MakeWeightMapName(Name);
-
 		SetValue(Context, Name, &Name);
 	}
 }
diff --git a/Engine/Plugins/Experimental/ChaosClothAssetEditor/Source/ChaosClothAssetDataflowNodes/Public/ChaosClothAsset/AddWeightMapNode.h b/Engine/Plugins/Experimental/ChaosClothAssetEditor/Source/ChaosClothAssetDataflowNodes/Public/ChaosClothAsset/AddWeightMapNode.h
--- a/Engine/Plugins/Experimental/ChaosClothAssetEditor/Source/ChaosClothAssetDataflowNodes/Public/ChaosClothAsset/AddWeightMapNode.h
+++ b/Engine/Plugins/Experimental/ChaosClothAssetEditor/Source/ChaosClothAssetDataflowNodes/Public/ChaosClothAsset/AddWeightMapNode.h
@@ -2,31 +2,71 @@
 
 #pragma once 
 
-#include "Dataflow/DataflowNode.h"
+#include "Dataflow/DataflowTerminalNode.h"
 #include "GeometryCollection/ManagedArrayCollection.h"
 #include "AddWeightMapNode.generated.h"
 
+UENUM()
+enum class EChaosClothAssetWeightMapTransferType : uint8
+{
+	/** Transfer weight maps from the 2D simulation mesh (pattern against pattern). */
+	Use2DSimMesh UMETA(DisplayName = "Use 2D Sim Mesh"),
+	/** Transfer weight maps from the 3D simulation mesh (rest mesh against rest mesh). */
+	Use3DSimMesh UMETA(DisplayName = "Use 3D Sim Mesh"),
+};
+
 /** Painted weight map attributes node. */
 USTRUCT(Meta = (DataflowCloth))
-struct FChaosClothAssetAddWeightMapNode : public FDataflowNode
+struct FChaosClothAssetAddWeightMapNode : public FDataflowTerminalNode
 {
 	GENERATED_USTRUCT_BODY()
 	DATAFLOW_NODE_DEFINE_INTERNAL(FChaosClothAssetAddWeightMapNode, "AddWeightMap", "Cloth", "Cloth Add Weight Map")
 
 public:
 
-	UPROPERTY(Meta = (Dataflowinput, DataflowOutput, DataflowPassthrough = "Collection"))
+	UPROPERTY(Meta = (DataflowInput, DataflowOutput, DataflowPassthrough = "Collection"))
 	FManagedArrayCollection Collection;
 
+	/**
+	 * The collection used to transfer weight map from.
+	 * Connecting a collection containing a weight map sharing the same name as the one
+	 * created with this node will transfer the weights to the input collection vertices.
+	 * Note this operation only happens once when the TransferCollection is first connected, or updated.
+	 * Changing the Name or the TransferType will also redo the transfer operation.
+	 */
+	UPROPERTY(Meta = (DataflowInput))
+	FManagedArrayCollection TransferCollection;
+
 	/** The name to be set as a weight map attribute. */
 	UPROPERTY(EditAnywhere, Category = "Add Weight Map", Meta = (DataflowOutput))
-	mutable FString Name;  // Mutable so that it can be name checked in the evaluate function
+	FString Name;
+
+	/**
+	 * The type of transfer used to transfer the weight map when a TransferCollection is connected.
+	 * This property is disabled when no TransferCollection input has been conencted.
+	 */
+	UPROPERTY(EditAnywhere, Category = "Add Weight Map", Meta = (EditCondition = "TransferCollectionHash != 0"))
+	EChaosClothAssetWeightMapTransferType TransferType = EChaosClothAssetWeightMapTransferType::Use2DSimMesh;
 
+	UE_DEPRECATED(5.4, "This property will be made private.")
 	UPROPERTY()
 	TArray<float> VertexWeights;
 
 	FChaosClothAssetAddWeightMapNode(const Dataflow::FNodeParameters& InParam, FGuid InGuid = FGuid::NewGuid());
 
 private:
+	friend class UClothEditorWeightMapPaintTool;
+
+PRAGMA_DISABLE_DEPRECATION_WARNINGS
+	const TArray<float>& GetVertexWeights() const { return VertexWeights; }
+	TArray<float>& GetVertexWeights() { return VertexWeights; }
+PRAGMA_ENABLE_DEPRECATION_WARNINGS
+
+	//~ Begin FDataflowNode interface
+	virtual void SetAssetValue(TObjectPtr<UObject> Asset, Dataflow::FContext& Context) const override;
 	virtual void Evaluate(Dataflow::FContext& Context, const FDataflowOutput* Out) const override;
+	//~ End FDataflowNode interface
+
+	UPROPERTY()
+	uint32 TransferCollectionHash = 0;
 };
diff --git a/Engine/Plugins/Experimental/ChaosClothAssetEditor/Source/ChaosClothAssetEditor/Private/ChaosClothAsset/ClothEditorToolkit.cpp b/Engine/Plugins/Experimental/ChaosClothAssetEditor/Source/ChaosClothAssetEditor/Private/ChaosClothAsset/ClothEditorToolkit.cpp
--- a/Engine/Plugins/Experimental/ChaosClothAssetEditor/Source/ChaosClothAssetEditor/Private/ChaosClothAsset/ClothEditorToolkit.cpp
+++ b/Engine/Plugins/Experimental/ChaosClothAssetEditor/Source/ChaosClothAssetEditor/Private/ChaosClothAsset/ClothEditorToolkit.cpp
@@ -1125,6 +1125,22 @@ void FChaosClothAssetEditorToolkit::OnNodeSelectionChanged(const TSet<UObject*>&
 											SelectedDataflowNode->OnSelected(*DataflowContext);
 										}
 									}
+									// The detail panel won't update correctly when a change affects a node's properties and needs a refresh
+									if (NodeDetailsEditor && NodeDetailsEditor->GetDetailsView())
+									{
+										NodeDetailsEditor->GetDetailsView()->ForceRefresh();
+									}
+									// The node has just been invalidated, best to close the current tool
+									if (const UChaosClothAssetEditorMode* const ClothMode = 
+										CastChecked<UChaosClothAssetEditorMode>(EditorModeManager->GetActiveScriptableMode(UChaosClothAssetEditorMode::EM_ChaosClothAssetEditorModeId)))
+									{
+										UEditorInteractiveToolsContext* const ToolsContext = ClothMode->GetInteractiveToolsContext();
+										checkf(ToolsContext, TEXT("No valid ToolsContext found for FChaosClothAssetEditorToolkit"));
+										if (ToolsContext->HasActiveTool())
+										{
+											ToolsContext->EndTool(ToolsContext->CanCancelActiveTool() ? EToolShutdownType::Cancel : EToolShutdownType::Completed);
+										}
+									}
 								});
 						}
 					}
diff --git a/Engine/Plugins/Experimental/ChaosClothAssetEditor/Source/ChaosClothAssetEditorTools/Private/ChaosClothAsset/ClothWeightMapPaintTool.cpp b/Engine/Plugins/Experimental/ChaosClothAssetEditor/Source/ChaosClothAssetEditorTools/Private/ChaosClothAsset/ClothWeightMapPaintTool.cpp
--- a/Engine/Plugins/Experimental/ChaosClothAssetEditor/Source/ChaosClothAssetEditorTools/Private/ChaosClothAsset/ClothWeightMapPaintTool.cpp
+++ b/Engine/Plugins/Experimental/ChaosClothAssetEditor/Source/ChaosClothAssetEditorTools/Private/ChaosClothAsset/ClothWeightMapPaintTool.cpp
@@ -345,7 +345,7 @@ void UClothEditorWeightMapPaintTool::Setup()
 	ActiveWeightMap->SetName(FName("PaintLayer"));
 
 	// Copy weights from selected node to the preview mesh
-	const TArray<float>& CurrentWeights = WeightMapNodeToUpdate->VertexWeights;
+	const TArray<float>& CurrentWeights = WeightMapNodeToUpdate->GetVertexWeights();
 	
 	if (bHaveDynamicMeshToWeightConversion)
 	{
@@ -1696,15 +1696,15 @@ void UClothEditorWeightMapPaintTool::UpdateSelectedNode()
 
 	if (bHaveDynamicMeshToWeightConversion)
 	{
-		WeightMapNodeToUpdate->VertexWeights.Init(0.f, WeightToDynamicMesh.Num());
+		WeightMapNodeToUpdate->GetVertexWeights().Init(0.f, WeightToDynamicMesh.Num());
 		for (int32 DynamicMeshIdx = 0; DynamicMeshIdx < CurrentWeights.Num(); ++DynamicMeshIdx)
 		{
-			WeightMapNodeToUpdate->VertexWeights[DynamicMeshToWeight[DynamicMeshIdx]] = CurrentWeights[DynamicMeshIdx];
+			WeightMapNodeToUpdate->GetVertexWeights()[DynamicMeshToWeight[DynamicMeshIdx]] = CurrentWeights[DynamicMeshIdx];
 		}
 	}
 	else
 	{
-		WeightMapNodeToUpdate->VertexWeights = CurrentWeights;
+		WeightMapNodeToUpdate->GetVertexWeights() = CurrentWeights;
 	}
 	
 	WeightMapNodeToUpdate->Name = UpdateWeightMapProperties->Name;
