CIS - Fixed Clang deprecation warning in AddWeightMapNode.h.

diff --git a/Engine/Plugins/Experimental/ChaosClothAssetEditor/Source/ChaosClothAssetDataflowNodes/Public/ChaosClothAsset/AddWeightMapNode.h b/Engine/Plugins/Experimental/ChaosClothAssetEditor/Source/ChaosClothAssetDataflowNodes/Public/ChaosClothAsset/AddWeightMapNode.h
--- a/Engine/Plugins/Experimental/ChaosClothAssetEditor/Source/ChaosClothAssetDataflowNodes/Public/ChaosClothAsset/AddWeightMapNode.h
+++ b/Engine/Plugins/Experimental/ChaosClothAssetEditor/Source/ChaosClothAssetDataflowNodes/Public/ChaosClothAsset/AddWeightMapNode.h
@@ -16,6 +16,7 @@ enum class EChaosClothAssetWeightMapTransferType : uint8
 };
 
 /** Painted weight map attributes node. */
+PRAGMA_DISABLE_DEPRECATION_WARNINGS  // For deprecated VertexWeights in copy constructor (Clang)
 USTRUCT(Meta = (DataflowCloth))
 struct FChaosClothAssetAddWeightMapNode : public FDataflowTerminalNode
 {
@@ -57,10 +58,8 @@ public:
 private:
 	friend class UClothEditorWeightMapPaintTool;
 
-PRAGMA_DISABLE_DEPRECATION_WARNINGS
 	const TArray<float>& GetVertexWeights() const { return VertexWeights; }
 	TArray<float>& GetVertexWeights() { return VertexWeights; }
-PRAGMA_ENABLE_DEPRECATION_WARNINGS
 
 	//~ Begin FDataflowNode interface
 	virtual void SetAssetValue(TObjectPtr<UObject> Asset, Dataflow::FContext& Context) const override;
@@ -70,3 +69,4 @@ PRAGMA_ENABLE_DEPRECATION_WARNINGS
 	UPROPERTY()
 	uint32 TransferCollectionHash = 0;
 };
+PRAGMA_ENABLE_DEPRECATION_WARNINGS
