Chaos Cloth Editor: fix crash selecting comment nodes

diff --git a/Engine/Plugins/Experimental/ChaosClothAssetEditor/Source/ChaosClothAssetEditor/Private/ChaosClothAsset/ClothEditorMode.cpp b/Engine/Plugins/Experimental/ChaosClothAssetEditor/Source/ChaosClothAssetEditor/Private/ChaosClothAsset/ClothEditorMode.cpp
--- a/Engine/Plugins/Experimental/ChaosClothAssetEditor/Source/ChaosClothAssetEditor/Private/ChaosClothAsset/ClothEditorMode.cpp
+++ b/Engine/Plugins/Experimental/ChaosClothAssetEditor/Source/ChaosClothAssetEditor/Private/ChaosClothAsset/ClothEditorMode.cpp
@@ -1346,7 +1346,13 @@ UEdGraphNode* UChaosClothAssetEditorMode::GetSingleSelectedNodeWithOutputType(co
 		return nullptr;
 	}
 
-	const UDataflowEdNode* const SelectedDataflowEdNode = CastChecked<UDataflowEdNode>(SelectedNode);
+	const UDataflowEdNode* const SelectedDataflowEdNode = Cast<UDataflowEdNode>(SelectedNode);
+	if (!SelectedDataflowEdNode)
+	{
+		// The graph can contain UEdGraphNode_Comment, which is not a UDataflowEdNode
+		return nullptr;
+	}
+
 	const TSharedPtr<const FDataflowNode> SelectedDataflowNode = SelectedDataflowEdNode->GetDataflowNode();
 
 	if (!SelectedDataflowNode)
