Cloth Editor: add camera speed settings.

diff --git a/Engine/Plugins/Experimental/ChaosClothAssetEditor/Source/ChaosClothAssetEditor/Private/ChaosClothAsset/SClothEditorRestSpaceViewportToolBar.cpp b/Engine/Plugins/Experimental/ChaosClothAssetEditor/Source/ChaosClothAssetEditor/Private/ChaosClothAsset/SClothEditorRestSpaceViewportToolBar.cpp
--- a/Engine/Plugins/Experimental/ChaosClothAssetEditor/Source/ChaosClothAssetEditor/Private/ChaosClothAsset/SClothEditorRestSpaceViewportToolBar.cpp
+++ b/Engine/Plugins/Experimental/ChaosClothAssetEditor/Source/ChaosClothAssetEditor/Private/ChaosClothAsset/SClothEditorRestSpaceViewportToolBar.cpp
@@ -81,6 +81,15 @@ TSharedRef<SWidget> SChaosClothAssetEditorRestSpaceViewportToolBar::GenerateClot
 			if (bIsPerspective)
 			{
 				OptionsMenuBuilder.AddWidget(GenerateFOVMenu(), LOCTEXT("FOVAngle", "Field of View (H)"));
+
+				OptionsMenuBuilder.AddSubMenu(
+					LOCTEXT("CameraSpeedSettings", "Camera Speed Settings"),
+					LOCTEXT("CameraSpeedSettingsToolTip", "Adjust camera speed settings"),
+					FNewMenuDelegate::CreateLambda([this](FMenuBuilder& MenuBuilder)
+				{
+					MenuBuilder.AddWidget(GenerateCameraSpeedSettingsMenu(), FText());
+				}
+				));
 			}
 			OptionsMenuBuilder.AddWidget(GenerateLightMenu(), LOCTEXT("LightIntensity", "Render Light Intensity"));
 
diff --git a/Engine/Plugins/Experimental/ChaosClothAssetEditor/Source/ChaosClothAssetEditor/Private/ChaosClothAsset/SClothEditorViewportToolBarBase.cpp b/Engine/Plugins/Experimental/ChaosClothAssetEditor/Source/ChaosClothAssetEditor/Private/ChaosClothAsset/SClothEditorViewportToolBarBase.cpp
--- a/Engine/Plugins/Experimental/ChaosClothAssetEditor/Source/ChaosClothAssetEditor/Private/ChaosClothAsset/SClothEditorViewportToolBarBase.cpp
+++ b/Engine/Plugins/Experimental/ChaosClothAssetEditor/Source/ChaosClothAssetEditor/Private/ChaosClothAsset/SClothEditorViewportToolBarBase.cpp
@@ -3,7 +3,9 @@
 #include "ChaosClothAsset/SClothEditorViewportToolBarBase.h"
 #include "SEditorViewport.h"
 #include "Widgets/Input/SSpinBox.h"
+#include "Widgets/Input/SSlider.h"
 #include "Widgets/Layout/SBorder.h"
+#include "EditorViewportClient.h"
 
 #define LOCTEXT_NAMESPACE "SChaosClothAssetEditorViewportToolBarBase"
 
@@ -23,6 +25,14 @@ TSharedRef<SWidget> SChaosClothAssetEditorViewportToolBarBase::GenerateClothView
 			if (bIsPerspective)
 			{
 				OptionsMenuBuilder.AddWidget(GenerateFOVMenu(), LOCTEXT("FOVAngle", "Field of View (H)"));
+				OptionsMenuBuilder.AddSubMenu(
+					LOCTEXT("CameraSpeedSettings", "Camera Speed Settings"),
+					LOCTEXT("CameraSpeedSettingsToolTip", "Adjust camera speed settings"),
+					FNewMenuDelegate::CreateLambda([this](FMenuBuilder& MenuBuilder)
+				{
+					MenuBuilder.AddWidget(GenerateCameraSpeedSettingsMenu(), FText());
+				}
+				));
 			}
 		}
 		OptionsMenuBuilder.EndSection();
@@ -66,6 +76,83 @@ TSharedRef<SWidget> SChaosClothAssetEditorViewportToolBarBase::GenerateFOVMenu()
 		];
 }
 
+TSharedRef<SWidget> SChaosClothAssetEditorViewportToolBarBase::GenerateCameraSpeedSettingsMenu() const
+{
+	// This comes from STransformViewportToolBar::FillCameraSpeedMenu
+	return SNew(SBorder)
+		.BorderImage(FAppStyle::GetBrush(TEXT("Menu.Background")))
+		[
+			SNew( SVerticalBox )
+			+SVerticalBox::Slot()
+			.AutoHeight()
+			.Padding( FMargin(8.0f, 2.0f, 60.0f, 2.0f) )
+			.HAlign( HAlign_Left )
+			[
+				SNew( STextBlock )
+				.Text( LOCTEXT("MouseSettingsCamSpeed", "Camera Speed")  )
+				.Font( FAppStyle::GetFontStyle( TEXT( "MenuItem.Font" ) ) )
+			]
+			+SVerticalBox::Slot()
+			.AutoHeight()
+			.Padding( FMargin(8.0f, 4.0f) )
+			[	
+				SNew( SHorizontalBox )
+				+SHorizontalBox::Slot()
+				.FillWidth(1)
+				.Padding( FMargin(0.0f, 2.0f) )
+				[
+					SNew( SBox )
+					.MinDesiredWidth(220)
+					[
+						SNew(SSlider)
+						.Value(this, &SChaosClothAssetEditorViewportToolBarBase::GetCamSpeedSliderPosition)
+						.OnValueChanged(this, &SChaosClothAssetEditorViewportToolBarBase::OnSetCamSpeed)
+					]
+				]
+				+SHorizontalBox::Slot()
+				.AutoWidth()
+				.Padding( 8.0f, 2.0f, 0.0f, 2.0f)
+				[
+					SNew( SBox )
+					.WidthOverride(40)
+					[
+						SNew( STextBlock )
+						.Text(this, &SChaosClothAssetEditorViewportToolBarBase::GetCameraSpeedLabel)
+						.Font( FAppStyle::GetFontStyle( TEXT( "MenuItem.Font" ) ) )
+					]
+				]
+			] // Camera Speed Scalar
+			+ SVerticalBox::Slot()
+			.AutoHeight()
+			.Padding(FMargin(8.0f, 2.0f, 60.0f, 2.0f))
+			.HAlign(HAlign_Left)
+			[
+				SNew(STextBlock)
+				.Text(LOCTEXT("MouseSettingsCamSpeedScalar", "Camera Speed Scalar"))
+				.Font(FAppStyle::GetFontStyle(TEXT("MenuItem.Font")))
+			]
+			+ SVerticalBox::Slot()
+			.AutoHeight()
+			.Padding(FMargin(8.0f, 4.0f))
+			[
+				SNew(SHorizontalBox)
+				+ SHorizontalBox::Slot()
+				.FillWidth(1)
+				.Padding(FMargin(0.0f, 2.0f))
+				[
+					SNew(SSpinBox<float>)
+					.MinValue(1.0f)
+ 					.MaxValue(std::numeric_limits<float>::max())
+					.MinSliderValue(1.0f)
+					.MaxSliderValue(128.0f)
+					.Value(this, &SChaosClothAssetEditorViewportToolBarBase::GetCamSpeedScalarBoxValue)
+					.OnValueChanged(this, &SChaosClothAssetEditorViewportToolBarBase::OnSetCamSpeedScalarBoxValue)
+					.ToolTipText(LOCTEXT("CameraSpeedScalar_ToolTip", "Scalar to increase camera movement range"))
+				]
+			]
+		];
+}
+
 float SChaosClothAssetEditorViewportToolBarBase::OnGetFOVValue() const
 {
 	// Re-implementation of inaccessible SCommonEditorViewportToolbarBase::OnGetFOVValue()
@@ -73,4 +160,42 @@ float SChaosClothAssetEditorViewportToolBarBase::OnGetFOVValue() const
 	return GetViewportClient().ViewFOV;
 }
 
+FText SChaosClothAssetEditorViewportToolBarBase::GetCameraSpeedLabel() const
+{
+	const float CameraSpeed = GetViewportClient().GetCameraSpeed();
+	FNumberFormattingOptions FormattingOptions = FNumberFormattingOptions::DefaultNoGrouping();
+	FormattingOptions.MaximumFractionalDigits = CameraSpeed > 1 ? 1 : 3;
+	return FText::AsNumber(CameraSpeed, &FormattingOptions);
+}
+
+float SChaosClothAssetEditorViewportToolBarBase::GetCamSpeedSliderPosition() const
+{
+	return (GetViewportClient().GetCameraSpeedSetting() - 1) / ((float)FEditorViewportClient::MaxCameraSpeeds - 1);
+}
+
+void SChaosClothAssetEditorViewportToolBarBase::OnSetCamSpeed(float NewValue) const
+{
+	const int32 OldSpeedSetting = GetViewportClient().GetCameraSpeedSetting();
+	const int32 NewSpeedSetting = NewValue * ((float)FEditorViewportClient::MaxCameraSpeeds - 1) + 1;
+
+	if (OldSpeedSetting != NewSpeedSetting)
+	{
+		GetViewportClient().SetCameraSpeedSetting(NewSpeedSetting);
+	}
+}
+
+FText SChaosClothAssetEditorViewportToolBarBase::GetCameraSpeedScalarLabel() const
+{
+	return FText::AsNumber(GetViewportClient().GetCameraSpeedScalar());
+}
+
+float SChaosClothAssetEditorViewportToolBarBase::GetCamSpeedScalarBoxValue() const
+{
+	return GetViewportClient().GetCameraSpeedScalar();
+}
+
+void SChaosClothAssetEditorViewportToolBarBase::OnSetCamSpeedScalarBoxValue(float NewValue) const
+{
+	GetViewportClient().SetCameraSpeedScalar(NewValue);
+}
 #undef LOCTEXT_NAMESPACE
diff --git a/Engine/Plugins/Experimental/ChaosClothAssetEditor/Source/ChaosClothAssetEditor/Private/ChaosClothAsset/SClothEditorViewportToolBarBase.h b/Engine/Plugins/Experimental/ChaosClothAssetEditor/Source/ChaosClothAssetEditor/Private/ChaosClothAsset/SClothEditorViewportToolBarBase.h
--- a/Engine/Plugins/Experimental/ChaosClothAssetEditor/Source/ChaosClothAssetEditor/Private/ChaosClothAsset/SClothEditorViewportToolBarBase.h
+++ b/Engine/Plugins/Experimental/ChaosClothAssetEditor/Source/ChaosClothAssetEditor/Private/ChaosClothAsset/SClothEditorViewportToolBarBase.h
@@ -10,6 +10,13 @@ class CHAOSCLOTHASSETEDITOR_API SChaosClothAssetEditorViewportToolBarBase : publ
 protected:
 	TSharedRef<SWidget> GenerateClothViewportOptionsMenu() const;
 	TSharedRef<SWidget> GenerateFOVMenu() const;
+	TSharedRef<SWidget> GenerateCameraSpeedSettingsMenu() const;
 	float OnGetFOVValue() const;
+	FText GetCameraSpeedLabel() const;
+	float GetCamSpeedSliderPosition() const;
+	void OnSetCamSpeed(float NewValue) const;
+	FText GetCameraSpeedScalarLabel() const;
+	float GetCamSpeedScalarBoxValue() const;
+	void OnSetCamSpeedScalarBoxValue(float NewValue) const;
 };
 
