Cloth Editor: fix crash when DynamicMeshComponent is nullptr.

diff --git a/Engine/Plugins/Experimental/ChaosClothAssetEditor/Source/ChaosClothAssetEditor/Private/ChaosClothAsset/ClothEditorMode.cpp b/Engine/Plugins/Experimental/ChaosClothAssetEditor/Source/ChaosClothAssetEditor/Private/ChaosClothAsset/ClothEditorMode.cpp
--- a/Engine/Plugins/Experimental/ChaosClothAssetEditor/Source/ChaosClothAssetEditor/Private/ChaosClothAsset/ClothEditorMode.cpp
+++ b/Engine/Plugins/Experimental/ChaosClothAssetEditor/Source/ChaosClothAssetEditor/Private/ChaosClothAsset/ClothEditorMode.cpp
@@ -1224,12 +1224,15 @@ bool UChaosClothAssetEditorMode::CanSetConstructionViewSeamsActive() const
 	}
 
 	// Seam view not available on non-manifold meshes
-	if (const UE::Geometry::FDynamicMesh3* const Mesh = DynamicMeshComponent->GetMesh())
+	if (DynamicMeshComponent)
 	{
-		const UE::Geometry::FNonManifoldMappingSupport NonManifold(*Mesh);
-		if (NonManifold.IsNonManifoldVertexInSource())
+		if (const UE::Geometry::FDynamicMesh3* const Mesh = DynamicMeshComponent->GetMesh())
 		{
-			return false;
+			const UE::Geometry::FNonManifoldMappingSupport NonManifold(*Mesh);
+			if (NonManifold.IsNonManifoldVertexInSource())
+			{
+				return false;
+			}
 		}
 	}
 
