Cloth Editor: move the Simulation Visualization panel from drop down menu to its own tab

diff --git a/Engine/Plugins/Experimental/ChaosClothAssetEditor/Source/ChaosClothAssetEditor/Private/ChaosClothAsset/ClothEditorToolkit.cpp b/Engine/Plugins/Experimental/ChaosClothAssetEditor/Source/ChaosClothAssetEditor/Private/ChaosClothAsset/ClothEditorToolkit.cpp
--- a/Engine/Plugins/Experimental/ChaosClothAssetEditor/Source/ChaosClothAssetEditor/Private/ChaosClothAsset/ClothEditorToolkit.cpp
+++ b/Engine/Plugins/Experimental/ChaosClothAssetEditor/Source/ChaosClothAssetEditor/Private/ChaosClothAsset/ClothEditorToolkit.cpp
@@ -52,7 +52,7 @@ const FName FChaosClothAssetEditorToolkit::OutlinerTabID(TEXT("ChaosClothAssetEd
 const FName FChaosClothAssetEditorToolkit::PreviewSceneDetailsTabID(TEXT("ChaosClothAssetEditor_PreviewSceneDetailsTab"));
 const FName FChaosClothAssetEditorToolkit::GraphCanvasTabId(TEXT("ChaosClothAssetEditor_GraphCanvas"));
 const FName FChaosClothAssetEditorToolkit::NodeDetailsTabId(TEXT("ChaosClothAssetEditor_NodeDetails"));
-
+const FName FChaosClothAssetEditorToolkit::SimulationVisualizationTabID(TEXT("ChaosClothAssetEditor_SimulationVisualizationTab"));
 
 namespace Private
 {
@@ -96,7 +96,7 @@ FChaosClothAssetEditorToolkit::FChaosClothAssetEditorToolkit(UAssetEditor* InOwn
 	// Note: Changes to the layout should include a increment to the layout's ID, i.e.
 	// ChaosClothAssetEditorLayout[X] -> ChaosClothAssetEditorLayout[X+1]. Otherwise, layouts may be messed up
 	// without a full reset to layout defaults inside the editor.
-	StandaloneDefaultLayout = FTabManager::NewLayout(FName("ChaosClothAssetEditorLayout5"))
+	StandaloneDefaultLayout = FTabManager::NewLayout(FName("ChaosClothAssetEditorLayout6"))
 		->AddArea
 		(
 			FTabManager::NewPrimaryArea()->SetOrientation(Orient_Horizontal)
@@ -153,6 +153,7 @@ FChaosClothAssetEditorToolkit::FChaosClothAssetEditorToolkit(UAssetEditor* InOwn
 					->SetSizeCoefficient(0.65f)	// Relative height of (Asset Details, Preview Scene Details) vs (Dataflow Node Details)
 					->AddTab(DetailsTabID, ETabState::OpenedTab)
 					->AddTab(PreviewSceneDetailsTabID, ETabState::OpenedTab)
+					->AddTab(SimulationVisualizationTabID, ETabState::OpenedTab)
 					->SetExtensionId("DetailsArea")
 					->SetHideTabWell(true)
 					->SetForegroundTab(DetailsTabID)
@@ -605,6 +606,10 @@ void FChaosClothAssetEditorToolkit::RegisterTabSpawners(const TSharedRef<FTabMan
 		.SetGroup(AssetEditorTabsCategory.ToSharedRef())
 		.SetIcon(FSlateIcon(FAppStyle::GetAppStyleSetName(), "LevelEditor.Tabs.Details"));
 
+	InTabManager->RegisterTabSpawner(SimulationVisualizationTabID, FOnSpawnTab::CreateSP(this, &FChaosClothAssetEditorToolkit::SpawnTab_SimulationVisualization))
+		.SetDisplayName(LOCTEXT("SimulationVisualizationTabDisplayName", "Simulation Visualization"))
+		.SetGroup(AssetEditorTabsCategory.ToSharedRef());
+
 	InTabManager->RegisterTabSpawner(OutlinerTabID, FOnSpawnTab::CreateSP(this, &FChaosClothAssetEditorToolkit::SpawnTab_Outliner))
 		.SetDisplayName(LOCTEXT("Outliner", "Outliner"))
 		.SetGroup(AssetEditorTabsCategory.ToSharedRef())
@@ -634,6 +639,7 @@ void FChaosClothAssetEditorToolkit::UnregisterTabSpawners(const TSharedRef<FTabM
 	InTabManager->UnregisterTabSpawner(ClothPreviewTabID);
 	InTabManager->UnregisterTabSpawner(ViewportTabID);
 	InTabManager->UnregisterTabSpawner(DetailsTabID);
+	InTabManager->UnregisterTabSpawner(SimulationVisualizationTabID);
 	InTabManager->UnregisterTabSpawner(OutlinerTabID);
 	InTabManager->UnregisterTabSpawner(PreviewSceneDetailsTabID);
 	InTabManager->UnregisterTabSpawner(GraphCanvasTabId);
@@ -679,6 +685,21 @@ TSharedRef<SDockTab> FChaosClothAssetEditorToolkit::SpawnTab_Outliner(const FSpa
 }
 
 
+TSharedRef<SDockTab> FChaosClothAssetEditorToolkit::SpawnTab_SimulationVisualization(const FSpawnTabArgs& Args)
+{
+	SAssignNew(SimulationVisualizationDockTab, SDockTab)
+		.Label(LOCTEXT("SimulationVisualizationTitle", "Simulation Visualization"));
+
+	FMenuBuilder MenuBuilder(false, nullptr);
+	if (FClothEditorSimulationVisualization* const Visualization = ClothPreviewViewportClient->GetSimulationVisualization().Pin().Get())
+	{
+		Visualization->ExtendViewportShowMenu(MenuBuilder, ClothPreviewViewportClient.ToSharedRef());
+	}
+	SimulationVisualizationDockTab->SetContent(MenuBuilder.MakeWidget());
+
+	return SimulationVisualizationDockTab.ToSharedRef();
+}
+
 TSharedRef<SDockTab> FChaosClothAssetEditorToolkit::SpawnTab_PreviewSceneDetails(const FSpawnTabArgs& Args)
 {
 	SAssignNew(PreviewSceneDockTab, SDockTab)
diff --git a/Engine/Plugins/Experimental/ChaosClothAssetEditor/Source/ChaosClothAssetEditor/Private/ChaosClothAsset/ClothEditorToolkit.h b/Engine/Plugins/Experimental/ChaosClothAssetEditor/Source/ChaosClothAssetEditor/Private/ChaosClothAsset/ClothEditorToolkit.h
--- a/Engine/Plugins/Experimental/ChaosClothAssetEditor/Source/ChaosClothAssetEditor/Private/ChaosClothAsset/ClothEditorToolkit.h
+++ b/Engine/Plugins/Experimental/ChaosClothAssetEditor/Source/ChaosClothAssetEditor/Private/ChaosClothAsset/ClothEditorToolkit.h
@@ -62,6 +62,7 @@ private:
 	static const FName ClothPreviewTabID;
 	static const FName OutlinerTabID;
 	static const FName PreviewSceneDetailsTabID;
+	static const FName SimulationVisualizationTabID;
 
 	// FTickableEditorObject
 	virtual void Tick(float DeltaTime) override;
@@ -109,6 +110,7 @@ private:
 	TSharedRef<SDockTab> SpawnTab_PreviewSceneDetails(const FSpawnTabArgs& Args);
 	TSharedRef<SDockTab> SpawnTab_GraphCanvas(const FSpawnTabArgs& Args);
 	TSharedRef<SDockTab> SpawnTab_NodeDetails(const FSpawnTabArgs& Args);
+	TSharedRef<SDockTab> SpawnTab_SimulationVisualization(const FSpawnTabArgs& Args);
 
 	void InitDetailsViewPanel();
 	void OnFinishedChangingAssetProperties(const FPropertyChangedEvent&);
@@ -149,6 +151,8 @@ private:
 	TSharedPtr<SDockTab> PreviewSceneDockTab;
 	TSharedPtr<SWidget> AdvancedPreviewSettingsWidget;
 
+	TSharedPtr<SDockTab> SimulationVisualizationDockTab;
+
 	TSharedPtr<SClothCollectionOutliner> Outliner;
 
 	// Dataflow
diff --git a/Engine/Plugins/Experimental/ChaosClothAssetEditor/Source/ChaosClothAssetEditor/Private/ChaosClothAsset/SClothEditor3DViewportToolBar.cpp b/Engine/Plugins/Experimental/ChaosClothAssetEditor/Source/ChaosClothAssetEditor/Private/ChaosClothAsset/SClothEditor3DViewportToolBar.cpp
--- a/Engine/Plugins/Experimental/ChaosClothAssetEditor/Source/ChaosClothAssetEditor/Private/ChaosClothAsset/SClothEditor3DViewportToolBar.cpp
+++ b/Engine/Plugins/Experimental/ChaosClothAssetEditor/Source/ChaosClothAssetEditor/Private/ChaosClothAsset/SClothEditor3DViewportToolBar.cpp
@@ -232,25 +232,5 @@ TSharedRef<SWidget> SChaosClothAssetEditor3DViewportToolBar::MakeLODMenu() const
 	return MenuBuilder.MakeWidget();
 }
 
-void SChaosClothAssetEditor3DViewportToolBar::ExtendOptionsMenu(FMenuBuilder& OptionsMenuBuilder) const
-{
-	using namespace UE::Chaos::ClothAsset;
-	{
-		// Debug visualization
-		constexpr bool bOpenSubMenuOnClick = false;
-		constexpr bool bShouldCloseWindowAfterMenuSelection = false;
-		OptionsMenuBuilder.AddSubMenu(
-			LOCTEXT("SimulationVisualization", "Simulation Visualization"),
-			LOCTEXT("SimulationVisualizationToolTip", "Options to control simulation visualization"),
-			FNewMenuDelegate::CreateLambda([this](FMenuBuilder& MenuBuilder)
-		{
-			TSharedRef<FChaosClothAssetEditor3DViewportClient> ViewportClient = StaticCastSharedPtr<FChaosClothAssetEditor3DViewportClient>(ChaosClothAssetEditor3DViewportPtr.Pin()->GetViewportClient()).ToSharedRef();
-			if (FClothEditorSimulationVisualization* const Visualization = ViewportClient->GetSimulationVisualization().Pin().Get())
-			{
-				Visualization->ExtendViewportShowMenu(MenuBuilder, ViewportClient);
-			}
-		}), bOpenSubMenuOnClick, FSlateIcon(), bShouldCloseWindowAfterMenuSelection);
-	}
-}
 
 #undef LOCTEXT_NAMESPACE
diff --git a/Engine/Plugins/Experimental/ChaosClothAssetEditor/Source/ChaosClothAssetEditor/Private/ChaosClothAsset/SClothEditor3DViewportToolBar.h b/Engine/Plugins/Experimental/ChaosClothAssetEditor/Source/ChaosClothAssetEditor/Private/ChaosClothAsset/SClothEditor3DViewportToolBar.h
--- a/Engine/Plugins/Experimental/ChaosClothAssetEditor/Source/ChaosClothAssetEditor/Private/ChaosClothAsset/SClothEditor3DViewportToolBar.h
+++ b/Engine/Plugins/Experimental/ChaosClothAssetEditor/Source/ChaosClothAssetEditor/Private/ChaosClothAsset/SClothEditor3DViewportToolBar.h
@@ -29,8 +29,6 @@ private:
 	FText GetLODMenuLabel() const;
 	TSharedRef<SWidget> MakeLODMenu() const;
 
-	virtual void ExtendOptionsMenu(FMenuBuilder& OptionsMenuBuilder) const override;
-
 	/** The viewport that we are in */
 	TWeakPtr<class SChaosClothAssetEditor3DViewport> ChaosClothAssetEditor3DViewportPtr;
 
