Update GravityScale tooltip to explain that the map-based values are only used by the Force Based Solver. Also make the Force Based Solver default to 0 newton iterations. Existing assets with Force Based Solver disabled will pick up this new value.

diff --git a/Engine/Plugins/Experimental/ChaosClothAssetEditor/Source/ChaosClothAssetDataflowNodes/Private/ChaosClothAsset/SimulationSolverConfigNode.cpp b/Engine/Plugins/Experimental/ChaosClothAssetEditor/Source/ChaosClothAssetDataflowNodes/Private/ChaosClothAsset/SimulationSolverConfigNode.cpp
--- a/Engine/Plugins/Experimental/ChaosClothAssetEditor/Source/ChaosClothAssetDataflowNodes/Private/ChaosClothAsset/SimulationSolverConfigNode.cpp
+++ b/Engine/Plugins/Experimental/ChaosClothAssetEditor/Source/ChaosClothAssetDataflowNodes/Private/ChaosClothAsset/SimulationSolverConfigNode.cpp
@@ -24,3 +24,19 @@ void FChaosClothAssetSimulationSolverConfigNode::AddProperties(FPropertyHelper&
 	PropertyHelper.SetProperty(this, &CGResidualTolerance);
 	PropertyHelper.SetPropertyBool(this, &bDoQuasistatics);
 }
+
+void FChaosClothAssetSimulationSolverConfigNode::Serialize(FArchive& Ar)
+{
+	Super::Serialize(Ar);
+	Ar.UsingCustomVersion(FFortniteValkyrieBranchObjectVersion::GUID);
+	if (Ar.IsLoading())
+	{
+		if (Ar.CustomVer(FFortniteValkyrieBranchObjectVersion::GUID) < FFortniteValkyrieBranchObjectVersion::ChaosClothAssetWeightedMassAndGravity)
+		{
+			if (!bEnableForceBasedSolver)
+			{
+				NumNewtonIterations = 0;
+			}
+		}
+	}
+}
diff --git a/Engine/Plugins/Experimental/ChaosClothAssetEditor/Source/ChaosClothAssetDataflowNodes/Public/ChaosClothAsset/SimulationGravityConfigNode.h b/Engine/Plugins/Experimental/ChaosClothAssetEditor/Source/ChaosClothAssetDataflowNodes/Public/ChaosClothAsset/SimulationGravityConfigNode.h
--- a/Engine/Plugins/Experimental/ChaosClothAssetEditor/Source/ChaosClothAssetDataflowNodes/Public/ChaosClothAsset/SimulationGravityConfigNode.h
+++ b/Engine/Plugins/Experimental/ChaosClothAssetEditor/Source/ChaosClothAssetDataflowNodes/Public/ChaosClothAsset/SimulationGravityConfigNode.h
@@ -17,7 +17,9 @@ public:
 	UPROPERTY(EditAnywhere, Category = "Gravity Properties", Meta = (InlineEditConditionToggle))
 	bool bUseGravityOverride = false;
 
-	/** Scale factor applied to the world gravity and also to the clothing simulation interactor gravity. Does not affect the gravity if set using the override below. */
+	/** Scale factor applied to the world gravity and also to the clothing simulation interactor gravity. Does not affect the gravity if set using the override below. 
+	* Map only has an effect if using "Enable Force Based Solver" on solver config. Otherwise, only the Low value is used
+	*/
 	UPROPERTY(EditAnywhere, Category = "Gravity Properties", DisplayName = "Gravity Scale", Meta = (UIMin = "0", UIMax = "10", EditCondition = "!bUseGravityOverride"))
 	FChaosClothAssetWeightedValue GravityScaleWeighted = {true, 1.f, 1.f, TEXT("GravityScale")};
 
diff --git a/Engine/Plugins/Experimental/ChaosClothAssetEditor/Source/ChaosClothAssetDataflowNodes/Public/ChaosClothAsset/SimulationSolverConfigNode.h b/Engine/Plugins/Experimental/ChaosClothAssetEditor/Source/ChaosClothAssetDataflowNodes/Public/ChaosClothAsset/SimulationSolverConfigNode.h
--- a/Engine/Plugins/Experimental/ChaosClothAssetEditor/Source/ChaosClothAssetDataflowNodes/Public/ChaosClothAsset/SimulationSolverConfigNode.h
+++ b/Engine/Plugins/Experimental/ChaosClothAssetEditor/Source/ChaosClothAssetDataflowNodes/Public/ChaosClothAsset/SimulationSolverConfigNode.h
@@ -60,7 +60,7 @@ public:
 	 * Number of Newton iterations for force-based solver. Prototype only--very few constraints support this.
 	 */
 	UPROPERTY(EditAnywhere, Category = Experimental, meta = (UIMin = "1", UIMax = "10", ClampMin = "0", ClampMax = "100", EditCondition = "bEnableForceBasedSolver"))
-	int32 NumNewtonIterations = 1;
+	int32 NumNewtonIterations = 0;
 
 	/**
 	 * Max number of CG Iterations per linear solve
@@ -82,6 +82,8 @@ public:
 
 	FChaosClothAssetSimulationSolverConfigNode(const Dataflow::FNodeParameters& InParam, FGuid InGuid = FGuid::NewGuid());
 
+	virtual void Serialize(FArchive& Ar) override;
+
 private:
 	virtual void AddProperties(FPropertyHelper& PropertyHelper) const override;
 };
