Chaos Cloth Asset Editor - Fixed a crash in the USD Import node caused when both the file path value and the focus were changed at the same time, or near the same time.

diff --git a/Engine/Plugins/Experimental/ChaosClothAssetEditor/Source/ChaosClothAssetDataflowNodes/Private/ChaosClothAsset/ImportFilePathCustomization.cpp b/Engine/Plugins/Experimental/ChaosClothAssetEditor/Source/ChaosClothAssetDataflowNodes/Private/ChaosClothAsset/ImportFilePathCustomization.cpp
--- a/Engine/Plugins/Experimental/ChaosClothAssetEditor/Source/ChaosClothAssetDataflowNodes/Private/ChaosClothAsset/ImportFilePathCustomization.cpp
+++ b/Engine/Plugins/Experimental/ChaosClothAssetEditor/Source/ChaosClothAssetDataflowNodes/Private/ChaosClothAsset/ImportFilePathCustomization.cpp
@@ -156,7 +156,13 @@ namespace UE::Chaos::ClothAsset
 
 		if (PathStringProperty)
 		{
-			PathStringProperty->SetValue(FinalPath);
+			// The value can be set twice from pressing enter and losing the focus, most likely triggering two reentrant evaluations if not avoided
+			FString OldPath;
+			PathStringProperty->GetValue(OldPath);
+			if (OldPath != FinalPath)
+			{
+				PathStringProperty->SetValue(FinalPath);
+			}
 		}
 		FEditorDirectories::Get().SetLastDirectory(ELastDirectory::GENERIC_OPEN, FPaths::GetPath(PickedPath));
 	}
