Chaos Cloth Asset: Add WeightMapToSelectionNode

diff --git a/Engine/Plugins/Experimental/ChaosClothAssetEditor/Source/ChaosClothAssetDataflowNodes/Private/ChaosClothAsset/ChaosClothAssetDataflowNodesModule.cpp b/Engine/Plugins/Experimental/ChaosClothAssetEditor/Source/ChaosClothAssetDataflowNodes/Private/ChaosClothAsset/ChaosClothAssetDataflowNodesModule.cpp
--- a/Engine/Plugins/Experimental/ChaosClothAssetEditor/Source/ChaosClothAssetDataflowNodes/Private/ChaosClothAsset/ChaosClothAssetDataflowNodesModule.cpp
+++ b/Engine/Plugins/Experimental/ChaosClothAssetEditor/Source/ChaosClothAssetDataflowNodes/Private/ChaosClothAsset/ChaosClothAssetDataflowNodesModule.cpp
@@ -50,6 +50,7 @@
 #include "ChaosClothAsset/TransformUVsNode.h"
 #include "ChaosClothAsset/USDImportNode.h"
 #include "ChaosClothAsset/WeightedValueCustomization.h"
+#include "ChaosClothAsset/WeightMapToSelectionNode.h"
 #include "Dataflow/DataflowNodeColorsRegistry.h"
 #include "Dataflow/DataflowNodeFactory.h"
 #include "Modules/ModuleManager.h"
@@ -106,6 +107,7 @@ namespace UE::Chaos::ClothAsset
 			DATAFLOW_NODE_REGISTER_CREATION_FACTORY(FChaosClothAssetTransformPositionsNode);
 			DATAFLOW_NODE_REGISTER_CREATION_FACTORY(FChaosClothAssetTransformUVsNode);
 			DATAFLOW_NODE_REGISTER_CREATION_FACTORY(FChaosClothAssetUSDImportNode);
+			DATAFLOW_NODE_REGISTER_CREATION_FACTORY(FChaosClothAssetWeightMapToSelectionNode);
 			DATAFLOW_NODE_REGISTER_CREATION_FACTORY_NODE_COLORS_BY_CATEGORY("Cloth", FColorScheme::NodeHeader, FColorScheme::NodeBody);
 		}
 	}  // End namespace Private
diff --git a/Engine/Plugins/Experimental/ChaosClothAssetEditor/Source/ChaosClothAssetDataflowNodes/Private/ChaosClothAsset/SelectionToWeightMapNode.cpp b/Engine/Plugins/Experimental/ChaosClothAssetEditor/Source/ChaosClothAssetDataflowNodes/Private/ChaosClothAsset/SelectionToWeightMapNode.cpp
--- a/Engine/Plugins/Experimental/ChaosClothAssetEditor/Source/ChaosClothAssetDataflowNodes/Private/ChaosClothAsset/SelectionToWeightMapNode.cpp
+++ b/Engine/Plugins/Experimental/ChaosClothAssetEditor/Source/ChaosClothAssetDataflowNodes/Private/ChaosClothAsset/SelectionToWeightMapNode.cpp
@@ -42,7 +42,7 @@ void FChaosClothAssetSelectionToWeightMapNode::Evaluate(Dataflow::FContext& Cont
 					SelectionGroup == ClothCollectionGroup::SimVertices3D || 
 					SelectionGroup == ClothCollectionGroup::SimFaces)
 				{
-					const FName InMapName(WeightMapName.IsEmpty() ? SelectionName : WeightMapName);
+					const FName InMapName(WeightMapName.IsEmpty() ? InSelectionName : FName(WeightMapName));
 					ClothFacade.AddWeightMap(InMapName);
 					TArrayView<float> OutClothWeights = ClothFacade.GetWeightMap(InMapName);
 					if (OutClothWeights.Num() == ClothFacade.GetNumSimVertices3D())
@@ -125,7 +125,7 @@ void FChaosClothAssetSelectionToWeightMapNode::Evaluate(Dataflow::FContext& Cont
 	}
 	else if (Out->IsA<FString>(&WeightMapName))
 	{
-		SetValue(Context, WeightMapName.IsEmpty() ? SelectionName : WeightMapName, &WeightMapName);
+		SetValue(Context, WeightMapName.IsEmpty() ? GetValue<FString>(Context, &SelectionName) : WeightMapName, &WeightMapName);
 	}
 }
 
diff --git a/Engine/Plugins/Experimental/ChaosClothAssetEditor/Source/ChaosClothAssetDataflowNodes/Private/ChaosClothAsset/WeightMapToSelectionNode.cpp b/Engine/Plugins/Experimental/ChaosClothAssetEditor/Source/ChaosClothAssetDataflowNodes/Private/ChaosClothAsset/WeightMapToSelectionNode.cpp
new file mode 100644
--- /dev/null
+++ b/Engine/Plugins/Experimental/ChaosClothAssetEditor/Source/ChaosClothAssetDataflowNodes/Private/ChaosClothAsset/WeightMapToSelectionNode.cpp
@@ -0,0 +1,123 @@
+// Copyright Epic Games, Inc. All Rights Reserved.
+
+#include "ChaosClothAsset/WeightMapToSelectionNode.h"
+#include "ChaosClothAsset/ClothCollectionGroup.h"
+#include "ChaosClothAsset/ClothDataflowTools.h"
+#include "ChaosClothAsset/CollectionClothFacade.h"
+#include "ChaosClothAsset/CollectionClothSelectionFacade.h"
+#include "Dataflow/DataflowInputOutput.h"
+
+#include UE_INLINE_GENERATED_CPP_BY_NAME(WeightMapToSelectionNode)
+#define LOCTEXT_NAMESPACE "FChaosClothAssetWeightMapToSelectionNode"
+
+FChaosClothAssetWeightMapToSelectionNode::FChaosClothAssetWeightMapToSelectionNode(const Dataflow::FNodeParameters& InParam, FGuid InGuid)
+	: FDataflowNode(InParam, InGuid)
+{
+	RegisterInputConnection(&Collection);
+	RegisterOutputConnection(&Collection, &Collection);
+	RegisterInputConnection(&WeightMapName);
+	RegisterOutputConnection(&SelectionName);
+}
+
+void FChaosClothAssetWeightMapToSelectionNode::Evaluate(Dataflow::FContext& Context, const FDataflowOutput* Out) const
+{
+	using namespace UE::Chaos::ClothAsset;
+
+	if (Out->IsA<FManagedArrayCollection>(&Collection))
+	{
+		FManagedArrayCollection InCollection = GetValue<FManagedArrayCollection>(Context, &Collection);
+		const TSharedRef<FManagedArrayCollection> ClothCollection = MakeShared<FManagedArrayCollection>(MoveTemp(InCollection));
+		FCollectionClothConstFacade ClothFacade(ClothCollection);
+		const FName InWeightMapName(*GetValue<FString>(Context, &WeightMapName));
+
+		if (ClothFacade.IsValid())
+		{
+			if (InWeightMapName != NAME_None && ClothFacade.HasWeightMap(InWeightMapName))
+			{
+				TConstArrayView<float> WeightMap = ClothFacade.GetWeightMap(InWeightMapName);
+
+				FName SelectionGroup;
+				switch (SelectionType)
+				{
+				default:
+				case EChaosClothAssetWeightMapConvertableSelectionType::SimVertices3D:
+					SelectionGroup = ClothCollectionGroup::SimVertices3D;
+					break;
+				case EChaosClothAssetWeightMapConvertableSelectionType::SimVertices2D:
+					SelectionGroup = ClothCollectionGroup::SimVertices2D;
+					break;
+				case EChaosClothAssetWeightMapConvertableSelectionType::SimFaces:
+					SelectionGroup = ClothCollectionGroup::SimFaces;
+					break;
+				}
+
+				FCollectionClothSelectionFacade SelectionFacade(ClothCollection);
+				SelectionFacade.DefineSchema();
+
+				const FName InSelectionName(SelectionName.IsEmpty() ? InWeightMapName : FName(SelectionName));
+				TSet<int32>& SelectionSet = SelectionFacade.FindOrAddSelectionSet(InSelectionName, SelectionGroup);
+				SelectionSet.Reset();
+
+				if (SelectionGroup == ClothCollectionGroup::SimVertices3D)
+				{
+					for (int32 VertexIndex = 0; VertexIndex < WeightMap.Num(); ++VertexIndex)
+					{
+						if (WeightMap[VertexIndex] >= SelectionThreshold)
+						{
+							SelectionSet.Add(VertexIndex);
+						}
+					}
+				}
+				else if (SelectionGroup == ClothCollectionGroup::SimVertices2D)
+				{
+					const TConstArrayView<TArray<int32>> Vertex3DTo2D = ClothFacade.GetSimVertex2DLookup();
+
+					for (int32 Vertex3DIndex = 0; Vertex3DIndex < WeightMap.Num(); ++Vertex3DIndex)
+					{
+						// Select all 2D vertices that correspond with 3D vertices above the threshold
+
+						if (WeightMap[Vertex3DIndex] >= SelectionThreshold)
+						{
+							for (const int32 Vertex2DIndex : Vertex3DTo2D[Vertex3DIndex])
+							{
+								SelectionSet.Add(Vertex2DIndex);
+							}
+						}
+					}
+				}
+				else
+				{
+					check(SelectionGroup == ClothCollectionGroup::SimFaces);
+
+					const TConstArrayView<FIntVector3> SimIndices3D = ClothFacade.GetSimIndices3D();
+					// Select a face if all 3 vertices are above the threshold
+					for (int32 FaceIndex = 0; FaceIndex < SimIndices3D.Num(); ++FaceIndex)
+					{
+						const FIntVector3& Index = SimIndices3D[FaceIndex];
+						if (WeightMap[Index[0]] >= SelectionThreshold &&
+							WeightMap[Index[1]] >= SelectionThreshold &&
+							WeightMap[Index[2]] >= SelectionThreshold)
+						{
+							SelectionSet.Add(FaceIndex);
+						}
+					}
+				}
+			}
+			else
+			{
+				FClothDataflowTools::LogAndToastWarning(*this,
+					LOCTEXT("WeightMapNameNotFoundHeadline", "Invalid weight map name."),
+					FText::Format(LOCTEXT("WeightMapeNotFoundDetails", "Weight map \"{0}\" was not found in the collection."),
+						FText::FromName(InWeightMapName)));
+
+			}
+		}
+
+		SetValue(Context, MoveTemp(*ClothCollection), &Collection);
+	}
+	else if (Out->IsA<FString>(&SelectionName))
+	{
+		SetValue(Context, SelectionName.IsEmpty() ? GetValue<FString>(Context, &WeightMapName) : SelectionName, &SelectionName);
+	}
+}
+#undef LOCTEXT_NAMESPACE
diff --git a/Engine/Plugins/Experimental/ChaosClothAssetEditor/Source/ChaosClothAssetDataflowNodes/Public/ChaosClothAsset/WeightMapToSelectionNode.h b/Engine/Plugins/Experimental/ChaosClothAssetEditor/Source/ChaosClothAssetDataflowNodes/Public/ChaosClothAsset/WeightMapToSelectionNode.h
new file mode 100644
--- /dev/null
+++ b/Engine/Plugins/Experimental/ChaosClothAssetEditor/Source/ChaosClothAssetDataflowNodes/Public/ChaosClothAsset/WeightMapToSelectionNode.h
@@ -0,0 +1,58 @@
+// Copyright Epic Games, Inc. All Rights Reserved.
+
+#pragma once 
+
+#include "Dataflow/DataflowNode.h"
+#include "GeometryCollection/ManagedArrayCollection.h"
+#include "WeightMapToSelectionNode.generated.h"
+
+/** What type of element to convert to */
+UENUM()
+enum class EChaosClothAssetWeightMapConvertableSelectionType : uint8
+{
+	/** 2D simulation vertices */
+	SimVertices2D,
+
+	/** 3D simulation vertices */
+	SimVertices3D,
+
+	/** Simulation faces (2D/3D are the same) */
+	SimFaces,
+};
+
+/** Convert a vertex weight map to an integer selection set. */
+USTRUCT(Meta = (DataflowCloth))
+struct FChaosClothAssetWeightMapToSelectionNode : public FDataflowNode
+{
+	GENERATED_USTRUCT_BODY()
+	DATAFLOW_NODE_DEFINE_INTERNAL(FChaosClothAssetWeightMapToSelectionNode, "WeightMapToSelection", "Cloth", "Cloth Weight Map To Selection")
+
+public:
+
+	UPROPERTY(Meta = (Dataflowinput, DataflowOutput, DataflowPassthrough = "Collection"))
+	FManagedArrayCollection Collection;
+
+	/** The name of the Weight Map to convert. */
+	UPROPERTY(Meta = (DataflowInput))
+	FString WeightMapName;
+
+	/**
+	 * The name of the select attribute that will be added to the collection.
+	 * If left empty the same name as the Weight Map  name will be used instead.
+	 */
+	UPROPERTY(EditAnywhere, Category = "Weight Map To Selection", Meta = (DataflowOutput))
+	FString SelectionName;
+
+	/** The type of element the selection refers to */
+	UPROPERTY(EditAnywhere, Category = "Weight Map To Selection")
+	EChaosClothAssetWeightMapConvertableSelectionType SelectionType = EChaosClothAssetWeightMapConvertableSelectionType::SimVertices3D;
+
+	/** Map values above this will be selected. */
+	UPROPERTY(EditAnywhere, Category = "Weight Map To Selection", Meta = (ClampMin = "0", ClampMax = "1"))
+	float SelectionThreshold = 0.95f;
+
+	FChaosClothAssetWeightMapToSelectionNode(const Dataflow::FNodeParameters& InParam, FGuid InGuid = FGuid::NewGuid());
+private:
+
+	virtual void Evaluate(Dataflow::FContext& Context, const FDataflowOutput* Out) const override;
+};
