Chaos Cloth Editor: Add action to import selection set from the collection in the Selection Tool

diff --git a/Engine/Plugins/Experimental/ChaosClothAssetEditor/Source/ChaosClothAssetEditor/Private/ChaosClothAsset/ClothEditorMode.cpp b/Engine/Plugins/Experimental/ChaosClothAssetEditor/Source/ChaosClothAssetEditor/Private/ChaosClothAsset/ClothEditorMode.cpp
--- a/Engine/Plugins/Experimental/ChaosClothAssetEditor/Source/ChaosClothAssetEditor/Private/ChaosClothAsset/ClothEditorMode.cpp
+++ b/Engine/Plugins/Experimental/ChaosClothAssetEditor/Source/ChaosClothAssetEditor/Private/ChaosClothAsset/ClothEditorMode.cpp
@@ -454,9 +454,10 @@ bool UChaosClothAssetEditorMode::IsComponentSelected(const UPrimitiveComponent*
 }
 
 
-void UChaosClothAssetEditorMode::SetSelectedClothCollection(TSharedPtr<FManagedArrayCollection> Collection)
+void UChaosClothAssetEditorMode::SetSelectedClothCollection(TSharedPtr<FManagedArrayCollection> Collection, TSharedPtr<FManagedArrayCollection> InputCollection)
 {
 	SelectedClothCollection = Collection;
+	SelectedInputClothCollection = InputCollection;
 	ReinitializeDynamicMeshComponents();
 
 	// The first time we get a valid mesh, refocus the camera on it
@@ -470,6 +471,13 @@ TSharedPtr<FManagedArrayCollection> UChaosClothAssetEditorMode::GetClothCollecti
 	// TODO: If no cloth collection node is selected, show the ClothAsset's collection. In this case, also ensure that any interactive tools are disabled. (UE-181574)
 }
 
+TSharedPtr<FManagedArrayCollection> UChaosClothAssetEditorMode::GetInputClothCollection()
+{
+	return SelectedInputClothCollection;
+
+	// TODO: If no cloth collection node is selected, show the ClothAsset's collection. In this case, also ensure that any interactive tools are disabled. (UE-181574)
+}
+
 
 void UChaosClothAssetEditorMode::InitializeSeamDraw()
 {
@@ -812,7 +820,7 @@ void UChaosClothAssetEditorMode::ReinitializeDynamicMeshComponents()
 	UClothEditorContextObject* EditorContextObject = RestSpaceToolsContext->ContextObjectStore->FindContext<UClothEditorContextObject>();
 	if (ensure(EditorContextObject))
 	{
-		EditorContextObject->SetClothCollection(ConstructionViewMode, Collection);
+		EditorContextObject->SetClothCollection(ConstructionViewMode, Collection, GetInputClothCollection());
 	}
 }
 
@@ -1286,7 +1294,7 @@ void UChaosClothAssetEditorMode::InitializeContextObject()
 		RestSpaceToolsContext->ContextObjectStore->AddContextObject(EditorContextObject);
 	}
 
-	EditorContextObject->Init(DataflowGraphEditor, ConstructionViewMode, SelectedClothCollection);
+	EditorContextObject->Init(DataflowGraphEditor, ConstructionViewMode, SelectedClothCollection, SelectedInputClothCollection);
 
 	check(EditorContextObject);
 
diff --git a/Engine/Plugins/Experimental/ChaosClothAssetEditor/Source/ChaosClothAssetEditor/Private/ChaosClothAsset/ClothEditorMode.h b/Engine/Plugins/Experimental/ChaosClothAssetEditor/Source/ChaosClothAssetEditor/Private/ChaosClothAsset/ClothEditorMode.h
--- a/Engine/Plugins/Experimental/ChaosClothAssetEditor/Source/ChaosClothAssetEditor/Private/ChaosClothAsset/ClothEditorMode.h
+++ b/Engine/Plugins/Experimental/ChaosClothAssetEditor/Source/ChaosClothAssetEditor/Private/ChaosClothAsset/ClothEditorMode.h
@@ -160,10 +160,11 @@ private:
 	void FirstTimeFocusRestSpaceViewport();
 
 	// intended to be called by the toolkit when selected node in the Dataflow graph changes
-	void SetSelectedClothCollection(TSharedPtr<FManagedArrayCollection> Collection);
+	void SetSelectedClothCollection(TSharedPtr<FManagedArrayCollection> Collection, TSharedPtr<FManagedArrayCollection> InputCollection = nullptr);
 
 	// gets the currently selected cloth collection, as specified by the toolkit
 	TSharedPtr<FManagedArrayCollection> GetClothCollection();
+	TSharedPtr<FManagedArrayCollection> GetInputClothCollection();
 
 	void SetDataflowGraphEditor(TSharedPtr<SDataflowGraphEditor> InGraphEditor);
 	
@@ -280,6 +281,7 @@ private:
 	TObjectPtr<UEditorInteractiveToolsContext> ActiveToolsContext = nullptr;
 
 	TSharedPtr<FManagedArrayCollection> SelectedClothCollection = nullptr;
+	TSharedPtr<FManagedArrayCollection> SelectedInputClothCollection = nullptr;
 
 	// Correspondence between node types and commands to launch tools
 	TMap<FName, TSharedPtr<const FUICommandInfo>> NodeTypeToToolCommandMap;
diff --git a/Engine/Plugins/Experimental/ChaosClothAssetEditor/Source/ChaosClothAssetEditor/Private/ChaosClothAsset/ClothEditorToolkit.cpp b/Engine/Plugins/Experimental/ChaosClothAssetEditor/Source/ChaosClothAssetEditor/Private/ChaosClothAsset/ClothEditorToolkit.cpp
--- a/Engine/Plugins/Experimental/ChaosClothAssetEditor/Source/ChaosClothAssetEditor/Private/ChaosClothAsset/ClothEditorToolkit.cpp
+++ b/Engine/Plugins/Experimental/ChaosClothAssetEditor/Source/ChaosClothAssetEditor/Private/ChaosClothAsset/ClothEditorToolkit.cpp
@@ -909,9 +909,10 @@ void FChaosClothAssetEditorToolkit::EvaluateNode(FDataflowNode* Node, bool bForc
 			// Refresh the construction viewport
 			if (SelectedDataflowNode)
 			{
+				TSharedPtr<FManagedArrayCollection> InputCollection = GetInputClothCollectionIfPossible(SelectedDataflowNode, DataflowContext);
 				TSharedPtr<FManagedArrayCollection> Collection = GetClothCollectionIfPossible(SelectedDataflowNode, DataflowContext);
 				UChaosClothAssetEditorMode* const ClothMode = CastChecked<UChaosClothAssetEditorMode>(EditorModeManager->GetActiveScriptableMode(UChaosClothAssetEditorMode::EM_ChaosClothAssetEditorModeId));
-				ClothMode->SetSelectedClothCollection(Collection);
+				ClothMode->SetSelectedClothCollection(Collection, InputCollection);
 			}
 		}
 	}
@@ -1073,9 +1074,39 @@ TSharedPtr<FManagedArrayCollection> FChaosClothAssetEditorToolkit::GetClothColle
 	return TSharedPtr<FManagedArrayCollection>();
 }
 
+
+TSharedPtr<FManagedArrayCollection> FChaosClothAssetEditorToolkit::GetInputClothCollectionIfPossible(const TSharedPtr<FDataflowNode> InDataflowNode, const TSharedPtr<Dataflow::FEngineContext> Context)
+{
+	if (InDataflowNode && Context)
+	{
+		for (const FDataflowInput* const Input : InDataflowNode->GetInputs())
+		{
+			if (Input->GetType() == FName("FManagedArrayCollection"))
+			{
+				const FManagedArrayCollection DefaultValue;
+				TSharedRef<FManagedArrayCollection> Collection = MakeShared<FManagedArrayCollection>(Input->GetValue<FManagedArrayCollection>(*Context, DefaultValue));
+
+				// see if the input collection is a ClothCollection
+				const UE::Chaos::ClothAsset::FCollectionClothConstFacade ClothFacade(Collection);
+				if (ClothFacade.IsValid())
+				{
+					return Collection;
+				}
+
+				// The cloth collection schema must be applied to prevent the dynamic mesh conversion and tools from crashing trying to access invalid facades
+				break;
+			}
+		}
+	}
+
+	return TSharedPtr<FManagedArrayCollection>();
+}
+
+
 void FChaosClothAssetEditorToolkit::OnNodeSelectionChanged(const TSet<UObject*>& NewSelection)
 {
 	TSharedPtr<FManagedArrayCollection> Collection = nullptr;
+	TSharedPtr<FManagedArrayCollection> InputCollection = nullptr;
 
 	// Get any selected node with a ClothCollection output
 	// Also, set the selected node(s) to be the Dataflow's RenderTargets
@@ -1128,6 +1159,7 @@ void FChaosClothAssetEditorToolkit::OnNodeSelectionChanged(const TSet<UObject*>&
 
 						if (SelectedDataflowNode)
 						{
+							InputCollection = GetInputClothCollectionIfPossible(SelectedDataflowNode, DataflowContext);
 							Collection = GetClothCollectionIfPossible(SelectedDataflowNode, DataflowContext);
 							if (DataflowContext.IsValid())  // TODO: The context shouldn't be nullptr on the first node creation
 							{
@@ -1188,7 +1220,7 @@ void FChaosClothAssetEditorToolkit::OnNodeSelectionChanged(const TSet<UObject*>&
 			}
 
 			// Update the Construction viewport with the newly selected node's Collection
-			ClothMode->SetSelectedClothCollection(Collection);
+			ClothMode->SetSelectedClothCollection(Collection, InputCollection);
 		}
 
 		if (Outliner)
diff --git a/Engine/Plugins/Experimental/ChaosClothAssetEditor/Source/ChaosClothAssetEditor/Private/ChaosClothAsset/ClothEditorToolkit.h b/Engine/Plugins/Experimental/ChaosClothAssetEditor/Source/ChaosClothAssetEditor/Private/ChaosClothAsset/ClothEditorToolkit.h
--- a/Engine/Plugins/Experimental/ChaosClothAssetEditor/Source/ChaosClothAssetEditor/Private/ChaosClothAsset/ClothEditorToolkit.h
+++ b/Engine/Plugins/Experimental/ChaosClothAssetEditor/Source/ChaosClothAssetEditor/Private/ChaosClothAsset/ClothEditorToolkit.h
@@ -128,6 +128,8 @@ private:
 
 	TSharedPtr<FManagedArrayCollection> GetClothCollectionIfPossible(const TSharedPtr<FDataflowNode> InDataflowNode, const TSharedPtr<Dataflow::FEngineContext> Context);
 
+	TSharedPtr<FManagedArrayCollection> GetInputClothCollectionIfPossible(const TSharedPtr<FDataflowNode> InDataflowNode, const TSharedPtr<Dataflow::FEngineContext> Context);
+
 	// DataflowEditorActions
 	void OnPropertyValueChanged(const FPropertyChangedEvent& PropertyChangedEvent);
 	bool OnNodeVerifyTitleCommit(const FText& NewText, UEdGraphNode* GraphNode, FText& OutErrorMessage) const;
diff --git a/Engine/Plugins/Experimental/ChaosClothAssetEditor/Source/ChaosClothAssetEditorTools/Private/ChaosClothAsset/ClothEditorContextObject.cpp b/Engine/Plugins/Experimental/ChaosClothAssetEditor/Source/ChaosClothAssetEditorTools/Private/ChaosClothAsset/ClothEditorContextObject.cpp
--- a/Engine/Plugins/Experimental/ChaosClothAssetEditor/Source/ChaosClothAssetEditorTools/Private/ChaosClothAsset/ClothEditorContextObject.cpp
+++ b/Engine/Plugins/Experimental/ChaosClothAssetEditor/Source/ChaosClothAssetEditorTools/Private/ChaosClothAsset/ClothEditorContextObject.cpp
@@ -8,15 +8,17 @@
 
 #include UE_INLINE_GENERATED_CPP_BY_NAME(ClothEditorContextObject)
 
-void UClothEditorContextObject::Init(TWeakPtr<SDataflowGraphEditor> InDataflowGraphEditor,UE::Chaos::ClothAsset::EClothPatternVertexType InConstructionViewMode, TWeakPtr<FManagedArrayCollection> InSelectedClothCollection)
+void UClothEditorContextObject::Init(TWeakPtr<SDataflowGraphEditor> InDataflowGraphEditor,UE::Chaos::ClothAsset::EClothPatternVertexType InConstructionViewMode, TWeakPtr<FManagedArrayCollection> InSelectedClothCollection, TWeakPtr<FManagedArrayCollection> InSelectedInputClothCollection)
 {
 	DataflowGraphEditor = InDataflowGraphEditor;
 	ConstructionViewMode = InConstructionViewMode;
 	SelectedClothCollection = InSelectedClothCollection;
+	SelectedInputClothCollection = InSelectedInputClothCollection;
 }
 
-void UClothEditorContextObject::SetClothCollection(UE::Chaos::ClothAsset::EClothPatternVertexType ViewMode, TWeakPtr<FManagedArrayCollection> ClothCollection)
+void UClothEditorContextObject::SetClothCollection(UE::Chaos::ClothAsset::EClothPatternVertexType ViewMode, TWeakPtr<FManagedArrayCollection> ClothCollection, TWeakPtr<FManagedArrayCollection> InputClothCollection)
 {
 	ConstructionViewMode = ViewMode;
 	SelectedClothCollection = ClothCollection;
+	SelectedInputClothCollection = InputClothCollection;
 }
\ No newline at end of file
diff --git a/Engine/Plugins/Experimental/ChaosClothAssetEditor/Source/ChaosClothAssetEditorTools/Private/ChaosClothAsset/ClothMeshSelectionTool.cpp b/Engine/Plugins/Experimental/ChaosClothAssetEditor/Source/ChaosClothAssetEditorTools/Private/ChaosClothAsset/ClothMeshSelectionTool.cpp
--- a/Engine/Plugins/Experimental/ChaosClothAssetEditor/Source/ChaosClothAssetEditorTools/Private/ChaosClothAsset/ClothMeshSelectionTool.cpp
+++ b/Engine/Plugins/Experimental/ChaosClothAssetEditor/Source/ChaosClothAssetEditorTools/Private/ChaosClothAsset/ClothMeshSelectionTool.cpp
@@ -6,6 +6,7 @@
 #include "ChaosClothAsset/ClothPatternVertexType.h"
 #include "ChaosClothAsset/SelectionNode.h"
 #include "ChaosClothAsset/CollectionClothFacade.h"
+#include "ChaosClothAsset/CollectionClothSelectionFacade.h"
 #include "ChaosClothAsset/WeightedValue.h"
 #include "InteractiveToolManager.h"
 #include "PreviewMesh.h"
@@ -63,6 +64,15 @@ UInteractiveTool* UClothMeshSelectionToolBuilder::BuildTool(const FToolBuilderSt
 	return NewTool;
 }
 
+// ------------------- Actions ----------------------
+void UClothMeshSelectionToolActions::PostAction(EClothMeshSelectionToolActions Action)
+{
+	if (ParentTool.IsValid())
+	{
+		ParentTool->RequestAction(Action);
+	}
+}
+
 // ------------------- Properties -------------------
 
 void UClothMeshSelectionToolProperties::PostEditChangeProperty(FPropertyChangedEvent& PropertyChangedEvent)
@@ -231,6 +241,13 @@ void UClothMeshSelectionTool::Setup()
 		}
 	}
 
+	// 
+	// Actions
+	//
+
+	ActionsProps = NewObject<UClothMeshSelectionToolActions>();
+	ActionsProps->Initialize(this);
+	AddToolPropertySource(ActionsProps);
 
 	//
 	// Properties
@@ -297,6 +314,18 @@ void UClothMeshSelectionTool::DrawHUD(FCanvas* Canvas, IToolsContextRenderAPI* R
 	SelectionMechanic->DrawHUD(Canvas, RenderAPI);
 }
 
+void UClothMeshSelectionTool::OnTick(float DeltaTime)
+{
+	USingleSelectionMeshEditingTool::OnTick(DeltaTime);
+
+	if (bHavePendingAction)
+	{
+		ApplyAction(PendingAction);
+		bHavePendingAction = false;
+		PendingAction = EClothMeshSelectionToolActions::NoAction;
+	}
+}
+
 bool UClothMeshSelectionTool::CanAccept() const
 {
 	return bAnyChangeMade;
@@ -442,5 +471,129 @@ void UClothMeshSelectionTool::UpdateSelectedNode()
 	MeshSelectionNode->Invalidate();
 }
 
+void UClothMeshSelectionTool::RequestAction(EClothMeshSelectionToolActions ActionType)
+{
+	if (!bHavePendingAction)
+	{
+		PendingAction = ActionType;
+		bHavePendingAction = true;
+	}
+}
+
+
+void UClothMeshSelectionTool::ApplyAction(EClothMeshSelectionToolActions ActionType)
+{
+	switch (ActionType)
+	{
+	case EClothMeshSelectionToolActions::ImportFromCollection:
+		ImportFromCollectionAction();
+		break;
+	}
+}
+
+void UClothMeshSelectionTool::ImportFromCollectionAction()
+{
+	if (const TSharedPtr<const FManagedArrayCollection> ClothCollection = ClothEditorContextObject->GetSelectedInputClothCollection().Pin())
+	{
+		using namespace UE::Chaos::ClothAsset;
+		const FCollectionClothSelectionConstFacade SelectionFacade(ClothCollection.ToSharedRef());
+		if (SelectionFacade.IsValid())
+		{
+			const EClothPatternVertexType ViewMode = ClothEditorContextObject->GetConstructionViewMode();
+			FName GroupName = ClothCollectionGroup::SimVertices2D;
+			if (SelectionMechanic->Properties->bSelectVertices)
+			{
+				check(!SelectionMechanic->Properties->bSelectEdges);
+				check(!SelectionMechanic->Properties->bSelectFaces);
+
+				switch (ViewMode)
+				{
+				case UE::Chaos::ClothAsset::EClothPatternVertexType::Sim2D:
+					GroupName = ClothCollectionGroup::SimVertices2D;
+					break;
+				case UE::Chaos::ClothAsset::EClothPatternVertexType::Sim3D:
+					GroupName = ClothCollectionGroup::SimVertices3D;
+					break;
+				case UE::Chaos::ClothAsset::EClothPatternVertexType::Render:
+					GroupName = ClothCollectionGroup::RenderVertices;
+					break;
+				}
+			}
+			else
+			{
+				switch (ViewMode)
+				{
+				case UE::Chaos::ClothAsset::EClothPatternVertexType::Sim2D:
+				case UE::Chaos::ClothAsset::EClothPatternVertexType::Sim3D:
+					GroupName = ClothCollectionGroup::SimFaces;
+					break;
+				case UE::Chaos::ClothAsset::EClothPatternVertexType::Render:
+					GroupName = ClothCollectionGroup::RenderFaces;
+					break;
+				}
+			}
+			const FName InSelectionName(ToolProperties->Name);
+			if (const TSet<int32>* const SelectionSet = SelectionFacade.FindSelectionSet(InSelectionName))
+			{
+				const FName& ExistingSelectionGroup = SelectionFacade.GetSelectionGroup(InSelectionName);
+				if (ExistingSelectionGroup == GroupName)
+				{
+					auto AppendVerticesIfValid = [this]<typename T>(TSet<int32>&Dest, const T & Source)
+					{
+						PreviewMesh->ProcessMesh([this, &Dest, &Source](const UE::Geometry::FDynamicMesh3& Mesh)
+						{
+							for (const int32& VertexIndex : Source)
+							{
+								if (Mesh.IsVertex(VertexIndex))
+								{
+									Dest.Add(VertexIndex);
+								}
+							}
+						});
+					};
+
+					auto AppendFacesIfValid = [this](TSet<int32>& Dest, const TSet<int32>& Source)
+					{
+						PreviewMesh->ProcessMesh([this, &Dest, &Source](const UE::Geometry::FDynamicMesh3& Mesh)
+						{
+							for (const int32& FaceIndex : Source)
+							{
+								if (Mesh.IsTriangle(FaceIndex))
+								{
+									Dest.Add(FaceIndex);
+								}
+							}
+						});
+					};
+
+					UE::Geometry::FGroupTopologySelection OutSelection;
+
+					if (GroupName == ClothCollectionGroup::SimVertices2D ||
+						GroupName == ClothCollectionGroup::SimVertices3D ||
+						GroupName == ClothCollectionGroup::RenderVertices)
+					{
+						if (bHasNonManifoldMapping)
+						{
+							for (const int32 SelectionIndex : *SelectionSet)
+							{
+								AppendVerticesIfValid(OutSelection.SelectedCornerIDs, SelectionToDynamicMesh[SelectionIndex]);
+							}
+						}
+						else
+						{
+							AppendVerticesIfValid(OutSelection.SelectedCornerIDs, *SelectionSet);
+						}
+					}
+					else if (GroupName == ClothCollectionGroup::SimFaces ||
+						GroupName == ClothCollectionGroup::RenderFaces)
+					{
+						AppendFacesIfValid(OutSelection.SelectedGroupIDs, *SelectionSet);
+					}
+					SelectionMechanic->SetSelection(OutSelection);
+				}
+			}
+		}
+	}
+}
 
 #undef LOCTEXT_NAMESPACE
diff --git a/Engine/Plugins/Experimental/ChaosClothAssetEditor/Source/ChaosClothAssetEditorTools/Public/ChaosClothAsset/ClothEditorContextObject.h b/Engine/Plugins/Experimental/ChaosClothAssetEditor/Source/ChaosClothAssetEditorTools/Public/ChaosClothAsset/ClothEditorContextObject.h
--- a/Engine/Plugins/Experimental/ChaosClothAssetEditor/Source/ChaosClothAssetEditorTools/Public/ChaosClothAsset/ClothEditorContextObject.h
+++ b/Engine/Plugins/Experimental/ChaosClothAssetEditor/Source/ChaosClothAssetEditorTools/Public/ChaosClothAsset/ClothEditorContextObject.h
@@ -24,7 +24,7 @@ class CHAOSCLOTHASSETEDITORTOOLS_API UClothEditorContextObject : public UObject
 
 public:
 
-	void Init(TWeakPtr<SDataflowGraphEditor> DataflowGraphEditor, UE::Chaos::ClothAsset::EClothPatternVertexType InConstructionViewMode, TWeakPtr<FManagedArrayCollection> SelectedClothCollection);
+	void Init(TWeakPtr<SDataflowGraphEditor> DataflowGraphEditor, UE::Chaos::ClothAsset::EClothPatternVertexType InConstructionViewMode, TWeakPtr<FManagedArrayCollection> SelectedClothCollection, TWeakPtr<FManagedArrayCollection> SelectedInputClothCollection = nullptr);
 
 	/**
 	* Get a single selected node of the specified type. Return nullptr if the specified node is not selected, or if multiple nodes are selected
@@ -47,9 +47,10 @@ public:
 		return nullptr;
 	}
 
-	void SetClothCollection(UE::Chaos::ClothAsset::EClothPatternVertexType ViewMode, TWeakPtr<FManagedArrayCollection> ClothCollection);
+	void SetClothCollection(UE::Chaos::ClothAsset::EClothPatternVertexType ViewMode, TWeakPtr<FManagedArrayCollection> ClothCollection, TWeakPtr<FManagedArrayCollection> InputClothCollection = nullptr);
 
 	const TWeakPtr<const FManagedArrayCollection> GetSelectedClothCollection() const { return SelectedClothCollection; }
+	const TWeakPtr<const FManagedArrayCollection> GetSelectedInputClothCollection() const { return SelectedInputClothCollection; }
 	UE::Chaos::ClothAsset::EClothPatternVertexType GetConstructionViewMode() const { return ConstructionViewMode; }
 private:
 
@@ -57,6 +58,7 @@ private:
 
 	UE::Chaos::ClothAsset::EClothPatternVertexType ConstructionViewMode;
 	TWeakPtr<const FManagedArrayCollection> SelectedClothCollection;
+	TWeakPtr<const FManagedArrayCollection> SelectedInputClothCollection;
 };
 
 
diff --git a/Engine/Plugins/Experimental/ChaosClothAssetEditor/Source/ChaosClothAssetEditorTools/Public/ChaosClothAsset/ClothMeshSelectionTool.h b/Engine/Plugins/Experimental/ChaosClothAssetEditor/Source/ChaosClothAssetEditorTools/Public/ChaosClothAsset/ClothMeshSelectionTool.h
--- a/Engine/Plugins/Experimental/ChaosClothAssetEditor/Source/ChaosClothAssetEditorTools/Public/ChaosClothAsset/ClothMeshSelectionTool.h
+++ b/Engine/Plugins/Experimental/ChaosClothAssetEditor/Source/ChaosClothAssetEditorTools/Public/ChaosClothAsset/ClothMeshSelectionTool.h
@@ -30,6 +30,34 @@ private:
 	virtual bool CanSetConstructionViewWireframeActive() const { return false; }
 };
 
+UENUM()
+enum class EClothMeshSelectionToolActions
+{
+	NoAction,
+
+	ImportFromCollection
+};
+
+UCLASS()
+class CHAOSCLOTHASSETEDITORTOOLS_API UClothMeshSelectionToolActions :  public UInteractiveToolPropertySet
+{
+	GENERATED_BODY()
+
+public:
+
+	TWeakObjectPtr<UClothMeshSelectionTool> ParentTool;
+
+	void Initialize(UClothMeshSelectionTool* ParentToolIn) { ParentTool = ParentToolIn; }
+
+	void PostAction(EClothMeshSelectionToolActions Action);
+
+	UFUNCTION(CallInEditor, Category = Operations)
+	void ImportFromCollection()
+	{
+		PostAction(EClothMeshSelectionToolActions::ImportFromCollection);
+	}
+};
+
 UCLASS()
 class CHAOSCLOTHASSETEDITORTOOLS_API UClothMeshSelectionToolProperties : public UInteractiveToolPropertySet
 {
@@ -66,6 +94,8 @@ private:
 	virtual void Render(IToolsContextRenderAPI* RenderAPI) override;
 	virtual void DrawHUD(FCanvas* Canvas, IToolsContextRenderAPI* RenderAPI) override;
 
+	virtual void OnTick(float DeltaTime) override;
+
 	virtual bool HasCancel() const override { return true; }
 	virtual bool HasAccept() const override { return true; }
 	virtual bool CanAccept() const override;
@@ -94,5 +124,22 @@ private:
 	TArray<int32> DynamicMeshToSelection;
 	TArray<TArray<int32>> SelectionToDynamicMesh;
 
+	//
+	// Action support
+	//
+
+public:
+	virtual void RequestAction(EClothMeshSelectionToolActions ActionType);
+
+	UPROPERTY()
+	TObjectPtr<UClothMeshSelectionToolActions> ActionsProps;
+
+private:
+	bool bHavePendingAction = false;
+	EClothMeshSelectionToolActions PendingAction;
+	virtual void ApplyAction(EClothMeshSelectionToolActions ActionType);
+
+	void ImportFromCollectionAction();
+
 };
 
