Add an option to the ChaosCloth TransferSkinWeights node that allows to choose the max number of bone influences per vertex.
Propagate this change to the TransferBoneWeights and SmoothBoneWeights geometry operators by allowing to manually specify the max number of influences per bone for the output per-vertex weights.

#rnx

diff --git a/Engine/Plugins/Experimental/ChaosClothAssetEditor/Source/ChaosClothAssetDataflowNodes/Private/ChaosClothAsset/TransferSkinWeightsNode.cpp b/Engine/Plugins/Experimental/ChaosClothAssetEditor/Source/ChaosClothAssetDataflowNodes/Private/ChaosClothAsset/TransferSkinWeightsNode.cpp
--- a/Engine/Plugins/Experimental/ChaosClothAssetEditor/Source/ChaosClothAssetDataflowNodes/Private/ChaosClothAsset/TransferSkinWeightsNode.cpp
+++ b/Engine/Plugins/Experimental/ChaosClothAssetEditor/Source/ChaosClothAssetDataflowNodes/Private/ChaosClothAsset/TransferSkinWeightsNode.cpp
@@ -77,7 +77,8 @@ namespace UE::Chaos::ClothAsset::Private
 	/** Copy the skin weights from DynamicMesh to Collection, handling split vertices. */
 	static void CopySkinWeightsFromDynamicMeshToCloth(const FDynamicMesh3& WeldedSimMesh, 
 													  const bool bUseParallel,
-													  const TSharedRef<FManagedArrayCollection>& ClothCollection)
+													  const TSharedRef<FManagedArrayCollection>& ClothCollection,
+													  const int32 MaxNumInfluences)
 	{
 		using namespace UE::Geometry;
 
@@ -97,7 +98,7 @@ namespace UE::Chaos::ClothAsset::Private
 				SimMeshToDynamicMesh[NonManifoldMapping.GetOriginalNonManifoldVertexID(DynamicMeshVert)].Add(DynamicMeshVert);
 			}
 
-			ParallelFor(ClothFacade.GetNumSimVertices3D(), [&ClothFacade, &WeldedSimMesh, &SimMeshToDynamicMesh](int32 SimVertexID)
+			ParallelFor(ClothFacade.GetNumSimVertices3D(), [&ClothFacade, &WeldedSimMesh, &SimMeshToDynamicMesh, MaxNumInfluences](int32 SimVertexID)
 			{
 				const FDynamicMeshVertexSkinWeightsAttribute* const OutAttribute = WeldedSimMesh.Attributes()->GetSkinWeightsAttribute(FSkeletalMeshAttributes::DefaultSkinWeightProfileName);
 				checkSlow(OutAttribute);
@@ -146,9 +147,9 @@ namespace UE::Chaos::ClothAsset::Private
 						BoneWeights.Add(FloatVal);
 						WeightsSum += FloatVal;
 					}
-					if (BoneIndices.Num() > MAX_TOTAL_INFLUENCES)
+					if (BoneIndices.Num() > MaxNumInfluences)
 					{
-						// Choose MAX_TOTAL_INFLUENCES highest weighted bones.
+						// Choose MaxNumInfluences highest weighted bones.
 						TArray<TPair<float, int32>> SortableData;
 						SortableData.Reserve(BoneIndices.Num());
 						for (int32 Idx = 0; Idx < BoneIndices.Num(); ++Idx)
@@ -157,10 +158,10 @@ namespace UE::Chaos::ClothAsset::Private
 						}
 						SortableData.Sort([](const TPair<float, int32>& A, const TPair<float, int32>& B) { return A > B; });
 
-						BoneIndices.SetNum(MAX_TOTAL_INFLUENCES);
-						BoneWeights.SetNum(MAX_TOTAL_INFLUENCES);
+						BoneIndices.SetNum(MaxNumInfluences);
+						BoneWeights.SetNum(MaxNumInfluences);
 						WeightsSum = 0.f;
-						for (int32 Idx = 0; Idx < MAX_TOTAL_INFLUENCES; ++Idx)
+						for (int32 Idx = 0; Idx < MaxNumInfluences; ++Idx)
 						{
 							BoneIndices[Idx] = SortableData[Idx].Get<1>();
 							BoneWeights[Idx] = SortableData[Idx].Get<0>();
@@ -327,7 +328,7 @@ namespace UE::Chaos::ClothAsset::Private
 		//
 		// Copy the new bone weight data and inpaint mask from the welded sim mesh back to the sim cloth patterns.
 		//
-		CopySkinWeightsFromDynamicMeshToCloth(WeldedSimMesh, true, ClothCollection);
+		CopySkinWeightsFromDynamicMeshToCloth(WeldedSimMesh, true, ClothCollection, TransferBoneWeights.MaxNumInfluences);
 		CopyInpaintMapFromDynamicMeshToCloth(WeldedSimMesh, InpaintWeightMaskName, true, TransferBoneWeights.MatchedVertices, ClothCollection);
 
 		//
@@ -356,7 +357,7 @@ namespace UE::Chaos::ClothAsset::Private
 			FTransferBoneWeights SimToRenderMeshTransfer(&WeldedSimMesh, FSkeletalMeshAttributes::DefaultSkinWeightProfileName);
 			SimToRenderMeshTransfer.bUseParallel = bUseParallel;
 			SimToRenderMeshTransfer.TransferMethod = FTransferBoneWeights::ETransferBoneWeightsMethod::ClosestPointOnSurface;
-
+			SimToRenderMeshTransfer.MaxNumInfluences = TransferBoneWeights.MaxNumInfluences;
 			UE::Chaos::ClothAsset::FCollectionClothFacade ClothFacade(ClothCollection);
 			ParallelFor(ClothFacade.GetNumRenderVertices(), [&SimToRenderMeshTransfer, &ClothFacade](int32 VertexID)
 				{
@@ -478,6 +479,7 @@ void FChaosClothAssetTransferSkinWeightsNode::Evaluate(Dataflow::FContext& Conte
 			//
 			constexpr bool bUseParallel = true;
 			FTransferBoneWeights TransferBoneWeights(&SourceDynamicMesh, FSkeletalMeshAttributes::DefaultSkinWeightProfileName);
+			TransferBoneWeights.MaxNumInfluences = static_cast<int32>(FChaosClothAssetTransferSkinWeightsNode::MaxNumInfluences);
 			TransferBoneWeights.bUseParallel = bUseParallel;
 			TransferBoneWeights.TransferMethod = static_cast<FTransferBoneWeights::ETransferBoneWeightsMethod>(TransferMethod);
 
@@ -513,6 +515,29 @@ void FChaosClothAssetTransferSkinWeightsNode::Evaluate(Dataflow::FContext& Conte
 				return;
 			}
 
+			// Optional check to make sure all vertices adhere to the bone influence limit
+			constexpr bool bCheckMaxInfluenceComplience = false; 
+			if (bCheckMaxInfluenceComplience)
+			{
+				for (int32 VID = 0; VID < ClothFacade.GetNumSimVertices3D(); ++VID)
+				{
+					const int32 NumInfluences = ClothFacade.GetSimBoneIndices()[VID].Num();
+					if (!ensure(NumInfluences <= static_cast<int32>(FChaosClothAssetTransferSkinWeightsNode::MaxNumInfluences)))
+					{
+						UE_LOG(LogChaosClothAssetDataflowNodes, Warning, TEXT("TransferSkinWeightsNode: Maximum number of influences is exceeded for sim vertex %i."), VID);
+					}
+				}
+
+				for (int32 VID = 0; VID < ClothFacade.GetNumRenderVertices(); ++VID)
+				{
+					const int32 NumInfluences = ClothFacade.GetRenderBoneIndices()[VID].Num();
+					if (!ensure(NumInfluences <= static_cast<int32>(FChaosClothAssetTransferSkinWeightsNode::MaxNumInfluences)))
+					{
+						UE_LOG(LogChaosClothAssetDataflowNodes, Warning, TEXT("TransferSkinWeightsNode: Maximum number of influences is exceeded for render vertex %i."), VID);
+					}
+				}
+			}
+
 		}
 
 		SetValue(Context, MoveTemp(*ClothCollection), &Collection);
diff --git a/Engine/Plugins/Experimental/ChaosClothAssetEditor/Source/ChaosClothAssetDataflowNodes/Public/ChaosClothAsset/TransferSkinWeightsNode.h b/Engine/Plugins/Experimental/ChaosClothAssetEditor/Source/ChaosClothAssetDataflowNodes/Public/ChaosClothAsset/TransferSkinWeightsNode.h
--- a/Engine/Plugins/Experimental/ChaosClothAssetEditor/Source/ChaosClothAssetDataflowNodes/Public/ChaosClothAsset/TransferSkinWeightsNode.h
+++ b/Engine/Plugins/Experimental/ChaosClothAssetEditor/Source/ChaosClothAssetDataflowNodes/Public/ChaosClothAsset/TransferSkinWeightsNode.h
@@ -24,6 +24,16 @@ enum class EChaosClothAssetTransferSkinWeightsMethod : uint8
 	InpaintWeights
 };
 
+
+UENUM(BlueprintType)
+enum class EChaosClothAssetMaxNumInfluences : uint8
+{
+	Uninitialized = 0	UMETA(Hidden),
+	Four = 4			UMETA(DisplayName = "4"),
+	Eight = 8			UMETA(DisplayName = "8"),
+	Twelve = 12			UMETA(DisplayName = "12")
+};
+
 /** Transfer the skinning weights set on a skeletal mesh to the simulation and/or render mesh stored in the cloth collection. */
 USTRUCT(Meta = (DataflowCloth))
 struct FChaosClothAssetTransferSkinWeightsNode : public FDataflowNode
@@ -81,7 +91,11 @@ public:
 	/** The smoothing strength of each smoothing iteration. */
 	UPROPERTY(EditAnywhere, Category = "Transfer Skin Weights|Transfer Method", Meta = (UIMin = 0, UIMax = 1, ClampMin = 0, ClampMax = 1, DisplayName = "Smoothing Strength", EditCondition="TransferMethod==EChaosClothAssetTransferSkinWeightsMethod::InpaintWeights"))
 	float SmoothingStrength = 0.1;
-
+	
+	/** The maximum number of bones that will influence each vertex. */
+	UPROPERTY(EditAnywhere, Category = "Transfer Skin Weights|Transfer Method", Meta = (DisplayName = "Max Bone Influences"))
+	EChaosClothAssetMaxNumInfluences MaxNumInfluences = EChaosClothAssetMaxNumInfluences::Eight;
+	
     /** Optional mask where a non-zero value indicates that we want the skinning weights for the vertex to be computed automatically instead of it being copied over from the source mesh. */
 	UPROPERTY(EditAnywhere, Category = "Transfer Skin Weights|Transfer Method", Meta = (DisplayName = "Inpaint Mask", EditCondition="TransferMethod==EChaosClothAssetTransferSkinWeightsMethod::InpaintWeights"))
 	FChaosClothAssetWeightedValueNonAnimatableNoLowHighRange InpaintMask = { TEXT("InpaintMask") };
diff --git a/Engine/Plugins/Runtime/GeometryProcessing/Source/DynamicMesh/Private/Operations/SmoothBoneWeights.cpp b/Engine/Plugins/Runtime/GeometryProcessing/Source/DynamicMesh/Private/Operations/SmoothBoneWeights.cpp
--- a/Engine/Plugins/Runtime/GeometryProcessing/Source/DynamicMesh/Private/Operations/SmoothBoneWeights.cpp
+++ b/Engine/Plugins/Runtime/GeometryProcessing/Source/DynamicMesh/Private/Operations/SmoothBoneWeights.cpp
@@ -6,7 +6,6 @@
 #include "DynamicMesh/DynamicMeshAttributeSet.h"
 #include "DynamicMesh/NonManifoldMappingSupport.h"
 #include "Util/ProgressCancel.h"
-#include "BoneWeights.h"
 #include "Math/UnrealMathUtility.h"
 #include "Parameterization/MeshLocalParam.h"
 #include "DynamicMesh/MeshNormals.h"
@@ -237,6 +236,11 @@ EOperationValidationResult FSmoothDynamicMeshVertexSkinWeights::Validate()
 		return EOperationValidationResult::Failed_UnknownReason;
 	}
 
+	if (MaxNumInfluences <= 0)
+	{
+		return EOperationValidationResult::Failed_UnknownReason;
+	}
+
 	return TSmoothBoneWeights::Validate();
 }
 
@@ -254,6 +258,12 @@ bool FSmoothDynamicMeshVertexSkinWeights::SmoothWeightsAtVertex(const int32 Vert
 			WeightArray.SetBoneWeight(FBoneWeight(FinalWeight.Key, FinalWeight.Value), BoneSettings);
 		}
 		
+		BoneSettings.SetNormalizeType(EBoneWeightNormalizeType::Always);
+		BoneSettings.SetMaxWeightCount(MaxNumInfluences);
+
+		// make sure we do not exceed the max influence limit
+		WeightArray.Renormalize(BoneSettings);
+
 		Attribute->SetValue(VertexID, WeightArray);
 
 		return true;
diff --git a/Engine/Plugins/Runtime/GeometryProcessing/Source/DynamicMesh/Private/Operations/TransferBoneWeights.cpp b/Engine/Plugins/Runtime/GeometryProcessing/Source/DynamicMesh/Private/Operations/TransferBoneWeights.cpp
--- a/Engine/Plugins/Runtime/GeometryProcessing/Source/DynamicMesh/Private/Operations/TransferBoneWeights.cpp
+++ b/Engine/Plugins/Runtime/GeometryProcessing/Source/DynamicMesh/Private/Operations/TransferBoneWeights.cpp
@@ -10,7 +10,6 @@
 #include "Async/ParallelFor.h"
 #include "Util/ProgressCancel.h"
 #include "BoneIndices.h"
-#include "BoneWeights.h"
 #include "IndexTypes.h"
 #include "TransformTypes.h"
 #include "Solvers/Internal/QuadraticProgramming.h"
@@ -29,6 +28,7 @@ namespace TransferBoneWeightsLocals
 	 * @param TriVertices The vertices of a triangle containing the point we are interpolating the weights for
 	 * @param Bary Barycentric coordinates of the point
 	 * @param Attribute Attribute containing bone weights of the mesh that TriVertices belong to
+	 * @param MaxNumInfluences The maximum allowed number of influences per vertex stored in OutWeights
 	 * @param SourceIndexToBone Optional map from bone index to bone name for the source mesh
 	 * @param TargetBoneToIndex OPtional map from bone name to bone index for the target mesh
 	 * @param bNormalizeToOne If true, OutWeights will be normalized to sum to 1.
@@ -37,6 +37,7 @@ namespace TransferBoneWeightsLocals
 								const FIndex3i& TriVertices,
 								const FVector3f& Bary,
 								const FDynamicMeshVertexSkinWeightsAttribute* Attribute,
+								const int32 MaxNumInfluences,
 								const TArray<FName>* SourceIndexToBone = nullptr,
 								const TMap<FName, FBoneIndexType>* TargetBoneToIndex = nullptr,
 								bool bNormalizeToOne = true)
@@ -49,8 +50,13 @@ namespace TransferBoneWeightsLocals
 		FBoneWeightsSettings BlendSettings;
 		BlendSettings.SetNormalizeType(bNormalizeToOne ? EBoneWeightNormalizeType::Always : EBoneWeightNormalizeType::None);
 		BlendSettings.SetBlendZeroInfluence(true);
+		BlendSettings.SetMaxWeightCount(MaxNumInfluences);
 		OutWeights = FBoneWeights::Blend(Weight1, Weight2, Weight3, Bary[0], Bary[1], Bary[2], BlendSettings);
 
+		// TODO: Blend method can potentially skip applying renormalization and prunning weights to match MaxNumInfluences
+		// using the BlendSettings so force renormalization. Remove this once FBoneWeights::Blend is fixed.
+		OutWeights.Renormalize(BlendSettings);
+		 
 		// Check if we need to remap the indices
 		if (SourceIndexToBone && TargetBoneToIndex) 
 		{
@@ -492,6 +498,9 @@ bool FTransferBoneWeights::TransferWeightsToMesh(FDynamicMesh3& InOutTargetMesh,
 			FBoneWeightsSettings BoneSettings;
 			BoneSettings.SetNormalizeType(EBoneWeightNormalizeType::None);
 			
+			FBoneWeightsSettings RenormalizeBoneSettings;
+			RenormalizeBoneSettings.SetMaxWeightCount(MaxNumInfluences);
+
 			// Iterate over every column containing all bone weights for the vertex
 			for (int32 ColIdx = 0; ColIdx < TargetWeightsTransposed.outerSize(); ++ColIdx)
 			{
@@ -506,7 +515,7 @@ bool FTransferBoneWeights::TransferWeightsToMesh(FDynamicMesh3& InOutTargetMesh,
 					WeightArray.SetBoneWeight(Bweight, BoneSettings);
 				}
 
-				WeightArray.Renormalize(FBoneWeightsSettings());
+				WeightArray.Renormalize(RenormalizeBoneSettings);
 
 				const int32 VertexIDLinearalized = bVariablesOnly ? static_cast<int32>(VaribleRows[ColIdx]) : ColIdx; // linearized vertex ID (matrix row) of the variable in the Energy matrix
 				const int32 VertexID = static_cast<int32>(ToMeshV[VertexIDLinearalized]);
@@ -528,6 +537,7 @@ bool FTransferBoneWeights::TransferWeightsToMesh(FDynamicMesh3& InOutTargetMesh,
 				}
 
 				FSmoothDynamicMeshVertexSkinWeights SmoothWeights(&InOutTargetMesh, InTargetProfileName);
+				SmoothWeights.MaxNumInfluences = MaxNumInfluences;
 				if (ensure(SmoothWeights.Validate() == EOperationValidationResult::Ok))
 				{
 					ensure(SmoothWeights.SmoothWeightsAtVerticesWithinDistance(VerticesToSmooth, SmoothingStrength, SearchRadius, NumSmoothingIterations));
@@ -575,7 +585,7 @@ bool FTransferBoneWeights::TransferWeightsToPoint(UE::AnimationCore::FBoneWeight
 	if (SearchRadius < 0 && NormalThreshold < 0)
 	{
 		// If the radius and normals are ignored, simply interpolate the weights and return the result
-		TransferBoneWeightsLocals::InterpolateBoneWeights(OutWeights, TriVertex, BaryF, SourceSkinWeights, SourceBoneNames, TargetBoneToIndex);
+		TransferBoneWeightsLocals::InterpolateBoneWeights(OutWeights, TriVertex, BaryF, SourceSkinWeights, MaxNumInfluences, SourceBoneNames, TargetBoneToIndex);
 	}
 	else
 	{
@@ -614,7 +624,7 @@ bool FTransferBoneWeights::TransferWeightsToPoint(UE::AnimationCore::FBoneWeight
 		
 		if (bPassedRadiusCheck && bPassedNormalsCheck)
 		{
-			TransferBoneWeightsLocals::InterpolateBoneWeights(OutWeights, TriVertex, BaryF, SourceSkinWeights, SourceBoneNames, TargetBoneToIndex);
+			TransferBoneWeightsLocals::InterpolateBoneWeights(OutWeights, TriVertex, BaryF, SourceSkinWeights, MaxNumInfluences, SourceBoneNames, TargetBoneToIndex);
 		}
 		else
 		{
diff --git a/Engine/Plugins/Runtime/GeometryProcessing/Source/DynamicMesh/Public/Operations/SmoothBoneWeights.h b/Engine/Plugins/Runtime/GeometryProcessing/Source/DynamicMesh/Public/Operations/SmoothBoneWeights.h
--- a/Engine/Plugins/Runtime/GeometryProcessing/Source/DynamicMesh/Public/Operations/SmoothBoneWeights.h
+++ b/Engine/Plugins/Runtime/GeometryProcessing/Source/DynamicMesh/Public/Operations/SmoothBoneWeights.h
@@ -6,6 +6,7 @@
 #include "UObject/NameTypes.h"
 #include "Containers/Map.h"
 #include "BoneIndices.h"
+#include "BoneWeights.h" //  UE::AnimationCore::MaxInlineBoneWeightCount;
 
 // Forward declarations
 class FProgressCancel;
@@ -170,6 +171,13 @@ public:
 
 	using TSmoothBoneWeights<FBoneIndexType, float>::SmoothWeightsAtVertex; 
 	
+	
+	/** 
+	 * During the smoothing, only use this number of influences per vertex. Prune the excess with the smallest influences 
+	 * and re-normalize. 
+	 */
+	int32 MaxNumInfluences = UE::AnimationCore::MaxInlineBoneWeightCount;
+
 	/**
 	 * @param InSourceMesh The mesh whose vertices contain bone weights data.
 	 * @param InProfile    The name of the skin weights profile containing the bone weights data.
diff --git a/Engine/Plugins/Runtime/GeometryProcessing/Source/DynamicMesh/Public/Operations/TransferBoneWeights.h b/Engine/Plugins/Runtime/GeometryProcessing/Source/DynamicMesh/Public/Operations/TransferBoneWeights.h
--- a/Engine/Plugins/Runtime/GeometryProcessing/Source/DynamicMesh/Public/Operations/TransferBoneWeights.h
+++ b/Engine/Plugins/Runtime/GeometryProcessing/Source/DynamicMesh/Public/Operations/TransferBoneWeights.h
@@ -5,11 +5,11 @@
 #include "GeometryTypes.h"
 #include "UObject/NameTypes.h"
 #include "TransformTypes.h"
+#include "BoneWeights.h" // UE::AnimationCore::MaxInlineBoneWeightCount
 
 // Forward declarations
 class FProgressCancel;
 
-namespace UE::AnimationCore { class FBoneWeights; }
 
 namespace UE::Geometry
 {
@@ -152,6 +152,12 @@ public:
 	/** If true, will use the intrinsic Delaunay mesh to construct sparse Cotangent Laplacian matrix. */
 	bool bUseIntrinsicLaplacian = false;
 
+	/** 
+	 * During the transfer, only use this number of influences per vertex. Prune the excess with the smallest influences 
+	 * and re-normalize. 
+	 */
+	int32 MaxNumInfluences = UE::AnimationCore::MaxInlineBoneWeightCount;
+
 	/** 
 	 * Optional mask where if ForceInpaint[VertexID] != 0 we want to force the weights for the vertex to be computed  
 	 * automatically.
