Chaos Cloth Asset: update ClothAssetImportNode to import properties and selections.

diff --git a/Engine/Plugins/Experimental/ChaosClothAssetEditor/Source/ChaosClothAssetDataflowNodes/Private/ChaosClothAsset/ImportNode.cpp b/Engine/Plugins/Experimental/ChaosClothAssetEditor/Source/ChaosClothAssetDataflowNodes/Private/ChaosClothAsset/ImportNode.cpp
--- a/Engine/Plugins/Experimental/ChaosClothAssetEditor/Source/ChaosClothAssetDataflowNodes/Private/ChaosClothAsset/ImportNode.cpp
+++ b/Engine/Plugins/Experimental/ChaosClothAssetEditor/Source/ChaosClothAssetDataflowNodes/Private/ChaosClothAsset/ImportNode.cpp
@@ -4,6 +4,8 @@
 #include "ChaosClothAsset/ClothAsset.h"
 #include "ChaosClothAsset/ClothDataflowTools.h"
 #include "ChaosClothAsset/CollectionClothFacade.h"
+#include "ChaosClothAsset/CollectionClothSelectionFacade.h"
+#include "Chaos/CollectionPropertyFacade.h"
 #include "Dataflow/DataflowInputOutput.h"
 #include "Dataflow/DataflowObjectInterface.h"
 #include "Engine/SkinnedAssetCommon.h"
@@ -23,6 +25,7 @@ FChaosClothAssetImportNode::FChaosClothAssetImportNode(const Dataflow::FNodePara
 void FChaosClothAssetImportNode::Evaluate(Dataflow::FContext& Context, const FDataflowOutput* Out) const
 {
 	using namespace UE::Chaos::ClothAsset;
+	using namespace Chaos::Softs;
 
 	if (Out->IsA<FManagedArrayCollection>(&Collection))
 	{
@@ -66,8 +69,29 @@ void FChaosClothAssetImportNode::Evaluate(Dataflow::FContext& Context, const FDa
 			const TArray<TSharedRef<const FManagedArrayCollection>>& InClothCollections = ClothAsset->GetClothCollections();
 			if (ImportLod >= 0 && InClothCollections.Num() > ImportLod)
 			{
+				// Copy cloth
 				const FCollectionClothConstFacade InClothFacade(InClothCollections[ImportLod]);
 				ClothFacade.Initialize(InClothFacade);
+
+				// Copy properties
+				const FCollectionPropertyConstFacade InPropertyFacade(InClothCollections[ImportLod]);
+				if (InPropertyFacade.IsValid())
+				{
+					FCollectionPropertyMutableFacade PropertyFacade(ClothCollection);
+					PropertyFacade.DefineSchema();
+					constexpr bool bUpdateExistingPropertiesFalse = false; // There are no existing properties.
+					PropertyFacade.Append(InClothCollections[ImportLod].ToSharedPtr(), bUpdateExistingPropertiesFalse);
+				}
+
+				// Copy selections
+				const FCollectionClothSelectionConstFacade InSelectionFacade(InClothCollections[ImportLod]);
+				if (InSelectionFacade.IsValid())
+				{
+					FCollectionClothSelectionFacade SelectionFacade(ClothCollection);
+					SelectionFacade.DefineSchema();
+					constexpr bool bOverwriteExistingIfMismatchedTrue = true;
+					SelectionFacade.Append(InSelectionFacade, bOverwriteExistingIfMismatchedTrue);
+				}
 			}
 		}
 		SetValue(Context, MoveTemp(*ClothCollection), &Collection);
