Cloth Editor: adding two ways to hide/show triangles in the paint tool:
- A new HideTriangles brush mode that will hide triangles that are brushed
- A new tool property that is a list of Patterns and whether each one is visible or not
Limitations:
- Hide/show information does not come along when switching between 2D/3D modes
- Checkbox interface for selecting which patterns to show is clunky
- Manually hiding triangles via brush then unchecking all patterns will clear the manually hidden triangles

Also adding an optional white-red mode to visualize the weight map. Painting an all black object in 3D with a dark grey background is pretty annoying.

diff --git a/Engine/Plugins/Experimental/ChaosClothAssetEditor/Source/ChaosClothAssetEditorTools/Private/ChaosClothAsset/ClothWeightMapPaintTool.cpp b/Engine/Plugins/Experimental/ChaosClothAssetEditor/Source/ChaosClothAssetEditorTools/Private/ChaosClothAsset/ClothWeightMapPaintTool.cpp
--- a/Engine/Plugins/Experimental/ChaosClothAssetEditor/Source/ChaosClothAssetEditorTools/Private/ChaosClothAsset/ClothWeightMapPaintTool.cpp
+++ b/Engine/Plugins/Experimental/ChaosClothAssetEditor/Source/ChaosClothAssetEditorTools/Private/ChaosClothAsset/ClothWeightMapPaintTool.cpp
@@ -46,6 +46,7 @@
 #include "Framework/Notifications/NotificationManager.h"
 #include "Selection/PolygonSelectionMechanic.h"
 #include "BaseGizmos/BrushStampIndicator.h"
+#include "DynamicSubmesh3.h"
 
 #include UE_INLINE_GENERATED_CPP_BY_NAME(ClothWeightMapPaintTool)
 
@@ -112,6 +113,16 @@ void UClothEditorUpdateWeightMapProperties::PostEditChangeProperty(FPropertyChan
 }
 
 
+// Show/Hide properties
+void UClothEditorMeshWeightMapPaintToolShowHideProperties::PostAction(EClothEditorWeightMapPaintToolActions Action)
+{
+	if (ParentTool.IsValid())
+	{
+		ParentTool->RequestAction(Action);
+	}
+}
+
+
 /*
  * Tool
  */
@@ -217,6 +228,12 @@ void UClothEditorWeightMapPaintTool::Setup()
 
 	// initialize other properties
 	FilterProperties = NewObject<UClothEditorWeightMapPaintBrushFilterProperties>(this);
+	FilterProperties->WatchProperty(FilterProperties->ColorMap,
+		[this](EClothEditorWeightMapDisplayType NewType) 
+		{ 
+			UpdateVertexColorOverlay(); 
+			DynamicMeshComponent->FastNotifyVertexAttributesUpdated(EMeshRenderAttributeFlags::VertexColors);
+		});
 	FilterProperties->WatchProperty(FilterProperties->SubToolType,
 		[this](EClothEditorWeightMapPaintInteractionType NewType) { UpdateSubToolType(NewType); });
 	FilterProperties->WatchProperty(FilterProperties->BrushSize,
@@ -281,10 +298,76 @@ void UClothEditorWeightMapPaintTool::Setup()
 		MeshElementsDisplay->Settings->RestoreProperties(this, TEXT("ClothEditorWeightMapPaintTool2"));
 		AddToolPropertySource(MeshElementsDisplay->Settings);
 	}
-	MeshElementsDisplay->SetMeshAccessFunction([this](UMeshElementsVisualizer::ProcessDynamicMeshFunc ProcessFunc) {
-		ProcessFunc(*GetSculptMesh());
+	MeshElementsDisplay->SetMeshAccessFunction([this](UMeshElementsVisualizer::ProcessDynamicMeshFunc ProcessFunc) 
+	{
+		if (HiddenTriangles.Num() > 0 || PendingHiddenTriangles.Num() > 0)
+		{
+			const UE::Geometry::FDynamicMesh3* const FullMesh = GetSculptMesh();
+
+			TArray<int> NonHiddenTriangles;
+			for (const int32 TID : FullMesh->TriangleIndicesItr())
+			{
+				if (!HiddenTriangles.Contains(TID) && !PendingHiddenTriangles.Contains(TID))
+				{
+					NonHiddenTriangles.Add(TID);
+				}
+			}
+
+			UE::Geometry::FDynamicSubmesh3 Submesh(FullMesh, NonHiddenTriangles);
+			ProcessFunc(Submesh.GetSubmesh());
+		}
+		else
+		{
+			ProcessFunc(*GetSculptMesh());
+		}
 	});
 
+	ShowHideProperties = NewObject<UClothEditorMeshWeightMapPaintToolShowHideProperties>();
+	ShowHideProperties->Initialize(this);
+	ShowHideProperties->WatchProperty(ShowHideProperties->ShowPatterns,
+		[this](const TMap<int32, bool>& NewMap)
+		{
+			bool bAnySelected = false;
+			for (const TPair<int32, bool>& KeyValue : NewMap)
+			{
+				if (KeyValue.Value)
+				{
+					bAnySelected = true;
+					break;
+				}
+			}
+
+			HiddenTriangles.Reset();
+			if (bAnySelected)
+			{
+				for (const TPair<int32, bool>& KeyValue : NewMap)
+				{
+					if (KeyValue.Value == false)
+					{
+						const int32 PatternIndex = KeyValue.Key;
+						if (PatternIndex < PatternTriangleOffsetAndNum.Num())
+						{
+							const TPair<int32, int32> StartAndNum = PatternTriangleOffsetAndNum[PatternIndex];
+							for (int32 TID = StartAndNum.Key; TID < StartAndNum.Key + StartAndNum.Value; ++TID)
+							{
+								HiddenTriangles.Add(TID);
+							}
+						}
+					}
+				}
+			}
+
+			MeshElementsDisplay->NotifyMeshChanged();
+			DynamicMeshComponent->FastNotifySecondaryTrianglesChanged();
+		},
+		[this](const TMap<int32, bool>& A, const TMap<int32, bool>& B)		// Not-equal function for TMap
+		{
+			return (!A.OrderIndependentCompareEqual(B));
+		}
+	);
+	AddToolPropertySource(ShowHideProperties);
+
+
 	// disable view properties
 	SetViewPropertiesEnabled(false);
 	UpdateMaterialMode(EMeshEditingMaterialModes::VertexColor);
@@ -331,6 +414,43 @@ void UClothEditorWeightMapPaintTool::Setup()
 				DynamicMeshToWeight = Cloth.GetSimVertex3DLookup();
 				WeightToDynamicMesh = Cloth.GetSimVertex2DLookup();
 			}
+
+			const bool bIsRenderMode = ClothEditorContextObject->GetConstructionViewMode() == EClothPatternVertexType::Render;
+			const int32 NumPatterns = bIsRenderMode ? Cloth.GetNumRenderPatterns() : Cloth.GetNumSimPatterns();		
+			
+			PatternTriangleOffsetAndNum.SetNum(NumPatterns);
+
+			TSet<int32> NonEmptyPatternIDs;
+
+			for (int32 PatternIndex = 0; PatternIndex < NumPatterns; ++PatternIndex)
+			{
+				TPair<int32, int32>& OffsetAndNum = PatternTriangleOffsetAndNum[PatternIndex];
+				if (bIsRenderMode)
+				{
+					FCollectionClothRenderPatternConstFacade RenderPattern = Cloth.GetRenderPattern(PatternIndex);
+					OffsetAndNum.Key = RenderPattern.GetRenderFacesOffset();
+					OffsetAndNum.Value = RenderPattern.GetNumRenderFaces();
+				}
+				else
+				{
+					FCollectionClothSimPatternConstFacade SimPattern = Cloth.GetSimPattern(PatternIndex);
+					OffsetAndNum.Key = SimPattern.GetSimFacesOffset();
+					OffsetAndNum.Value = SimPattern.GetNumSimFaces();
+				}
+
+				if (OffsetAndNum.Value > 0)
+				{
+					NonEmptyPatternIDs.Add(PatternIndex);
+				}
+			}
+
+			// Initialize the ShowPatterns map from found pattern indices
+			ShowHideProperties->ShowPatterns.Reset();
+			for (const int32 PatternID : NonEmptyPatternIDs)
+			{
+				ShowHideProperties->ShowPatterns.Add({ PatternID, false });
+			}
+		
 		}
 	}
 
@@ -374,6 +494,13 @@ void UClothEditorWeightMapPaintTool::Setup()
 	UpdateWeightMapProperties->Name = WeightMapNodeToUpdate->Name;
 	SetToolPropertySourceEnabled(UpdateWeightMapProperties, true);
 
+	DynamicMeshComponent->EnableSecondaryTriangleBuffers(
+		[this](const FDynamicMesh3* Mesh, int32 TriangleID)
+		{
+			return PendingHiddenTriangles.Contains(TriangleID) || HiddenTriangles.Contains(TriangleID);
+		});
+	DynamicMeshComponent->SetSecondaryBuffersVisibility(false);
+
 
 	// update colors
 	UpdateVertexColorOverlay();
@@ -492,7 +619,8 @@ void UClothEditorWeightMapPaintTool::DecreaseBrushRadiusSmallStepAction()
 bool UClothEditorWeightMapPaintTool::IsInBrushSubMode() const
 {
 	return FilterProperties->SubToolType == EClothEditorWeightMapPaintInteractionType::Brush
-		|| FilterProperties->SubToolType == EClothEditorWeightMapPaintInteractionType::Fill;
+		|| FilterProperties->SubToolType == EClothEditorWeightMapPaintInteractionType::Fill
+	    || FilterProperties->SubToolType == EClothEditorWeightMapPaintInteractionType::HideTriangles;
 }
 
 
@@ -551,6 +679,14 @@ void UClothEditorWeightMapPaintTool::OnEndStroke()
 
 	GetActiveBrushOp()->EndStroke(GetSculptMesh(), LastStamp, VertexROI);
 
+	if (PendingHiddenTriangles.Num() > 0)
+	{
+		HiddenTriangles.Append(PendingHiddenTriangles);
+		PendingHiddenTriangles.Reset();
+		MeshElementsDisplay->NotifyMeshChanged();
+		DynamicMeshComponent->FastNotifySecondaryTrianglesChanged();
+	}
+
 	UpdateVertexColorOverlay(&TriangleROI);
 	DynamicMeshComponent->FastNotifyVertexAttributesUpdated(EMeshRenderAttributeFlags::VertexColors);
 
@@ -746,10 +882,29 @@ bool UClothEditorWeightMapPaintTool::ApplyStamp()
 		WeightBrushOp->bApplyRadiusLimit = false;
 	}
 
-	FDynamicMesh3* Mesh = GetSculptMesh();
-	WeightBrushOp->ApplyStampByVertices(Mesh, CurrentStamp, VertexROI, ROIWeightValueBuffer);
+	bool bUpdated = false;
+	if (FilterProperties->SubToolType == EClothEditorWeightMapPaintInteractionType::Brush || FilterProperties->SubToolType == EClothEditorWeightMapPaintInteractionType::Fill)
+	{
+		FDynamicMesh3* Mesh = GetSculptMesh();
+		WeightBrushOp->ApplyStampByVertices(Mesh, CurrentStamp, VertexROI, ROIWeightValueBuffer);
+		bUpdated = SyncMeshWithWeightBuffer(Mesh);
+	}
+	else
+	{
+		bool bAnyModified = false;
+		for (int32 TID : TriangleROI)
+		{
+			bool bModified = false;
+			PendingHiddenTriangles.Add(TID, &bModified);
+			bAnyModified = bAnyModified || bModified;
+		}
 
-	bool bUpdated = SyncMeshWithWeightBuffer(Mesh);
+		if (bAnyModified)
+		{
+			DynamicMeshComponent->FastNotifySecondaryTrianglesChanged();
+		}
+	}
+	
 
 	LastStamp = CurrentStamp;
 	LastStamp.TimeStamp = FDateTime::Now();
@@ -1284,9 +1439,18 @@ int32 UClothEditorWeightMapPaintTool::FindHitSculptMeshTriangle(const FRay3d& Lo
 		return IndexConstants::InvalidID;
 	}
 
-	if (GetBrushCanHitBackFaces())
+	if (FilterProperties->bHitBackFaces)
 	{
-		return Octree.FindNearestHitObject(LocalRay);
+		const int32 HitTID = Octree.FindNearestHitObject(LocalRay,
+			[this](int TriangleID)
+			{ 
+				if (HiddenTriangles.Contains(TriangleID))
+				{
+					return false;
+				}
+				return true;
+			});
+		return HitTID;
 	}
 	else
 	{
@@ -1295,13 +1459,18 @@ int32 UClothEditorWeightMapPaintTool::FindHitSculptMeshTriangle(const FRay3d& Lo
 		FViewCameraState StateOut;
 		GetToolManager()->GetContextQueriesAPI()->GetCurrentViewState(StateOut);
 		FVector3d LocalEyePosition(CurTargetTransform.InverseTransformPosition((FVector3d)StateOut.Position));
-		int HitTID = Octree.FindNearestHitObject(LocalRay,
-			[this, Mesh, &LocalEyePosition](int TriangleID) {
-			FVector3d Normal, Centroid;
-			double Area;
-			Mesh->GetTriInfo(TriangleID, Normal, Area, Centroid);
-			return Normal.Dot((Centroid - LocalEyePosition)) < 0;
-		});
+		const int32 HitTID = Octree.FindNearestHitObject(LocalRay,
+			[this, Mesh, &LocalEyePosition](int TriangleID) 
+			{
+				if (HiddenTriangles.Contains(TriangleID))
+				{
+					return false;
+				}
+				FVector3d Normal, Centroid;
+				double Area;
+				Mesh->GetTriInfo(TriangleID, Normal, Area, Centroid);
+				return Normal.Dot((Centroid - LocalEyePosition)) < 0;
+			});
 		return HitTID;
 	}
 }
@@ -1686,6 +1855,23 @@ void UClothEditorWeightMapPaintTool::InvertWeightsAction()
 	EndChange();
 }
 
+void UClothEditorWeightMapPaintTool::ClearHiddenAction()
+{
+	HiddenTriangles.Reset();
+
+	for (TPair<int32, bool>& PatternSelection : ShowHideProperties->ShowPatterns)
+	{
+		PatternSelection.Value = false;
+	}
+
+	UpdateVertexColorOverlay();
+	DynamicMeshComponent->FastNotifyVertexAttributesUpdated(EMeshRenderAttributeFlags::VertexColors);
+
+	MeshElementsDisplay->NotifyMeshChanged();
+	DynamicMeshComponent->FastNotifySecondaryTrianglesChanged();
+	GetToolManager()->PostInvalidation();
+}
+
 void UClothEditorWeightMapPaintTool::UpdateSelectedNode()
 {
 	check(ActiveWeightMap);
@@ -1985,6 +2171,9 @@ void UClothEditorWeightMapPaintTool::ApplyAction(EClothEditorWeightMapPaintToolA
 		InvertWeightsAction();
 		break;
 
+	case EClothEditorWeightMapPaintToolActions::ClearHiddenTriangles:
+		ClearHiddenAction();
+		break;
 	}
 }
 
@@ -2007,8 +2196,18 @@ void UClothEditorWeightMapPaintTool::UpdateVertexColorOverlay(const TSet<int>* T
 		{
 			float VertexWeight;
 			ActiveWeightMap->GetValue(Tri[TriVertIndex], &VertexWeight);
+			VertexWeight = FMath::Clamp(VertexWeight, 0.0f, 1.0f);
+
+			FVector4f NewColor;
+			if (FilterProperties->ColorMap == EClothEditorWeightMapDisplayType::BlackAndWhite)
+			{
+				NewColor = FVector4f(VertexWeight, VertexWeight, VertexWeight, 1.0f);
+			}
+			else
+			{
+				NewColor = VertexWeight * FVector4f(0.9f, 0.05f, 0.05f, 1.0f) + (1.0f - VertexWeight) * FVector4f(0.65f, 0.65f, 0.65f, 1.0f);
+			}
 
-			const FVector4f NewColor = FVector4f(VertexWeight, VertexWeight, VertexWeight, 1.0f);
 			ColorOverlay->SetElement(ColorElementTri[TriVertIndex], NewColor);
 		}
 	};
diff --git a/Engine/Plugins/Experimental/ChaosClothAssetEditor/Source/ChaosClothAssetEditorTools/Public/ChaosClothAsset/ClothWeightMapPaintTool.h b/Engine/Plugins/Experimental/ChaosClothAssetEditor/Source/ChaosClothAssetEditorTools/Public/ChaosClothAsset/ClothWeightMapPaintTool.h
--- a/Engine/Plugins/Experimental/ChaosClothAssetEditor/Source/ChaosClothAssetEditorTools/Public/ChaosClothAsset/ClothWeightMapPaintTool.h
+++ b/Engine/Plugins/Experimental/ChaosClothAssetEditor/Source/ChaosClothAssetEditorTools/Public/ChaosClothAsset/ClothWeightMapPaintTool.h
@@ -64,8 +64,13 @@ private:
 	virtual bool CanSetConstructionViewWireframeActive() const { return false; }
 };
 
-
-
+UENUM()
+enum class EClothEditorWeightMapDisplayType : uint8
+{
+	BlackAndWhite,
+	WhiteAndRed,
+	LastValue UMETA(Hidden)
+};
 
 /** Mesh Sculpting Brush Types */
 UENUM()
@@ -75,7 +80,7 @@ enum class EClothEditorWeightMapPaintInteractionType : uint8
 	Fill,
 	PolyLasso,
 	Gradient,
-
+	HideTriangles,
 	LastValue UMETA(Hidden)
 };
 
@@ -128,7 +133,10 @@ class CHAOSCLOTHASSETEDITORTOOLS_API UClothEditorWeightMapPaintBrushFilterProper
 	GENERATED_BODY()
 
 public:
-	
+
+	UPROPERTY(EditAnywhere, Category = Display)
+	EClothEditorWeightMapDisplayType ColorMap = EClothEditorWeightMapDisplayType::BlackAndWhite;
+
 	UPROPERTY(EditAnywhere, Category = ActionType, meta = (DisplayName = "Action"))
 	EClothEditorWeightMapPaintInteractionType SubToolType = EClothEditorWeightMapPaintInteractionType::Brush;
 
@@ -141,6 +149,10 @@ public:
 		HideEditConditionToggle, EditConditionHides, EditCondition = "SubToolType == EClothEditorWeightMapPaintInteractionType::Brush"))
 	float BrushSize = 0.25f;
 
+	/** Allow the Brush to hit the back-side of the mesh */
+	UPROPERTY(EditAnywhere, Category = Brush, meta = (HideEditConditionToggle, EditConditionHides, EditCondition = "SubToolType == EClothEditorWeightMapPaintInteractionType::Brush"))
+	bool bHitBackFaces = true;
+
 	/** The new value to paint on the mesh */
 	UPROPERTY(EditAnywhere, Category = Brush, meta = (UIMin = 0, ClampMin = 0, UIMax = 1, ClampMax = 1,
 		HideEditConditionToggle, EditConditionHides, EditCondition = 
@@ -203,7 +215,8 @@ enum class EClothEditorWeightMapPaintToolActions
 
 	FloodFillCurrent,
 	ClearAll,
-	Invert
+	Invert,
+	ClearHiddenTriangles
 };
 
 
@@ -255,6 +268,31 @@ private:
 	virtual void PostEditChangeProperty(struct FPropertyChangedEvent& PropertyChangedEvent) override;
 };
 
+/// Property panel for controlling hiding triangles
+UCLASS()
+class CHAOSCLOTHASSETEDITORTOOLS_API UClothEditorMeshWeightMapPaintToolShowHideProperties : public UInteractiveToolPropertySet
+{
+	GENERATED_BODY()
+
+public:
+
+	TWeakObjectPtr<UClothEditorWeightMapPaintTool> ParentTool;
+	void Initialize(UClothEditorWeightMapPaintTool* ParentToolIn) { ParentTool = ParentToolIn; }
+	void PostAction(EClothEditorWeightMapPaintToolActions Action);
+
+	/** Toggles whether each pattern is shown or hidden */
+	UPROPERTY(EditAnywhere, EditFixedSize, Transient, Category = TriangleVisibility, meta = (DisplayPriority = 2));
+	TMap<int32, bool> ShowPatterns;
+
+	/** Unhide all triangles */
+	UFUNCTION(CallInEditor, Category = TriangleVisibility, meta = (DisplayPriority = 1))
+	void ShowAll()
+	{
+		PostAction(EClothEditorWeightMapPaintToolActions::ClearHiddenTriangles);
+	}
+};
+
+
 /**
  * Mesh Element Paint Tool Class
  */
@@ -311,6 +349,7 @@ private:
 	void FloodFillCurrentWeightAction();
 	void ClearAllWeightsAction();
 	void InvertWeightsAction();
+	void ClearHiddenAction();
 
 public:
 	void SetVerticesToWeightMap(const TSet<int32>& Vertices, double WeightValue, bool bIsErase);
@@ -393,6 +432,28 @@ protected:
 	void ComputeGradient();
 	void OnSelectionModified();
 
+
+	//
+	// Show/Hide support
+	//
+protected:
+	UPROPERTY()
+	TObjectPtr<UClothEditorMeshWeightMapPaintToolShowHideProperties> ShowHideProperties;
+
+	// Hidden triangles are not rendered, their wireframes are not rendered, and they don't block ray casts from the mouse
+	TSet<int32> HiddenTriangles;
+
+	// Pending hidden triangles are triangles that are selected by the current HideTriangles brush stroke. When the stroke finishes, these pending triangles are
+	// added to the set of hidden triangles, and this set is cleared.
+	// 
+	// Pending hidden triangles are not rendered, but their wireframes *are* rendered, and they *do* block ray casts.
+	// - Wireframe update is delayed because its a relatively expensive operation, so we do it only when the stroke finishes
+	// - Ray cast blocking is delayed because otherwise one small mouse drag would drill through multiple mesh layers
+	TSet<int32> PendingHiddenTriangles;
+
+	// Triangle index offset and number of triangles for each non-empty pattern. This is used when we want to hide/show entire patterns.
+	TArray<TPair<int32, int32>> PatternTriangleOffsetAndNum;
+
 	//
 	// Internals
 	//
@@ -492,6 +553,7 @@ protected:
 	void UpdateSelectedNode();
 
 	void UpdateVertexColorOverlay(const TSet<int>* TrianglesToUpdate = nullptr);
+
 };
 
 
