Cloth Editor paint tool: allow brush to get super small. Also fix brush circle not updating after manually setting the brush size property

diff --git a/Engine/Plugins/Experimental/ChaosClothAssetEditor/Source/ChaosClothAssetEditorTools/Private/ChaosClothAsset/ClothWeightMapPaintTool.cpp b/Engine/Plugins/Experimental/ChaosClothAssetEditor/Source/ChaosClothAssetEditorTools/Private/ChaosClothAsset/ClothWeightMapPaintTool.cpp
--- a/Engine/Plugins/Experimental/ChaosClothAssetEditor/Source/ChaosClothAssetEditorTools/Private/ChaosClothAsset/ClothWeightMapPaintTool.cpp
+++ b/Engine/Plugins/Experimental/ChaosClothAssetEditor/Source/ChaosClothAssetEditorTools/Private/ChaosClothAsset/ClothWeightMapPaintTool.cpp
@@ -187,7 +187,7 @@ void UClothEditorWeightMapPaintTool::Setup()
 	DynamicMeshComponent->SetExternalDecomposition(MoveTemp(Decomp));
 
 	// initialize brush radius range interval, brush properties
-	UMeshSculptToolBase::InitializeBrushSizeRange(Bounds);
+	InitializeBrushSizeRange(Bounds);
 
 	// Set up control points mechanic
 	PolyLassoMechanic = NewObject<UPolyLassoMarqueeMechanic>(this);
@@ -236,8 +236,11 @@ void UClothEditorWeightMapPaintTool::Setup()
 		});
 	FilterProperties->WatchProperty(FilterProperties->SubToolType,
 		[this](EClothEditorWeightMapPaintInteractionType NewType) { UpdateSubToolType(NewType); });
-	FilterProperties->WatchProperty(FilterProperties->BrushSize,
-		[this](float NewSize) { UMeshSculptToolBase::BrushProperties->BrushSize.AdaptiveSize = NewSize; });
+	FilterProperties->WatchProperty(FilterProperties->BrushSize, [this](float NewSize) 
+		{ 
+			UMeshSculptToolBase::BrushProperties->BrushSize.AdaptiveSize = NewSize; 
+			CalculateBrushRadius();
+		});
 	FilterProperties->BrushSize = UMeshSculptToolBase::BrushProperties->BrushSize.AdaptiveSize;
 	FilterProperties->RestoreProperties(this);
 	AddToolPropertySource(FilterProperties);
@@ -509,6 +512,12 @@ void UClothEditorWeightMapPaintTool::Setup()
 
 }
 
+void UClothEditorWeightMapPaintTool::InitializeBrushSizeRange(const UE::Geometry::FAxisAlignedBox3d& TargetBounds)
+{
+	BrushRelativeSizeRange = FInterval1d(5e-5, 50.0);	// brush diameter ranges from micrometer to meter scale
+	BrushProperties->BrushSize.InitializeWorldSizeRange(TInterval<float>((float)BrushRelativeSizeRange.Min, (float)BrushRelativeSizeRange.Max));
+	CalculateBrushRadius();
+}
 
 void UClothEditorWeightMapPaintTool::SetClothEditorContextObject(TObjectPtr<UClothEditorContextObject> InClothEditorContextObject)
 {
diff --git a/Engine/Plugins/Experimental/ChaosClothAssetEditor/Source/ChaosClothAssetEditorTools/Public/ChaosClothAsset/ClothWeightMapPaintTool.h b/Engine/Plugins/Experimental/ChaosClothAssetEditor/Source/ChaosClothAssetEditorTools/Public/ChaosClothAsset/ClothWeightMapPaintTool.h
--- a/Engine/Plugins/Experimental/ChaosClothAssetEditor/Source/ChaosClothAssetEditorTools/Public/ChaosClothAsset/ClothWeightMapPaintTool.h
+++ b/Engine/Plugins/Experimental/ChaosClothAssetEditor/Source/ChaosClothAssetEditorTools/Public/ChaosClothAsset/ClothWeightMapPaintTool.h
@@ -379,6 +379,9 @@ protected:
 	virtual TUniquePtr<FMeshSculptBrushOp>& GetActiveBrushOp();
 
 	virtual bool SharesBrushPropertiesChanges() const override { return false; }
+
+	virtual void InitializeBrushSizeRange(const UE::Geometry::FAxisAlignedBox3d& TargetBounds) override;
+
 	// end UMeshSculptToolBase API
 
 
diff --git a/Engine/Plugins/Experimental/MeshModelingToolsetExp/Source/MeshModelingToolsExp/Public/Sculpting/MeshSculptToolBase.h b/Engine/Plugins/Experimental/MeshModelingToolsetExp/Source/MeshModelingToolsExp/Public/Sculpting/MeshSculptToolBase.h
--- a/Engine/Plugins/Experimental/MeshModelingToolsetExp/Source/MeshModelingToolsExp/Public/Sculpting/MeshSculptToolBase.h
+++ b/Engine/Plugins/Experimental/MeshModelingToolsetExp/Source/MeshModelingToolsExp/Public/Sculpting/MeshSculptToolBase.h
@@ -36,7 +36,7 @@ enum class EBrushToolSizeType : uint8
  * The brush size can be defined in various ways.
  */
 USTRUCT()
-struct FBrushToolRadius
+struct MESHMODELINGTOOLSEXP_API FBrushToolRadius
 {
 	GENERATED_BODY()
 
