Chaos Cloth Asset:
- Fixed missing Cloth Asset package dirty flag whenever the asset is re-updated.
- Added a refresh button to the Terminal Node UI to force regenerating the asset when the code but not the data has changed.

diff --git a/Engine/Plugins/Experimental/ChaosClothAssetEditor/Source/ChaosClothAssetDataflowNodes/Private/ChaosClothAsset/ChaosClothAssetDataflowNodesModule.cpp b/Engine/Plugins/Experimental/ChaosClothAssetEditor/Source/ChaosClothAssetDataflowNodes/Private/ChaosClothAsset/ChaosClothAssetDataflowNodesModule.cpp
--- a/Engine/Plugins/Experimental/ChaosClothAssetEditor/Source/ChaosClothAssetDataflowNodes/Private/ChaosClothAsset/ChaosClothAssetDataflowNodesModule.cpp
+++ b/Engine/Plugins/Experimental/ChaosClothAssetEditor/Source/ChaosClothAssetDataflowNodes/Private/ChaosClothAsset/ChaosClothAssetDataflowNodesModule.cpp
@@ -45,6 +45,7 @@
 #include "ChaosClothAsset/SkeletalMeshImportNode.h"
 #include "ChaosClothAsset/StaticMeshImportNode.h"
 #include "ChaosClothAsset/TerminalNode.h"
+#include "ChaosClothAsset/TerminalNodeRefreshAssetCustomization.h"
 #include "ChaosClothAsset/TransferSkinWeightsNode.h"
 #include "ChaosClothAsset/TransformPositionsNode.h"
 #include "ChaosClothAsset/TransformUVsNode.h"
@@ -131,6 +132,7 @@ namespace UE::Chaos::ClothAsset
 				PropertyModule->RegisterCustomPropertyTypeLayout(FChaosClothAssetConnectableIOStringValue::StaticStruct()->GetFName(), FOnGetPropertyTypeCustomizationInstance::CreateStatic(&FWeightedValueCustomization::MakeInstance));
 				PropertyModule->RegisterCustomPropertyTypeLayout(FChaosClothAssetNodeSelectionGroup::StaticStruct()->GetFName(), FOnGetPropertyTypeCustomizationInstance::CreateStatic(&FSelectionGroupCustomization::MakeInstance));
 				PropertyModule->RegisterCustomPropertyTypeLayout(FChaosClothAssetImportFilePath::StaticStruct()->GetFName(), FOnGetPropertyTypeCustomizationInstance::CreateStatic(&FImportFilePathCustomization::MakeInstance));
+				PropertyModule->RegisterCustomPropertyTypeLayout(FChaosClothAssetTerminalNodeRefreshAsset::StaticStruct()->GetFName(), FOnGetPropertyTypeCustomizationInstance::CreateStatic(&FTerminalNodeRefreshAssetCustomization::MakeInstance));
 			}
 
 			// Register modular features
@@ -147,7 +149,11 @@ namespace UE::Chaos::ClothAsset
 					PropertyModule->UnregisterCustomPropertyTypeLayout(FChaosClothAssetWeightedValue::StaticStruct()->GetFName());
 					PropertyModule->UnregisterCustomPropertyTypeLayout(FChaosClothAssetWeightedValueNonAnimatable::StaticStruct()->GetFName());
 					PropertyModule->UnregisterCustomPropertyTypeLayout(FChaosClothAssetWeightedValueNonAnimatableNoLowHighRange::StaticStruct()->GetFName());
+					PropertyModule->UnregisterCustomPropertyTypeLayout(FChaosClothAssetConnectableIStringValue::StaticStruct()->GetFName());
+					PropertyModule->UnregisterCustomPropertyTypeLayout(FChaosClothAssetConnectableIOStringValue::StaticStruct()->GetFName());
 					PropertyModule->UnregisterCustomPropertyTypeLayout(FChaosClothAssetNodeSelectionGroup::StaticStruct()->GetFName());
+					PropertyModule->UnregisterCustomPropertyTypeLayout(FChaosClothAssetImportFilePath::StaticStruct()->GetFName());
+					PropertyModule->UnregisterCustomPropertyTypeLayout(FChaosClothAssetTerminalNodeRefreshAsset::StaticStruct()->GetFName());
 				}
 			}
 
diff --git a/Engine/Plugins/Experimental/ChaosClothAssetEditor/Source/ChaosClothAssetDataflowNodes/Private/ChaosClothAsset/TerminalNode.cpp b/Engine/Plugins/Experimental/ChaosClothAssetEditor/Source/ChaosClothAssetDataflowNodes/Private/ChaosClothAsset/TerminalNode.cpp
--- a/Engine/Plugins/Experimental/ChaosClothAssetEditor/Source/ChaosClothAssetDataflowNodes/Private/ChaosClothAsset/TerminalNode.cpp
+++ b/Engine/Plugins/Experimental/ChaosClothAssetEditor/Source/ChaosClothAssetDataflowNodes/Private/ChaosClothAsset/TerminalNode.cpp
@@ -94,6 +94,12 @@ void FChaosClothAssetTerminalNode::SetAssetValue(TObjectPtr<UObject> Asset, Data
 	{
 		using namespace UE::Chaos::ClothAsset;
 
+		if (RefreshAsset.bRefreshAsset)
+		{
+			bClothCollectionChecksumValid = false;
+			RefreshAsset.bRefreshAsset = false;
+		}
+
 		TArray<TSharedRef<FManagedArrayCollection>> InClothCollections = GetCleanedCollectionLodValues(Context);
 		TArray<TSharedRef<FManagedArrayCollection>>& ClothCollections = ClothAsset->GetClothCollections();
 
@@ -112,6 +118,10 @@ void FChaosClothAssetTerminalNode::SetAssetValue(TObjectPtr<UObject> Asset, Data
 				Chaos::Softs::FCollectionPropertyFacade Properties(ClothCollections[LODIndex]);
 				Properties.UpdateProperties(InClothCollections[LODIndex].ToSharedPtr());
 			}
+
+			// Asset must be resaved
+			ClothAsset->MarkPackageDirty();
+
 			return;
 		}
 
@@ -215,6 +225,9 @@ void FChaosClothAssetTerminalNode::SetAssetValue(TObjectPtr<UObject> Asset, Data
 
 		// Rebuild the asset static data
 		ClothAsset->Build(&LODTransitionDataCache);
+
+		// Asset must be resaved
+		ClothAsset->MarkPackageDirty();
 	}
 }
 
diff --git a/Engine/Plugins/Experimental/ChaosClothAssetEditor/Source/ChaosClothAssetDataflowNodes/Private/ChaosClothAsset/TerminalNodeRefreshAssetCustomization.cpp b/Engine/Plugins/Experimental/ChaosClothAssetEditor/Source/ChaosClothAssetDataflowNodes/Private/ChaosClothAsset/TerminalNodeRefreshAssetCustomization.cpp
new file mode 100644
--- /dev/null
+++ b/Engine/Plugins/Experimental/ChaosClothAssetEditor/Source/ChaosClothAssetDataflowNodes/Private/ChaosClothAsset/TerminalNodeRefreshAssetCustomization.cpp
@@ -0,0 +1,75 @@
+// Copyright Epic Games, Inc. All Rights Reserved.
+
+#include "ChaosClothAsset/TerminalNodeRefreshAssetCustomization.h"
+#include "Misc/MessageDialog.h"
+#include "Misc/PackageName.h"
+#include "Widgets/Images/SImage.h"
+#include "Widgets/Input/SButton.h"
+#include "Widgets/Input/SFilePathPicker.h"
+#include "Widgets/Text/STextBlock.h"
+#include "DetailLayoutBuilder.h"
+#include "DetailWidgetRow.h"
+#include "EditorDirectories.h"
+#include "PropertyHandle.h"
+
+#define LOCTEXT_NAMESPACE "ImportFilePathCustomization"
+
+namespace UE::Chaos::ClothAsset
+{
+	TSharedRef<IPropertyTypeCustomization> FTerminalNodeRefreshAssetCustomization::MakeInstance()
+	{
+		return MakeShareable(new FTerminalNodeRefreshAssetCustomization);
+	}
+
+	void FTerminalNodeRefreshAssetCustomization::CustomizeHeader(TSharedRef<IPropertyHandle> StructPropertyHandle, FDetailWidgetRow& HeaderRow, IPropertyTypeCustomizationUtils& CustomizationUtils)
+	{
+		RefreshAsset = StructPropertyHandle->GetChildHandle(TEXT("bRefreshAsset"));
+
+		// Create whole row header
+		HeaderRow
+		.WholeRowContent()
+		.MinDesiredWidth(250)
+		.MaxDesiredWidth(0)
+		[
+			SNew(SHorizontalBox)
+			+ SHorizontalBox::Slot()
+			.AutoWidth()
+			[
+				SNew(SButton)
+				.ButtonStyle(FAppStyle::Get(), "RoundButton")
+				.ToolTipText(LOCTEXT("Button_ReimportAsset_Tooltip", "Refresh asset"))
+				.OnClicked(this, &FTerminalNodeRefreshAssetCustomization::OnClicked)
+				[
+					SNew(SHorizontalBox)
+					+ SHorizontalBox::Slot()
+					.AutoWidth()
+					[
+						SNew(SImage)
+						.Image(FAppStyle::GetBrush("Persona.ReimportAsset"))
+					]
+					+ SHorizontalBox::Slot()
+					.AutoWidth()
+					.Padding(FMargin(4, 0, 0, 0))
+					.VAlign(VAlign_Bottom)
+					[
+						SNew(STextBlock)
+						.Font(IDetailLayoutBuilder::GetDetailFontBold())
+						.Text(StructPropertyHandle->GetPropertyDisplayName())
+						.ShadowOffset(FVector2D(1, 1))
+					]
+				]
+			]
+		];
+	}
+
+	FReply FTerminalNodeRefreshAssetCustomization::OnClicked()
+	{
+		if (RefreshAsset)
+		{
+			RefreshAsset->SetValue(true);
+		}
+		return FReply::Handled();
+	}
+}  // End namespace UE::Chaos::ClothAsset
+
+#undef LOCTEXT_NAMESPACE
diff --git a/Engine/Plugins/Experimental/ChaosClothAssetEditor/Source/ChaosClothAssetDataflowNodes/Private/ChaosClothAsset/TerminalNodeRefreshAssetCustomization.h b/Engine/Plugins/Experimental/ChaosClothAssetEditor/Source/ChaosClothAssetDataflowNodes/Private/ChaosClothAsset/TerminalNodeRefreshAssetCustomization.h
new file mode 100644
--- /dev/null
+++ b/Engine/Plugins/Experimental/ChaosClothAssetEditor/Source/ChaosClothAssetDataflowNodes/Private/ChaosClothAsset/TerminalNodeRefreshAssetCustomization.h
@@ -0,0 +1,26 @@
+// Copyright Epic Games, Inc. All Rights Reserved.
+
+#pragma once
+
+#include "Input/Reply.h"
+#include "IPropertyTypeCustomization.h"
+
+class SButton;
+
+namespace UE::Chaos::ClothAsset
+{
+	/** Customization for terminal node refresh button. */
+	class FTerminalNodeRefreshAssetCustomization : public IPropertyTypeCustomization
+	{
+	public:
+		static TSharedRef<IPropertyTypeCustomization> MakeInstance();
+
+		virtual void CustomizeHeader(TSharedRef<IPropertyHandle> StructPropertyHandle, FDetailWidgetRow& HeaderRow, IPropertyTypeCustomizationUtils& CustomizationUtils) override;
+		virtual void CustomizeChildren(TSharedRef<IPropertyHandle> /*StructPropertyHandle*/, IDetailChildrenBuilder& /*ChildBuilder*/, IPropertyTypeCustomizationUtils& /*CustomizationUtils*/) override {}
+
+	private:
+		FReply OnClicked();
+
+		TSharedPtr<IPropertyHandle> RefreshAsset;
+	};
+}
diff --git a/Engine/Plugins/Experimental/ChaosClothAssetEditor/Source/ChaosClothAssetDataflowNodes/Public/ChaosClothAsset/TerminalNode.h b/Engine/Plugins/Experimental/ChaosClothAssetEditor/Source/ChaosClothAssetDataflowNodes/Public/ChaosClothAsset/TerminalNode.h
--- a/Engine/Plugins/Experimental/ChaosClothAssetEditor/Source/ChaosClothAssetDataflowNodes/Public/ChaosClothAsset/TerminalNode.h
+++ b/Engine/Plugins/Experimental/ChaosClothAssetEditor/Source/ChaosClothAssetDataflowNodes/Public/ChaosClothAsset/TerminalNode.h
@@ -7,6 +7,16 @@
 #include "ChaosClothAsset/ClothLodTransitionDataCache.h"
 #include "TerminalNode.generated.h"
 
+/** Refresh structure for push buton customization. */
+USTRUCT()
+struct FChaosClothAssetTerminalNodeRefreshAsset
+{
+	GENERATED_USTRUCT_BODY()
+
+	UPROPERTY(EditAnywhere, Category = "Terminal Node Refresh Asset")
+	bool bRefreshAsset = false;
+};
+
 /** Cloth terminal node to generate a cloth asset from a cloth collection. */
 USTRUCT(Meta = (DataflowCloth, DataflowTerminal))
 struct FChaosClothAssetTerminalNode : public FDataflowTerminalNode
@@ -38,6 +48,9 @@ public:
 	/** The number of LODs currently exposed to the node UI. */
 	UPROPERTY()
 	int32 NumLods = 1;
+	/** Refresh the asset even if the ClothCollection hasn't changed to take into account changes in the build code for example. */
+	UPROPERTY(EditAnywhere, Category = "Cloth Asset Terminal")
+	mutable FChaosClothAssetTerminalNodeRefreshAsset RefreshAsset;
 
 	FChaosClothAssetTerminalNode(const Dataflow::FNodeParameters& InParam, FGuid InGuid = FGuid::NewGuid());
 
