Cloth Weight Map Paint Tool/Node: allow weight map painting on Render vertices. The AddWeightMapNode now stores two sets of vertex weights: one for sim vertices and one for rendering. Changing the view mode while the paint tool is active will load the appropriate weight map from the node. There is also a drop-down selector for choosing which of the two vertex weights should be saved into the ClothCollection (or both.)

diff --git a/Engine/Plugins/Experimental/ChaosClothAssetEditor/Source/ChaosClothAssetDataflowNodes/Private/ChaosClothAsset/AddWeightMapNode.cpp b/Engine/Plugins/Experimental/ChaosClothAssetEditor/Source/ChaosClothAssetDataflowNodes/Private/ChaosClothAsset/AddWeightMapNode.cpp
--- a/Engine/Plugins/Experimental/ChaosClothAssetEditor/Source/ChaosClothAssetDataflowNodes/Private/ChaosClothAsset/AddWeightMapNode.cpp
+++ b/Engine/Plugins/Experimental/ChaosClothAssetEditor/Source/ChaosClothAssetDataflowNodes/Private/ChaosClothAsset/AddWeightMapNode.cpp
@@ -2,6 +2,7 @@
 
 #include "ChaosClothAsset/AddWeightMapNode.h"
 #include "ChaosClothAsset/ClothAsset.h"
+#include "ChaosClothAsset/ClothCollectionGroup.h"
 #include "ChaosClothAsset/ClothDataflowTools.h"
 #include "ChaosClothAsset/ClothGeometryTools.h"
 #include "ChaosClothAsset/CollectionClothFacade.h"
@@ -220,6 +221,25 @@ void FChaosClothAssetAddWeightMapNode::Evaluate(Dataflow::FContext& Context, con
 {
 	using namespace UE::Chaos::ClothAsset;
 
+	auto CopyWeightsIntoClothCollection = [this](TArrayView<float>& ClothWeights, const TArray<float>& VertexWeights, bool bIsSim)
+	{
+		const int32 MaxWeightIndex = FMath::Min(VertexWeights.Num(), ClothWeights.Num());
+		if (VertexWeights.Num() > 0 && VertexWeights.Num() != ClothWeights.Num())
+		{
+			FClothDataflowTools::LogAndToastWarning(*this,
+				LOCTEXT("VertexCountMismatchHeadline", "Vertex count mismatch."),
+				FText::Format(LOCTEXT("VertexCountMismatchDetails", "{0} vertex weights in the node: {1}\n{0} vertices in the cloth: {2}"),
+					bIsSim ? FText::FromString("Sim") : FText::FromString("Render"),
+					VertexWeights.Num(),
+					ClothWeights.Num()));
+		}
+
+		for (int32 VertexID = 0; VertexID < MaxWeightIndex; ++VertexID)
+		{
+			ClothWeights[VertexID] = VertexWeights[VertexID];
+		}
+	};
+
 	if (Out->IsA<FManagedArrayCollection>(&Collection))
 	{
 		// Evaluate in collection
@@ -229,32 +249,51 @@ void FChaosClothAssetAddWeightMapNode::Evaluate(Dataflow::FContext& Context, con
 		if (ClothFacade.IsValid())  // Can only act on the collection if it is a valid cloth collection
 		{
 			const FName InName(Name);
-			ClothFacade.AddWeightMap(InName);		// Does nothing if weight map already exists
 
-			TArrayView<float> ClothWeights = ClothFacade.GetWeightMap(InName);
-			if (ClothWeights.Num() != ClothFacade.GetNumSimVertices3D())
-			{
-				check(ClothWeights.Num() == 0);
-				FClothDataflowTools::LogAndToastWarning(*this,
-					LOCTEXT("InvalidWeightMapNameHeadline", "Invalid weight map name."),
-					FText::Format(LOCTEXT("InvalidWeightMapNameDetails", "Could not create a weight map with name \"{0}\" (reserved name? wrong type?)."),
-						FText::FromName(InName)));
-			}
+			// Copy simulation weights into cloth collection
 
-			const int32 MaxWeightIndex = FMath::Min(GetVertexWeights().Num(), ClothWeights.Num());
-			if (GetVertexWeights().Num() > 0 && GetVertexWeights().Num() != ClothWeights.Num())
+			if (MeshTarget == EChaosClothAssetWeightMapMeshType::Simulation || MeshTarget == EChaosClothAssetWeightMapMeshType::Both)
 			{
-				FClothDataflowTools::LogAndToastWarning(*this,
-					LOCTEXT("VertexCountMismatchHeadline", "Vertex count mismatch."),
-					FText::Format(LOCTEXT("VertexCountMismatchDetails", "Vertex weights in the node: {0}\n3D vertices in the cloth: {1}"),
-						GetVertexWeights().Num(),
-						ClothWeights.Num()));
+				ClothFacade.AddWeightMap(InName);		// Does nothing if weight map already exists
+				TArrayView<float> ClothSimWeights = ClothFacade.GetWeightMap(InName);
+
+				if (ClothSimWeights.Num() != ClothFacade.GetNumSimVertices3D())
+				{
+					check(ClothSimWeights.Num() == 0);
+					FClothDataflowTools::LogAndToastWarning(*this,
+						LOCTEXT("InvalidSimWeightMapNameHeadline", "Invalid weight map name."),
+						FText::Format(LOCTEXT("InvalidSimWeightMapNameDetails", "Could not create a sim weight map with name \"{0}\" (reserved name? wrong type?)."),
+							FText::FromName(InName)));
+				}
+				else
+				{
+					constexpr bool bIsSim = true;
+					CopyWeightsIntoClothCollection(ClothSimWeights, GetVertexWeights(), bIsSim);
+				}
 			}
+			
+			// Copy render weights into cloth collection
 
-			for (int32 VertexID = 0; VertexID < MaxWeightIndex; ++VertexID)
+			if (MeshTarget == EChaosClothAssetWeightMapMeshType::Render || MeshTarget == EChaosClothAssetWeightMapMeshType::Both)
 			{
-				ClothWeights[VertexID] = GetVertexWeights()[VertexID];
+				ClothFacade.AddUserDefinedAttribute<float>(InName, ClothCollectionGroup::RenderVertices);
+				TArrayView<float> ClothRenderWeights = ClothFacade.GetUserDefinedAttribute<float>(InName, ClothCollectionGroup::RenderVertices);
+
+				if (ClothRenderWeights.Num() != ClothFacade.GetNumRenderVertices())
+				{
+					check(ClothRenderWeights.Num() == 0);
+					FClothDataflowTools::LogAndToastWarning(*this,
+						LOCTEXT("InvalidRenderWeightMapNameHeadline", "Invalid weight map name."),
+						FText::Format(LOCTEXT("InvalidRenderWeightMapNameDetails", "Could not create a render weight map with name \"{0}\" (reserved name? wrong type?)."),
+							FText::FromName(InName)));
+				}
+				else
+				{
+					constexpr bool bIsSim = false;
+					CopyWeightsIntoClothCollection(ClothRenderWeights, GetRenderVertexWeights(), bIsSim);
+				}
 			}
+
 		}
 		SetValue(Context, MoveTemp(*ClothCollection), &Collection);
 	}
diff --git a/Engine/Plugins/Experimental/ChaosClothAssetEditor/Source/ChaosClothAssetDataflowNodes/Public/ChaosClothAsset/AddWeightMapNode.h b/Engine/Plugins/Experimental/ChaosClothAssetEditor/Source/ChaosClothAssetDataflowNodes/Public/ChaosClothAsset/AddWeightMapNode.h
--- a/Engine/Plugins/Experimental/ChaosClothAssetEditor/Source/ChaosClothAssetDataflowNodes/Public/ChaosClothAsset/AddWeightMapNode.h
+++ b/Engine/Plugins/Experimental/ChaosClothAssetEditor/Source/ChaosClothAssetDataflowNodes/Public/ChaosClothAsset/AddWeightMapNode.h
@@ -15,6 +15,16 @@ enum class EChaosClothAssetWeightMapTransferType : uint8
 	Use3DSimMesh UMETA(DisplayName = "Use 3D Sim Mesh"),
 };
 
+/** Which mesh to update with the corresponding weight map */
+UENUM()
+enum class EChaosClothAssetWeightMapMeshType : uint8
+{
+	Simulation,
+	Render,
+	Both
+};
+
+
 /** Painted weight map attributes node. */
 PRAGMA_DISABLE_DEPRECATION_WARNINGS  // For deprecated VertexWeights in copy constructor (Clang)
 USTRUCT(Meta = (DataflowCloth))
@@ -53,6 +63,9 @@ public:
 	UPROPERTY()
 	TArray<float> VertexWeights;
 
+	UPROPERTY(EditAnywhere, Category = "Add Weight Map")
+	EChaosClothAssetWeightMapMeshType MeshTarget = EChaosClothAssetWeightMapMeshType::Simulation;
+
 	FChaosClothAssetAddWeightMapNode(const Dataflow::FNodeParameters& InParam, FGuid InGuid = FGuid::NewGuid());
 
 private:
@@ -61,6 +74,12 @@ private:
 	const TArray<float>& GetVertexWeights() const { return VertexWeights; }
 	TArray<float>& GetVertexWeights() { return VertexWeights; }
 
+	const TArray<float>& GetRenderVertexWeights() const { return RenderVertexWeights; }
+	TArray<float>& GetRenderVertexWeights() { return RenderVertexWeights; }
+
+	UPROPERTY()
+	TArray<float> RenderVertexWeights;
+
 	//~ Begin FDataflowNode interface
 	virtual void SetAssetValue(TObjectPtr<UObject> Asset, Dataflow::FContext& Context) const override;
 	virtual void Evaluate(Dataflow::FContext& Context, const FDataflowOutput* Out) const override;
diff --git a/Engine/Plugins/Experimental/ChaosClothAssetEditor/Source/ChaosClothAssetEditorTools/Private/ChaosClothAsset/ClothWeightMapPaintTool.cpp b/Engine/Plugins/Experimental/ChaosClothAssetEditor/Source/ChaosClothAssetEditorTools/Private/ChaosClothAsset/ClothWeightMapPaintTool.cpp
--- a/Engine/Plugins/Experimental/ChaosClothAssetEditor/Source/ChaosClothAssetEditorTools/Private/ChaosClothAsset/ClothWeightMapPaintTool.cpp
+++ b/Engine/Plugins/Experimental/ChaosClothAssetEditor/Source/ChaosClothAssetEditorTools/Private/ChaosClothAsset/ClothWeightMapPaintTool.cpp
@@ -73,6 +73,7 @@ void UClothEditorWeightMapPaintToolBuilder::GetSupportedViewModes(TArray<UE::Cha
 {
 	Modes.Add(UE::Chaos::ClothAsset::EClothPatternVertexType::Sim3D);
 	Modes.Add(UE::Chaos::ClothAsset::EClothPatternVertexType::Sim2D);
+	Modes.Add(UE::Chaos::ClothAsset::EClothPatternVertexType::Render);
 }
 
 UMeshSurfacePointTool* UClothEditorWeightMapPaintToolBuilder::CreateNewTool(const FToolBuilderState& SceneState) const
@@ -468,7 +469,8 @@ void UClothEditorWeightMapPaintTool::Setup()
 	ActiveWeightMap->SetName(FName("PaintLayer"));
 
 	// Copy weights from selected node to the preview mesh
-	const TArray<float>& CurrentWeights = WeightMapNodeToUpdate->GetVertexWeights();
+	const bool bIsRenderMode = (ClothEditorContextObject->GetConstructionViewMode() == UE::Chaos::ClothAsset::EClothPatternVertexType::Render);
+	const TArray<float>& CurrentWeights = bIsRenderMode ? WeightMapNodeToUpdate->GetRenderVertexWeights() : WeightMapNodeToUpdate->GetVertexWeights();
 	
 	if (bHaveDynamicMeshToWeightConversion)
 	{
@@ -1888,18 +1890,21 @@ void UClothEditorWeightMapPaintTool::UpdateSelectedNode()
 	GetCurrentWeightMap(CurrentWeights);
 
 	check(WeightMapNodeToUpdate);
+	const bool bIsRenderMode = (ClothEditorContextObject->GetConstructionViewMode() == UE::Chaos::ClothAsset::EClothPatternVertexType::Render);
+
+	TArray<float>& NodeWeights = bIsRenderMode ? WeightMapNodeToUpdate->GetRenderVertexWeights() : WeightMapNodeToUpdate->GetVertexWeights();
 
 	if (bHaveDynamicMeshToWeightConversion)
 	{
-		WeightMapNodeToUpdate->GetVertexWeights().Init(0.f, WeightToDynamicMesh.Num());
+		NodeWeights.Init(0.f, WeightToDynamicMesh.Num());
 		for (int32 DynamicMeshIdx = 0; DynamicMeshIdx < CurrentWeights.Num(); ++DynamicMeshIdx)
 		{
-			WeightMapNodeToUpdate->GetVertexWeights()[DynamicMeshToWeight[DynamicMeshIdx]] = CurrentWeights[DynamicMeshIdx];
+			NodeWeights[DynamicMeshToWeight[DynamicMeshIdx]] = CurrentWeights[DynamicMeshIdx];
 		}
 	}
 	else
 	{
-		WeightMapNodeToUpdate->GetVertexWeights() = CurrentWeights;
+		NodeWeights = CurrentWeights;
 	}
 	
 	WeightMapNodeToUpdate->Name = UpdateWeightMapProperties->Name;
