Fix "lamba parameter has same name as member variable" error in CIS

diff --git a/Engine/Plugins/Experimental/ChaosClothAssetEditor/Source/ChaosClothAssetDataflowNodes/Private/ChaosClothAsset/AddWeightMapNode.cpp b/Engine/Plugins/Experimental/ChaosClothAssetEditor/Source/ChaosClothAssetDataflowNodes/Private/ChaosClothAsset/AddWeightMapNode.cpp
--- a/Engine/Plugins/Experimental/ChaosClothAssetEditor/Source/ChaosClothAssetDataflowNodes/Private/ChaosClothAsset/AddWeightMapNode.cpp
+++ b/Engine/Plugins/Experimental/ChaosClothAssetEditor/Source/ChaosClothAssetDataflowNodes/Private/ChaosClothAsset/AddWeightMapNode.cpp
@@ -221,22 +221,22 @@ void FChaosClothAssetAddWeightMapNode::Evaluate(Dataflow::FContext& Context, con
 {
 	using namespace UE::Chaos::ClothAsset;
 
-	auto CopyWeightsIntoClothCollection = [this](TArrayView<float>& ClothWeights, const TArray<float>& VertexWeights, bool bIsSim)
+	auto CopyWeightsIntoClothCollection = [this](TArrayView<float>& ClothWeights, const TArray<float>& SourceVertexWeights, bool bIsSim)
 	{
-		const int32 MaxWeightIndex = FMath::Min(VertexWeights.Num(), ClothWeights.Num());
-		if (VertexWeights.Num() > 0 && VertexWeights.Num() != ClothWeights.Num())
+		const int32 MaxWeightIndex = FMath::Min(SourceVertexWeights.Num(), ClothWeights.Num());
+		if (SourceVertexWeights.Num() > 0 && SourceVertexWeights.Num() != ClothWeights.Num())
 		{
 			FClothDataflowTools::LogAndToastWarning(*this,
 				LOCTEXT("VertexCountMismatchHeadline", "Vertex count mismatch."),
 				FText::Format(LOCTEXT("VertexCountMismatchDetails", "{0} vertex weights in the node: {1}\n{0} vertices in the cloth: {2}"),
 					bIsSim ? FText::FromString("Sim") : FText::FromString("Render"),
-					VertexWeights.Num(),
+					SourceVertexWeights.Num(),
 					ClothWeights.Num()));
 		}
 
 		for (int32 VertexID = 0; VertexID < MaxWeightIndex; ++VertexID)
 		{
-			ClothWeights[VertexID] = VertexWeights[VertexID];
+			ClothWeights[VertexID] = SourceVertexWeights[VertexID];
 		}
 	};
 
