Cloth Editor: don't start the preview animation any time anything changes in the graph

diff --git a/Engine/Plugins/Experimental/ChaosClothAssetEditor/Source/ChaosClothAssetEditor/Private/ChaosClothAsset/ClothEditorPreviewScene.cpp b/Engine/Plugins/Experimental/ChaosClothAssetEditor/Source/ChaosClothAssetEditor/Private/ChaosClothAsset/ClothEditorPreviewScene.cpp
--- a/Engine/Plugins/Experimental/ChaosClothAssetEditor/Source/ChaosClothAssetEditor/Private/ChaosClothAsset/ClothEditorPreviewScene.cpp
+++ b/Engine/Plugins/Experimental/ChaosClothAssetEditor/Source/ChaosClothAssetEditor/Private/ChaosClothAsset/ClothEditorPreviewScene.cpp
@@ -201,7 +201,10 @@ void FChaosClothPreviewScene::AddReferencedObjects(FReferenceCollector& Collecto
 void FChaosClothPreviewScene::UpdateSkeletalMeshAnimation()
 {
 	check(SkeletalMeshComponent);
-	
+
+	const bool bWasPlaying = SkeletalMeshComponent->IsPlaying();
+	SkeletalMeshComponent->Stop();
+
 	if (PreviewSceneDescription->AnimationAsset)
 	{
 		PreviewAnimInstance = NewObject<UAnimSingleNodeInstance>(SkeletalMeshComponent);
@@ -213,10 +216,14 @@ void FChaosClothPreviewScene::UpdateSkeletalMeshAnimation()
 		SkeletalMeshComponent->AnimScriptInstance = PreviewAnimInstance;
 		SkeletalMeshComponent->AnimScriptInstance->InitializeAnimation();
 		SkeletalMeshComponent->ValidateAnimation();
+
+		if (!bWasPlaying)
+		{
+			SkeletalMeshComponent->Stop();
+		}
 	}
 	else
 	{
-		SkeletalMeshComponent->Stop();
 		SkeletalMeshComponent->AnimationData = FSingleAnimationPlayData();
 		SkeletalMeshComponent->AnimScriptInstance = nullptr;
 	}
@@ -352,18 +359,25 @@ void FChaosClothPreviewScene::SetClothAsset(UChaosClothAsset* Asset)
 
 	if (USkeletalMesh* const SkeletalMesh = Asset->GetPreviewSceneSkeletalMesh())
 	{
-		PreviewSceneDescription->SkeletalMeshAsset = SkeletalMesh;
+		if (SkeletalMesh != PreviewSceneDescription->SkeletalMeshAsset)
+		{
+			PreviewSceneDescription->SkeletalMeshAsset = SkeletalMesh;
 
-		SkeletalMeshComponent->SetSkeletalMeshAsset(PreviewSceneDescription->SkeletalMeshAsset);
-		UpdateSkeletalMeshAnimation();
-		UpdateClothComponentAttachment();
+			SkeletalMeshComponent->SetSkeletalMeshAsset(PreviewSceneDescription->SkeletalMeshAsset);
+
+			UpdateSkeletalMeshAnimation();
+			UpdateClothComponentAttachment();
+		}
 	}
 	
 	if (UAnimationAsset* const Animation = Asset->GetPreviewSceneAnimation())
 	{
-		PreviewSceneDescription->AnimationAsset = Animation;
+		if (PreviewSceneDescription->AnimationAsset != Animation)
+		{
+			PreviewSceneDescription->AnimationAsset = Animation;
 
-		UpdateSkeletalMeshAnimation();
+			UpdateSkeletalMeshAnimation();
+		}
 	}
 }
 
