Cloth paint tool changes
- set default Strength to 1.0
- hook up Q/A hotkeys to cycle through brush modes
- hook up W/E hotkeys to adjust Value property

diff --git a/Engine/Plugins/Experimental/ChaosClothAssetEditor/Source/ChaosClothAssetEditorTools/Private/ChaosClothAsset/ClothWeightMapPaintTool.cpp b/Engine/Plugins/Experimental/ChaosClothAssetEditor/Source/ChaosClothAssetEditorTools/Private/ChaosClothAsset/ClothWeightMapPaintTool.cpp
--- a/Engine/Plugins/Experimental/ChaosClothAssetEditor/Source/ChaosClothAssetEditorTools/Private/ChaosClothAsset/ClothWeightMapPaintTool.cpp
+++ b/Engine/Plugins/Experimental/ChaosClothAssetEditor/Source/ChaosClothAssetEditorTools/Private/ChaosClothAsset/ClothWeightMapPaintTool.cpp
@@ -521,6 +521,35 @@ void UClothEditorWeightMapPaintTool::InitializeBrushSizeRange(const UE::Geometry
 	CalculateBrushRadius();
 }
 
+void UClothEditorWeightMapPaintTool::NextBrushModeAction()
+{
+	constexpr uint8 NumCyclableBrushes = 2;		// Don't cycle to the hidden Erase brush
+	FilterProperties->PrimaryBrushType = static_cast<EClothEditorWeightMapPaintBrushType>((static_cast<uint8>(FilterProperties->PrimaryBrushType) + 1) % NumCyclableBrushes);
+}
+
+void UClothEditorWeightMapPaintTool::PreviousBrushModeAction()
+{
+	constexpr uint8 NumCyclableBrushes = 2;		// Don't cycle to the hidden Erase brush
+	const uint8 CurrentBrushType = static_cast<uint8>(FilterProperties->PrimaryBrushType);
+	const uint8 NewBrushType = CurrentBrushType == 0 ? NumCyclableBrushes - 1 : (CurrentBrushType - 1) % NumCyclableBrushes;
+	FilterProperties->PrimaryBrushType = static_cast<EClothEditorWeightMapPaintBrushType>(NewBrushType);
+}
+
+void UClothEditorWeightMapPaintTool::IncreaseBrushSpeedAction()		// Actually increases AttributeValue
+{
+	const double CurrentValue = FilterProperties->AttributeValue;
+	FilterProperties->AttributeValue = FMath::Clamp(CurrentValue + 0.05, 0.0, 1.0);
+	NotifyOfPropertyChangeByTool(FilterProperties);
+}
+
+void UClothEditorWeightMapPaintTool::DecreaseBrushSpeedAction()		// Actually decreases AttributeValue
+{
+	const double CurrentValue = FilterProperties->AttributeValue;
+	FilterProperties->AttributeValue = FMath::Clamp(CurrentValue - 0.05, 0.0, 1.0);
+	NotifyOfPropertyChangeByTool(FilterProperties);
+}
+
+
 void UClothEditorWeightMapPaintTool::SetClothEditorContextObject(TObjectPtr<UClothEditorContextObject> InClothEditorContextObject)
 {
 	ClothEditorContextObject = InClothEditorContextObject;
@@ -576,6 +605,34 @@ void UClothEditorWeightMapPaintTool::RegisterActions(FInteractiveToolActionSet&
 		LOCTEXT("PickWeightValueUnderCursorTooltip", "Set the active weight painting value to that currently under the cursor"),
 		EModifierKey::Shift, EKeys::G,
 		[this]() { bPendingPickWeight = true; });
+
+
+	// E/W are overridden to decrease/increase the AttributeValue property. Use shift-E/shift-W to increment by smaller amount
+
+	ActionSet.RegisterAction(this, (int32)EStandardToolActions::BaseClientDefinedActionID + 503,
+		TEXT("WeightMapPaintIncreaseValueSmallStep"),
+		LOCTEXT("WeightMapPaintIncreaseValueSmallStep", "Increase Value"),
+		LOCTEXT("WeightMapPaintIncreaseValueSmallStepTooltip", "Increase Value (small increment)"),
+		EModifierKey::Shift, EKeys::E,
+		[this]()
+		{
+			const double CurrentValue = FilterProperties->AttributeValue;
+			FilterProperties->AttributeValue = FMath::Clamp(CurrentValue + 0.005, 0.0, 1.0);
+			NotifyOfPropertyChangeByTool(FilterProperties);
+		});
+
+	ActionSet.RegisterAction(this, (int32)EStandardToolActions::BaseClientDefinedActionID + 504,
+		TEXT("WeightMapPaintDecreaseValueSmallStep"),
+		LOCTEXT("WeightMapPaintDecreaseValueSmallStep", "Decrease Value"),
+		LOCTEXT("WeightMapPaintDecreaseValueSmallStepTooltip", "Decrease Value (small increment)"),
+		EModifierKey::Shift, EKeys::W,
+		[this]()
+		{
+			const double CurrentValue = FilterProperties->AttributeValue;
+			FilterProperties->AttributeValue = FMath::Clamp(CurrentValue - 0.005, 0.0, 1.0);
+			NotifyOfPropertyChangeByTool(FilterProperties);
+		});
+
 };
 
 
@@ -2145,7 +2202,16 @@ void UClothEditorWeightMapPaintTool::UpdateSubToolType(EClothEditorWeightMapPain
 
 void UClothEditorWeightMapPaintTool::UpdateBrushType(EClothEditorWeightMapPaintBrushType BrushType)
 {
-	static const FText BaseMessage = LOCTEXT("OnStartTool", "Hold Shift to Erase. [/] and S/D change Size (+Shift to small-step)");
+	FText BaseMessage;
+	if (BrushType == EClothEditorWeightMapPaintBrushType::Paint)
+	{
+		BaseMessage = LOCTEXT("OnStartTool", "Hold Shift to Erase. Use [/] and S/D keys to change brush size (+Shift to small-step). W/E to change Value (+Shift to small-step). Shift-G to get current Value under cursor. Q/A to cycle through brush modes.");
+	}
+	else if (BrushType == EClothEditorWeightMapPaintBrushType::Smooth)
+	{
+		BaseMessage = LOCTEXT("OnStartTool", "Hold Shift to Erase. Use [/] and S/D keys to change brush size (+Shift to small-step). Q/A to cycle through brush modes.");
+	}
+
 	FTextBuilder Builder;
 	Builder.AppendLine(BaseMessage);
 
diff --git a/Engine/Plugins/Experimental/ChaosClothAssetEditor/Source/ChaosClothAssetEditorTools/Public/ChaosClothAsset/ClothWeightMapPaintTool.h b/Engine/Plugins/Experimental/ChaosClothAssetEditor/Source/ChaosClothAssetEditorTools/Public/ChaosClothAsset/ClothWeightMapPaintTool.h
--- a/Engine/Plugins/Experimental/ChaosClothAssetEditor/Source/ChaosClothAssetEditorTools/Public/ChaosClothAsset/ClothWeightMapPaintTool.h
+++ b/Engine/Plugins/Experimental/ChaosClothAssetEditor/Source/ChaosClothAssetEditorTools/Public/ChaosClothAsset/ClothWeightMapPaintTool.h
@@ -163,7 +163,7 @@ public:
 	UPROPERTY(EditAnywhere, Category = Brush, meta = (UIMin = 0, ClampMin = 0, UIMax = 1, ClampMax = 1,
 		HideEditConditionToggle, EditConditionHides, EditCondition = 
 		"SubToolType == EClothEditorWeightMapPaintInteractionType::Brush || SubToolType == EClothEditorWeightMapPaintInteractionType::Fill"))
-	double Strength = 0.5;
+	double Strength = 1.0;
 
 	/** The Gradient upper limit value */
 	UPROPERTY(EditAnywhere, Category = Gradient, meta = (UIMin = 0, ClampMin = 0, UIMax = 1, ClampMax = 1,
@@ -382,6 +382,14 @@ protected:
 
 	virtual void InitializeBrushSizeRange(const UE::Geometry::FAxisAlignedBox3d& TargetBounds) override;
 
+	virtual void NextBrushModeAction() override;
+	virtual void PreviousBrushModeAction() override;
+
+	// Note: these will actually modify the brush's Attribute Value since we don't use Speed in our brush
+	virtual void IncreaseBrushSpeedAction() override;
+	virtual void DecreaseBrushSpeedAction() override;
+
+
 	// end UMeshSculptToolBase API
 
 
