Cloth Editor: optional information in the Construction viewport
- Numeric mesh stats (so far just triangle/vertex count)
- Color-coded cloth patterns

diff --git a/Engine/Plugins/Experimental/ChaosClothAssetEditor/Source/ChaosClothAssetEditor/Private/ChaosClothAsset/ClothEditorCommands.cpp b/Engine/Plugins/Experimental/ChaosClothAssetEditor/Source/ChaosClothAssetEditor/Private/ChaosClothAsset/ClothEditorCommands.cpp
--- a/Engine/Plugins/Experimental/ChaosClothAssetEditor/Source/ChaosClothAssetEditor/Private/ChaosClothAsset/ClothEditorCommands.cpp
+++ b/Engine/Plugins/Experimental/ChaosClothAssetEditor/Source/ChaosClothAssetEditor/Private/ChaosClothAsset/ClothEditorCommands.cpp
@@ -61,6 +61,9 @@ void FChaosClothAssetEditorCommands::RegisterCommands()
 	UI_COMMAND(ToggleConstructionViewSeams, "ToggleConstructionViewSeams", "Toggle construction view seams", EUserInterfaceActionType::ToggleButton, FInputChord());
 	UI_COMMAND(ToggleConstructionViewSeamsCollapse, "ToggleConstructionViewSeamsCollapse", "Toggle drawing one line between seams in construction view", EUserInterfaceActionType::ToggleButton, FInputChord());
 
+	UI_COMMAND(TogglePatternColor, "TogglePatternColor", "Toggle Pattern color in the construction view", EUserInterfaceActionType::ToggleButton, FInputChord());
+	UI_COMMAND(ToggleMeshStats, "ToggleMeshStats", "Toggle mesh information display", EUserInterfaceActionType::ToggleButton, FInputChord());
+	
 	UI_COMMAND(SoftResetSimulation, "SoftResetSimulation", "Soft reset simulation", EUserInterfaceActionType::ToggleButton, FInputChord());
 	UI_COMMAND(HardResetSimulation, "HardResetClothSimulation", "Hard reset simulation", EUserInterfaceActionType::ToggleButton, FInputChord(EModifierKey::Control | EModifierKey::Alt, EKeys::C));
 	UI_COMMAND(ToggleSimulationSuspended, "ToggleSimulationSuspended", "Toggle simulation suspended", EUserInterfaceActionType::ToggleButton, FInputChord());
diff --git a/Engine/Plugins/Experimental/ChaosClothAssetEditor/Source/ChaosClothAssetEditor/Private/ChaosClothAsset/ClothEditorCommands.h b/Engine/Plugins/Experimental/ChaosClothAssetEditor/Source/ChaosClothAssetEditor/Private/ChaosClothAsset/ClothEditorCommands.h
--- a/Engine/Plugins/Experimental/ChaosClothAssetEditor/Source/ChaosClothAssetEditor/Private/ChaosClothAsset/ClothEditorCommands.h
+++ b/Engine/Plugins/Experimental/ChaosClothAssetEditor/Source/ChaosClothAssetEditor/Private/ChaosClothAsset/ClothEditorCommands.h
@@ -61,6 +61,9 @@ public:
 	const static FString ToggleConstructionViewSeamsCollapseIdentifier;
 	TSharedPtr<FUICommandInfo> ToggleConstructionViewSeamsCollapse;
 
+	TSharedPtr<FUICommandInfo> TogglePatternColor;
+	TSharedPtr<FUICommandInfo> ToggleMeshStats;
+
 	TSharedPtr<FUICommandInfo> SetConstructionMode2D;
 	TSharedPtr<FUICommandInfo> SetConstructionMode3D;
 	TSharedPtr<FUICommandInfo> SetConstructionModeRender;
diff --git a/Engine/Plugins/Experimental/ChaosClothAssetEditor/Source/ChaosClothAssetEditor/Private/ChaosClothAsset/ClothEditorMode.cpp b/Engine/Plugins/Experimental/ChaosClothAssetEditor/Source/ChaosClothAssetEditor/Private/ChaosClothAsset/ClothEditorMode.cpp
--- a/Engine/Plugins/Experimental/ChaosClothAssetEditor/Source/ChaosClothAssetEditor/Private/ChaosClothAsset/ClothEditorMode.cpp
+++ b/Engine/Plugins/Experimental/ChaosClothAssetEditor/Source/ChaosClothAssetEditor/Private/ChaosClothAsset/ClothEditorMode.cpp
@@ -101,6 +101,17 @@ namespace UE::Chaos::ClothAsset::Private
 		return OutWeightMapNames;
 	}
 
+	FLinearColor PseudoRandomColor(int32 NumColorRotations)
+	{
+		constexpr uint8 Spread = 157;  // Prime number that gives a good spread of colors without getting too similar as a rand might do.
+		uint8 Seed = Spread;
+		for (int32 Rotation = 0; Rotation < NumColorRotations; ++Rotation)
+		{
+			Seed += Spread;
+		}
+		return FLinearColor::MakeFromHSV8(Seed, 180, 140);
+	}
+
 }
 
 
@@ -519,18 +530,6 @@ void UChaosClothAssetEditorMode::InitializeSeamDraw()
 		return;
 	}
 
-
-	auto PseudoRandomColor = [](int32 NumColorRotations) -> FLinearColor
-	{
-		constexpr uint8 Spread = 157;  // Prime number that gives a good spread of colors without getting too similar as a rand might do.
-		uint8 Seed = Spread;
-		for (int32 Rotation = 0; Rotation < NumColorRotations; ++Rotation)
-		{
-			Seed += Spread;
-		}
-		return FLinearColor::MakeFromHSV8(Seed, 180, 140);
-	};
-
 	const UE::Chaos::ClothAsset::FCollectionClothConstFacade ClothFacade(Collection.ToSharedRef());
 
 	ULineSetComponent* const Lines = ClothSeamDraw->AddLineSet("SeamLines");
@@ -557,7 +556,7 @@ void UChaosClothAssetEditorMode::InitializeSeamDraw()
 
 			for (const TArray<FIntVector2>& ConnectedSeam : ConnectedSeams)
 			{
-				const FColor SeamColor = PseudoRandomColor(ConnectedSeamIndex++).ToFColor(true);
+				const FColor SeamColor = UE::Chaos::ClothAsset::Private::PseudoRandomColor(ConnectedSeamIndex++).ToFColor(true);
 
 				// draw connected edge on each side of the seam
 				for (int32 StitchID = 0; StitchID < ConnectedSeam.Num() - 1; ++StitchID)
@@ -600,7 +599,7 @@ void UChaosClothAssetEditorMode::InitializeSeamDraw()
 		else if (ConstructionViewMode == UE::Chaos::ClothAsset::EClothPatternVertexType::Sim3D)
 		{
 			const TArray<int32> SeamStitches(SeamFacade.GetSeamStitch3DIndex());
-			const FColor SeamColor = PseudoRandomColor(SeamIndex).ToFColor(true);
+			const FColor SeamColor = UE::Chaos::ClothAsset::Private::PseudoRandomColor(SeamIndex).ToFColor(true);
 
 			// In 3D we should be able to draw the seam edges in any order, doesn't need to be in connected paths
 			for (int32 StitchIndexI = 0; StitchIndexI < SeamStitches.Num(); ++StitchIndexI)
@@ -629,6 +628,7 @@ void UChaosClothAssetEditorMode::InitializeSeamDraw()
 void UChaosClothAssetEditorMode::ReinitializeDynamicMeshComponents()
 {
 	using namespace UE::Chaos::ClothAsset;
+	using namespace UE::Geometry;
 
 	auto SetUpDynamicMeshComponentMaterial = [this](const UE::Chaos::ClothAsset::FCollectionClothConstFacade& ClothFacade, UDynamicMeshComponent& MeshComponent)
 	{
@@ -636,14 +636,32 @@ void UChaosClothAssetEditorMode::ReinitializeDynamicMeshComponents()
 		{
 			case EClothPatternVertexType::Sim2D:
 			{
-				UMaterialInterface* const Material = ToolSetupUtil::GetCustomTwoSidedDepthOffsetMaterial(GetToolManager(), FLinearColor{ 0.6, 0.6, 0.6 }, 0.0);
-				MeshComponent.SetMaterial(0, Material);
+				if (bPatternColors)
+				{
+					constexpr bool bTwoSided = true;
+					UMaterialInterface* const Material = ToolSetupUtil::GetVertexColorMaterial(GetToolManager(), bTwoSided);
+					MeshComponent.SetMaterial(0, Material);
+				}
+				else
+				{
+					UMaterialInterface* const Material = ToolSetupUtil::GetCustomTwoSidedDepthOffsetMaterial(GetToolManager(), FLinearColor{ 0.6, 0.6, 0.6 }, 0.0);
+					MeshComponent.SetMaterial(0, Material);
+				}
 			}
 			break;
 			case EClothPatternVertexType::Sim3D:
 			{
-				UMaterialInterface* const Material = ToolSetupUtil::GetDefaultSculptMaterial(GetToolManager());
-				MeshComponent.SetMaterial(0, Material);
+				if (bPatternColors)
+				{
+					constexpr bool bTwoSided = true;
+					UMaterialInterface* const Material = ToolSetupUtil::GetVertexColorMaterial(GetToolManager(), bTwoSided);
+					MeshComponent.SetMaterial(0, Material);
+				}
+				else
+				{
+					UMaterialInterface* const Material = ToolSetupUtil::GetDefaultSculptMaterial(GetToolManager());
+					MeshComponent.SetMaterial(0, Material);
+				}
 			}
 			break;
 			case EClothPatternVertexType::Render:
@@ -729,6 +747,34 @@ void UChaosClothAssetEditorMode::ReinitializeDynamicMeshComponents()
 		UE::Geometry::FMeshNormals::InitializeMeshToPerTriangleNormals(&LodMesh);
 	}
 
+	if (bPatternColors)
+	{
+		LodMesh.Attributes()->EnablePrimaryColors();
+		LodMesh.Attributes()->PrimaryColors()->CreateFromPredicate([&ClothFacade](int ParentVID, int TriIDA, int TriIDB)
+			{
+				return ClothFacade.FindSimPatternByFaceIndex(TriIDA) == ClothFacade.FindSimPatternByFaceIndex(TriIDB);
+			}
+		, 0.0f);
+
+		FDynamicMeshColorOverlay* const ColorAttributeLayer = LodMesh.Attributes()->PrimaryColors();
+		for (int32 PatternID = 0; PatternID < ClothFacade.GetNumSimPatterns(); ++PatternID)
+		{
+			const FCollectionClothSimPatternConstFacade Pattern = ClothFacade.GetSimPattern(PatternID);
+			const FLinearColor PatternColor = Private::PseudoRandomColor(PatternID);
+
+			for (int32 TriID = 0; TriID < Pattern.GetNumSimFaces(); ++TriID)
+			{
+				const int32 GlobalTriID = Pattern.GetSimFacesOffset() + TriID;
+				const FIndex3i AttrTri = LodMesh.Attributes()->PrimaryColors()->GetTriangle(GlobalTriID);
+				ColorAttributeLayer->SetElement(AttrTri[0], (FVector4f)PatternColor);
+				ColorAttributeLayer->SetElement(AttrTri[1], (FVector4f)PatternColor);
+				ColorAttributeLayer->SetElement(AttrTri[2], (FVector4f)PatternColor);
+			}
+		}
+	}
+
+	
+
 	// We only need an actor to allow use of HHitProxy for selection
 	const FRotator Rotation(0.0f, 0.0f, 0.0f);
 	const FActorSpawnParameters SpawnInfo;
@@ -943,6 +989,24 @@ bool UChaosClothAssetEditorMode::IsSimulationEnabled() const
 	return false;
 }
 
+int32 UChaosClothAssetEditorMode::GetConstructionViewTriangleCount() const
+{
+	if (DynamicMeshComponent && DynamicMeshComponent->GetMesh())
+	{
+		return DynamicMeshComponent->GetMesh()->TriangleCount();
+	}
+	return 0;
+}
+
+int32 UChaosClothAssetEditorMode::GetConstructionViewVertexCount() const
+{
+	if (DynamicMeshComponent && DynamicMeshComponent->GetMesh())
+	{
+		return DynamicMeshComponent->GetMesh()->VertexCount();
+	}
+	return 0;
+}
+
 void UChaosClothAssetEditorMode::SetLODModel(int32 LODIndex)
 {
 	if (PreviewScene && PreviewScene->GetClothComponent())
@@ -1265,6 +1329,40 @@ bool UChaosClothAssetEditorMode::CanSetConstructionViewSeamsCollapse() const
 	return bConstructionViewSeamsVisible && (ConstructionViewMode == UE::Chaos::ClothAsset::EClothPatternVertexType::Sim2D);
 }
 
+
+void UChaosClothAssetEditorMode::TogglePatternColor()
+{
+	bPatternColors = !bPatternColors;
+	ReinitializeDynamicMeshComponents();
+}
+
+bool UChaosClothAssetEditorMode::CanSetPatternColor() const
+{	
+	// Disallow pattern color view when any tool is active
+	if (GetToolManager()->HasActiveTool(EToolSide::Left))
+	{
+		return false;
+	}
+
+	return (ConstructionViewMode != UE::Chaos::ClothAsset::EClothPatternVertexType::Render);
+}
+
+void UChaosClothAssetEditorMode::ToggleMeshStats()
+{
+	bMeshStats = !bMeshStats;
+}
+
+bool UChaosClothAssetEditorMode::CanSetMeshStats() const
+{
+	// Disallow seam view when any tool is active
+	if (GetToolManager()->HasActiveTool(EToolSide::Left))
+	{
+		return false;
+	}
+
+	return true;
+}
+
 void UChaosClothAssetEditorMode::SetRestSpaceViewportClient(TWeakPtr<UE::Chaos::ClothAsset::FChaosClothEditorRestSpaceViewportClient, ESPMode::ThreadSafe> InViewportClient)
 {
 	RestSpaceViewportClient = InViewportClient;
diff --git a/Engine/Plugins/Experimental/ChaosClothAssetEditor/Source/ChaosClothAssetEditor/Private/ChaosClothAsset/ClothEditorMode.h b/Engine/Plugins/Experimental/ChaosClothAssetEditor/Source/ChaosClothAssetEditor/Private/ChaosClothAsset/ClothEditorMode.h
--- a/Engine/Plugins/Experimental/ChaosClothAssetEditor/Source/ChaosClothAssetEditor/Private/ChaosClothAsset/ClothEditorMode.h
+++ b/Engine/Plugins/Experimental/ChaosClothAssetEditor/Source/ChaosClothAssetEditor/Private/ChaosClothAsset/ClothEditorMode.h
@@ -87,6 +87,20 @@ public:
 		return bConstructionViewSeamsCollapse;
 	}
 
+	void TogglePatternColor();
+	bool CanSetPatternColor() const;
+	bool IsPatternColorActive() const
+	{
+		return bPatternColors;
+	}
+
+	void ToggleMeshStats();
+	bool CanSetMeshStats() const;
+	bool IsMeshStatsActive() const
+	{
+		return bMeshStats;
+	}
+
 	// Simulation controls
 	void SoftResetSimulation();
 	void HardResetSimulation();
@@ -96,6 +110,9 @@ public:
 	void SetEnableSimulation(bool bEnabled);
 	bool IsSimulationEnabled() const;
 
+	int32 GetConstructionViewTriangleCount() const;
+	int32 GetConstructionViewVertexCount() const;
+
 	// LODIndex == INDEX_NONE is LOD Auto
 	void SetLODModel(int32 LODIndex);
 	bool IsLODModelSelected(int32 LODIndex) const;
@@ -261,6 +278,9 @@ private:
 	bool bConstructionViewSeamsCollapse = false;
 	void InitializeSeamDraw();
 
+	bool bPatternColors = false;
+	bool bMeshStats = false;
+
 	// Create dynamic mesh components from the cloth component's rest space info
 	void ReinitializeDynamicMeshComponents();
 
diff --git a/Engine/Plugins/Experimental/ChaosClothAssetEditor/Source/ChaosClothAssetEditor/Private/ChaosClothAsset/SClothEditorRestSpaceViewport.cpp b/Engine/Plugins/Experimental/ChaosClothAssetEditor/Source/ChaosClothAssetEditor/Private/ChaosClothAsset/SClothEditorRestSpaceViewport.cpp
--- a/Engine/Plugins/Experimental/ChaosClothAssetEditor/Source/ChaosClothAssetEditor/Private/ChaosClothAsset/SClothEditorRestSpaceViewport.cpp
+++ b/Engine/Plugins/Experimental/ChaosClothAssetEditor/Source/ChaosClothAssetEditor/Private/ChaosClothAsset/SClothEditorRestSpaceViewport.cpp
@@ -206,6 +206,58 @@ void SChaosClothAssetEditorRestSpaceViewport::BindCommands()
 			return false;
 		}));
 
+	CommandList->MapAction(
+		CommandInfos.TogglePatternColor,
+		FExecuteAction::CreateLambda([this]()
+			{
+				if (UChaosClothAssetEditorMode* const ClothEdMode = GetEdMode())
+				{
+					ClothEdMode->TogglePatternColor();
+				}
+			}),
+		FCanExecuteAction::CreateLambda([this]()
+			{
+				if (const UChaosClothAssetEditorMode* const ClothEdMode = GetEdMode())
+				{
+					return ClothEdMode->CanSetPatternColor();
+				}
+				return false;
+			}),
+		FIsActionChecked::CreateLambda([this]()
+			{
+				if (const UChaosClothAssetEditorMode* const ClothEdMode = GetEdMode())
+				{
+					return ClothEdMode->IsPatternColorActive();
+				}
+				return false;
+			}));
+
+	CommandList->MapAction(
+		CommandInfos.ToggleMeshStats,
+		FExecuteAction::CreateLambda([this]()
+			{
+				if (UChaosClothAssetEditorMode* const ClothEdMode = GetEdMode())
+				{
+					ClothEdMode->ToggleMeshStats();
+				}
+			}),
+		FCanExecuteAction::CreateLambda([this]()
+			{
+				if (const UChaosClothAssetEditorMode* const ClothEdMode = GetEdMode())
+				{
+					return ClothEdMode->CanSetMeshStats();
+				}
+				return false;
+			}),
+		FIsActionChecked::CreateLambda([this]()
+			{
+				if (const UChaosClothAssetEditorMode* const ClothEdMode = GetEdMode())
+				{
+					return ClothEdMode->IsMeshStatsActive();
+				}
+				return false;
+			}));
+
 }
 
 TSharedPtr<SWidget> SChaosClothAssetEditorRestSpaceViewport::MakeViewportToolbar()
diff --git a/Engine/Plugins/Experimental/ChaosClothAssetEditor/Source/ChaosClothAssetEditor/Private/ChaosClothAsset/SClothEditorRestSpaceViewportToolBar.cpp b/Engine/Plugins/Experimental/ChaosClothAssetEditor/Source/ChaosClothAssetEditor/Private/ChaosClothAsset/SClothEditorRestSpaceViewportToolBar.cpp
--- a/Engine/Plugins/Experimental/ChaosClothAssetEditor/Source/ChaosClothAssetEditor/Private/ChaosClothAsset/SClothEditorRestSpaceViewportToolBar.cpp
+++ b/Engine/Plugins/Experimental/ChaosClothAssetEditor/Source/ChaosClothAssetEditor/Private/ChaosClothAsset/SClothEditorRestSpaceViewportToolBar.cpp
@@ -15,6 +15,7 @@
 #include "ChaosClothAsset/ClothEditorMode.h"
 #include "ToolMenus.h"
 #include "Widgets/Input/SSpinBox.h"
+#include "Widgets/Text/SRichTextBlock.h"
 
 #define LOCTEXT_NAMESPACE "SChaosClothAssetEditorRestSpaceViewportToolBar"
 
@@ -38,7 +39,21 @@ void SChaosClothAssetEditorRestSpaceViewportToolBar::Construct(const FArguments&
 		.BorderImage(FAppStyle::Get().GetBrush("EditorViewportToolBar.Background"))
 		.Cursor(EMouseCursor::Default)
 		[
-			SAssignNew( MainBoxPtr, SHorizontalBox )
+			SNew(SVerticalBox)
+				+ SVerticalBox::Slot()
+				.AutoHeight()
+			[
+				SAssignNew( MainBoxPtr, SHorizontalBox )
+			]
+			+ SVerticalBox::Slot()
+			.Padding(FMargin(4.0f, 3.0f, 0.0f, 0.0f))
+			[
+				// Display text (e.g., item being previewed)
+				SNew(SRichTextBlock)
+					.DecoratorStyleSet(&FAppStyle::Get())
+					.Text(this, &SChaosClothAssetEditorRestSpaceViewportToolBar::GetDisplayString)
+					.TextStyle(&FAppStyle::Get().GetWidgetStyle<FTextBlockStyle>("AnimViewport.MessageText"))
+			]
 		]
 	];
 
@@ -102,6 +117,17 @@ TSharedRef<SWidget> SChaosClothAssetEditorRestSpaceViewportToolBar::GenerateClot
 				NAME_None, 
 				LOCTEXT("SeamsCollapseLabel", "Collapse Seam Lines"), 
 				LOCTEXT("SeamsCollapseTooltip", "Display a single line connecting each seam, rather than all stitches"));
+
+			OptionsMenuBuilder.AddMenuEntry(FChaosClothAssetEditorCommands::Get().TogglePatternColor,
+				NAME_None,
+				LOCTEXT("ColorPatternsLabel", "Color Patterns"),
+				LOCTEXT("ColorPatternsTooltip", "Display each Pattern in a different color"));
+
+			OptionsMenuBuilder.AddMenuEntry(FChaosClothAssetEditorCommands::Get().ToggleMeshStats,
+				NAME_None,
+				LOCTEXT("ToggleMeshStatsLabel", "Mesh Stats"),
+				LOCTEXT("ToggleMeshStatsTooltip", "Show mesh stats in the viewport"));
+
 		}
 		OptionsMenuBuilder.EndSection();
 
@@ -217,6 +243,23 @@ TSharedRef<SWidget> SChaosClothAssetEditorRestSpaceViewportToolBar::MakeToolBar(
 	return ToolbarBuilder.MakeWidget();
 }
 
+FText SChaosClothAssetEditorRestSpaceViewportToolBar::GetDisplayString() const
+{
+	if (const FEditorModeTools* const EditorModeTools = RestSpaceViewportClient->GetModeTools())
+	{
+		if (UChaosClothAssetEditorMode* const ClothEdMode = Cast<UChaosClothAssetEditorMode>(EditorModeTools->GetActiveScriptableMode(UChaosClothAssetEditorMode::EM_ChaosClothAssetEditorModeId)))
+		{
+			if (ClothEdMode->IsMeshStatsActive())
+			{
+				const int TriangleCount = ClothEdMode->GetConstructionViewTriangleCount();
+				const int VertexCount = ClothEdMode->GetConstructionViewVertexCount();
+				const FText MeshStats = FText::Format(LOCTEXT("RestSpaceMeshStats", "Tris: {0}, Verts: {1}"), TriangleCount, VertexCount);
+				return MeshStats;
+			}
+		}
+	}
+	return FText();
+}
 
 FText SChaosClothAssetEditorRestSpaceViewportToolBar::GetViewModeMenuLabel() const
 {
diff --git a/Engine/Plugins/Experimental/ChaosClothAssetEditor/Source/ChaosClothAssetEditor/Private/ChaosClothAsset/SClothEditorRestSpaceViewportToolBar.h b/Engine/Plugins/Experimental/ChaosClothAssetEditor/Source/ChaosClothAssetEditor/Private/ChaosClothAsset/SClothEditorRestSpaceViewportToolBar.h
--- a/Engine/Plugins/Experimental/ChaosClothAssetEditor/Source/ChaosClothAssetEditor/Private/ChaosClothAsset/SClothEditorRestSpaceViewportToolBar.h
+++ b/Engine/Plugins/Experimental/ChaosClothAssetEditor/Source/ChaosClothAssetEditor/Private/ChaosClothAsset/SClothEditorRestSpaceViewportToolBar.h
@@ -36,6 +36,8 @@ private:
 	TSharedRef<SWidget> MakeDisplayToolBar(const TSharedPtr<FExtender> InExtenders);
 	TSharedRef<SWidget> MakeToolBar(const TSharedPtr<FExtender> InExtenders);
 
+	FText GetDisplayString() const;
+
 	FText GetViewModeMenuLabel() const;
 	const FSlateBrush* GetViewModeMenuLabelIcon() const;
 	TSharedRef<SWidget> GenerateViewModeMenuContent() const;
