Chaos Cloth: Fixes to SKM import and merging cloth selections.

SKM Import: allow specifying the LOD to import, fix a bug on skin weight bone indexing, don't overwrite skinning data.
Merging Selections: offset the merged indices when merging, handle not having an existing selection schema on the first cloth collection.

diff --git a/Engine/Plugins/Experimental/ChaosClothAsset/Source/ChaosClothAsset/Private/ChaosClothAsset/CollectionClothSelectionFacade.cpp b/Engine/Plugins/Experimental/ChaosClothAsset/Source/ChaosClothAsset/Private/ChaosClothAsset/CollectionClothSelectionFacade.cpp
--- a/Engine/Plugins/Experimental/ChaosClothAsset/Source/ChaosClothAsset/Private/ChaosClothAsset/CollectionClothSelectionFacade.cpp
+++ b/Engine/Plugins/Experimental/ChaosClothAsset/Source/ChaosClothAsset/Private/ChaosClothAsset/CollectionClothSelectionFacade.cpp
@@ -86,7 +86,7 @@ namespace UE::Chaos::ClothAsset
 		}
 	}
 
-	void FCollectionClothSelectionFacade::Append(const FCollectionClothSelectionConstFacade& Other, bool bOverwriteExistingIfMismatched)
+	void FCollectionClothSelectionFacade::AppendWithOffsets(const FCollectionClothSelectionConstFacade& Other, bool bOverwriteExistingIfMismatched, const TMap<FName, int32>& GroupOffsets)
 	{
 		if (Other.IsValid())
 		{
@@ -95,12 +95,31 @@ namespace UE::Chaos::ClothAsset
 			for (int32 InSelectionIndex = 0; InSelectionIndex < NumInSelections; ++InSelectionIndex)
 			{
 				const FName& SelectionName = InSelectionNames[InSelectionIndex];
+				const FName OtherGroupName = Other.GetSelectionGroup(SelectionName);
+				const TSet<int32>& OtherSet = Other.GetSelectionSet(SelectionName);
+				TSet<int32> OffsetSetIfNeeded;
+				const TSet<int32>* OtherSetWithOffset = &OtherSet;
+				if (const int32* const Offset = GroupOffsets.Find(OtherGroupName))
+				{
+					if (*Offset != 0)
+					{
+						OffsetSetIfNeeded.Reserve(OtherSet.Num());
+						for (const int32 OrigIndex : OtherSet)
+						{
+							OffsetSetIfNeeded.Emplace(OrigIndex + *Offset);
+						}
+						OtherSetWithOffset = &OffsetSetIfNeeded;
+					}
+				}	
+				
+				check(OtherSetWithOffset);
 				if (HasSelection(SelectionName))
 				{
-					if (GetSelectionGroup(SelectionName) == Other.GetSelectionGroup(SelectionName))
+					if (GetSelectionGroup(SelectionName) == OtherGroupName)
 					{
 						TSet<int32>& UnionedSet = GetSelectionSet(SelectionName);
-						UnionedSet.Append(Other.GetSelectionSet(SelectionName));
+
+						UnionedSet.Append(*OtherSetWithOffset);
 						continue;
 					}
 					if (!bOverwriteExistingIfMismatched)
@@ -109,7 +128,12 @@ namespace UE::Chaos::ClothAsset
 					}
 				}
 
-				FindOrAddSelectionSet(SelectionName, Other.GetSelectionGroup(SelectionName)) = Other.GetSelectionSet(SelectionName);
+				if (!IsValid())
+				{
+					DefineSchema();
+				}
+
+				FindOrAddSelectionSet(SelectionName, OtherGroupName) = *OtherSetWithOffset;
 			}
 		}
 	}
diff --git a/Engine/Plugins/Experimental/ChaosClothAsset/Source/ChaosClothAsset/Public/ChaosClothAsset/CollectionClothSelectionFacade.h b/Engine/Plugins/Experimental/ChaosClothAsset/Source/ChaosClothAsset/Public/ChaosClothAsset/CollectionClothSelectionFacade.h
--- a/Engine/Plugins/Experimental/ChaosClothAsset/Source/ChaosClothAsset/Public/ChaosClothAsset/CollectionClothSelectionFacade.h
+++ b/Engine/Plugins/Experimental/ChaosClothAsset/Source/ChaosClothAsset/Public/ChaosClothAsset/CollectionClothSelectionFacade.h
@@ -86,7 +86,20 @@ namespace UE::Chaos::ClothAsset
 		 * @param bOverwriteExistingIfMismatched: If true, overwrite existing mismatched sets.
 		 *   If false, keep the existing set.
 		 */
-		CHAOSCLOTHASSET_API void Append(const FCollectionClothSelectionConstFacade& Other, bool bOverwriteExistingIfMismatched);
+		void Append(const FCollectionClothSelectionConstFacade& Other, bool bOverwriteExistingIfMismatched)
+		{
+			AppendWithOffsets(Other, bOverwriteExistingIfMismatched, TMap<FName, int32>());
+		}
+
+		/**
+		 * Append all sets from an existing collection to this collection. Apply an offset to Other's data if its GroupName is in the GroupOffsets table.
+		 * Matching sets (i.e., same name and type) will be unioned.
+		 * Mismatching sets (same name, different type) will be handled according to bOverwriteExistingIfMismatched
+		 *
+		 * @param bOverwriteExistingIfMismatched: If true, overwrite existing mismatched sets.
+		 *   If false, keep the existing set.
+		 */
+		CHAOSCLOTHASSET_API void AppendWithOffsets(const FCollectionClothSelectionConstFacade& Other, bool bOverwriteExistingIfMismatched, const TMap<FName, int32>& GroupOffsets);
 
 		/**
 		 * Find, or add if it doesn't already exist, a selection for the specified group with the given name.
diff --git a/Engine/Plugins/Experimental/ChaosClothAssetEditor/Source/ChaosClothAssetDataflowNodes/Private/ChaosClothAsset/ClothDataflowTools.cpp b/Engine/Plugins/Experimental/ChaosClothAssetEditor/Source/ChaosClothAssetDataflowNodes/Private/ChaosClothAsset/ClothDataflowTools.cpp
--- a/Engine/Plugins/Experimental/ChaosClothAssetEditor/Source/ChaosClothAssetDataflowNodes/Private/ChaosClothAsset/ClothDataflowTools.cpp
+++ b/Engine/Plugins/Experimental/ChaosClothAssetEditor/Source/ChaosClothAssetDataflowNodes/Private/ChaosClothAsset/ClothDataflowTools.cpp
@@ -308,7 +308,7 @@ namespace UE::Chaos::ClothAsset
 				BoneWeightArray.SetNumUninitialized(NumInfluences);
 				for (int32 Idx = 0; Idx < NumInfluences; ++Idx)
 				{
-					BoneWeightArray[Idx] = FBoneWeight(SoftVertex.InfluenceBones[Idx], (float)SoftVertex.InfluenceWeights[Idx] * UE::AnimationCore::InvMaxRawBoneWeightFloat);
+					BoneWeightArray[Idx] = FBoneWeight(SourceSection.BoneMap[SoftVertex.InfluenceBones[Idx]], (float)SoftVertex.InfluenceWeights[Idx] * UE::AnimationCore::InvMaxRawBoneWeightFloat);
 				}
 				return FBoneWeights::Create(BoneWeightArray, FBoneWeightsSettings());
 			}
diff --git a/Engine/Plugins/Experimental/ChaosClothAssetEditor/Source/ChaosClothAssetDataflowNodes/Private/ChaosClothAsset/MergeClothCollectionsNode.cpp b/Engine/Plugins/Experimental/ChaosClothAssetEditor/Source/ChaosClothAssetDataflowNodes/Private/ChaosClothAsset/MergeClothCollectionsNode.cpp
--- a/Engine/Plugins/Experimental/ChaosClothAssetEditor/Source/ChaosClothAssetDataflowNodes/Private/ChaosClothAsset/MergeClothCollectionsNode.cpp
+++ b/Engine/Plugins/Experimental/ChaosClothAssetEditor/Source/ChaosClothAssetDataflowNodes/Private/ChaosClothAsset/MergeClothCollectionsNode.cpp
@@ -50,12 +50,39 @@ void FChaosClothAssetMergeClothCollectionsNode::Evaluate(Dataflow::FContext& Con
 		{
 			FManagedArrayCollection OtherCollection = GetValue<FManagedArrayCollection>(Context, Collections[InputIndex]);  // Can't use a const reference here sadly since the facade needs a SharedRef to be created
 			const TSharedRef<const FManagedArrayCollection> OtherClothCollection = MakeShared<const FManagedArrayCollection>(MoveTemp(OtherCollection));
+
+			// Selections need to update with offsets. Gather offsets before appending cloth data.
+			const FCollectionClothSelectionConstFacade OtherSelectionFacade(OtherClothCollection);
+			TMap<FName, int32> GroupNameOffsets;
+			if (OtherSelectionFacade.IsValid())
+			{
+				const TArray<FName> SelectionNames = OtherSelectionFacade.GetNames();
+				for (const FName& SelectionName : SelectionNames)
+				{
+					const FName GroupName = OtherSelectionFacade.GetSelectionGroup(SelectionName);
+					if (!GroupNameOffsets.Find(GroupName))
+					{
+						GroupNameOffsets.Add(GroupName) = ClothCollection->NumElements(GroupName); // NumElements will return zero if the group doesn't exist.
+					}
+				}
+			}
+
+			// Append cloth
 			const FCollectionClothConstFacade OtherClothFacade(OtherClothCollection);
 			if (OtherClothFacade.IsValid())
 			{
 				ClothFacade.Append(OtherClothFacade);
 				bAreAnyValid = true;
 			}
+
+			// Append selections (with offsets)
+			if (OtherSelectionFacade.IsValid())
+			{
+				constexpr bool bUpdateExistingSelections = true; // Want last one wins.
+				SelectionFacade.AppendWithOffsets(OtherSelectionFacade, bUpdateExistingSelections, GroupNameOffsets);
+				bAreAnyValid = true;
+			}
+
 			// Copy properties
 			const FCollectionPropertyConstFacade OtherPropertyFacade(OtherClothCollection);
 			if (OtherPropertyFacade.IsValid())
@@ -64,14 +91,6 @@ void FChaosClothAssetMergeClothCollectionsNode::Evaluate(Dataflow::FContext& Con
 				PropertyFacade.Append(OtherClothCollection.ToSharedPtr(), bUpdateExistingProperties);
 				bAreAnyValid = true;
 			}
-			// Copy selections
-			const FCollectionClothSelectionConstFacade OtherSelectionFacade(OtherClothCollection);
-			if (OtherSelectionFacade.IsValid())
-			{
-				constexpr bool bUpdateExistingSelections = true; // Want last one wins.
-				SelectionFacade.Append(OtherSelectionFacade, bUpdateExistingSelections);
-				bAreAnyValid = true;
-			}
 		}
 
 		// Set the output
diff --git a/Engine/Plugins/Experimental/ChaosClothAssetEditor/Source/ChaosClothAssetDataflowNodes/Private/ChaosClothAsset/SkeletalMeshImportNode.cpp b/Engine/Plugins/Experimental/ChaosClothAssetEditor/Source/ChaosClothAssetDataflowNodes/Private/ChaosClothAsset/SkeletalMeshImportNode.cpp
--- a/Engine/Plugins/Experimental/ChaosClothAssetEditor/Source/ChaosClothAssetDataflowNodes/Private/ChaosClothAsset/SkeletalMeshImportNode.cpp
+++ b/Engine/Plugins/Experimental/ChaosClothAssetEditor/Source/ChaosClothAssetDataflowNodes/Private/ChaosClothAsset/SkeletalMeshImportNode.cpp
@@ -84,9 +84,6 @@ void FChaosClothAssetSkeletalMeshImportNode::Evaluate(Dataflow::FContext& Contex
 			}
 
 			ClothFacade.SetSkeletalMeshPathName(SkeletalMesh->GetPathName());
-
-			// Must be bound to root bone by default
-			FClothGeometryTools::BindMeshToRootBone(ClothCollection, bImportSimMesh, bImportRenderMesh);
 		}
 		SetValue(Context, MoveTemp(*ClothCollection), &Collection);
 	}
diff --git a/Engine/Plugins/Experimental/ChaosClothAssetEditor/Source/ChaosClothAssetDataflowNodes/Public/ChaosClothAsset/SkeletalMeshImportNode.h b/Engine/Plugins/Experimental/ChaosClothAssetEditor/Source/ChaosClothAssetDataflowNodes/Public/ChaosClothAsset/SkeletalMeshImportNode.h
--- a/Engine/Plugins/Experimental/ChaosClothAssetEditor/Source/ChaosClothAssetDataflowNodes/Public/ChaosClothAsset/SkeletalMeshImportNode.h
+++ b/Engine/Plugins/Experimental/ChaosClothAssetEditor/Source/ChaosClothAssetDataflowNodes/Public/ChaosClothAsset/SkeletalMeshImportNode.h
@@ -24,7 +24,7 @@ public:
 	TObjectPtr<const USkeletalMesh> SkeletalMesh;
 
 	/** The skeletal mesh LOD to import. */
-	UPROPERTY(VisibleAnywhere, Category = "Skeletal Mesh Import", Meta = (ClampMin = "0", DisplayName = "LOD Index"))
+	UPROPERTY(EditAnywhere, Category = "Skeletal Mesh Import", Meta = (ClampMin = "0", DisplayName = "LOD Index"))
 	int32 LODIndex = 0;
 
 	/** The skeletal mesh LOD section to import. */
