Fix duplicate LOCTEXT macro CIS issue

diff --git a/Engine/Plugins/Experimental/ChaosClothAssetEditor/Source/ChaosClothAssetEditorTools/Private/ChaosClothAsset/ClothWeightMapPaintTool.cpp b/Engine/Plugins/Experimental/ChaosClothAssetEditor/Source/ChaosClothAssetEditorTools/Private/ChaosClothAsset/ClothWeightMapPaintTool.cpp
--- a/Engine/Plugins/Experimental/ChaosClothAssetEditor/Source/ChaosClothAssetEditorTools/Private/ChaosClothAsset/ClothWeightMapPaintTool.cpp
+++ b/Engine/Plugins/Experimental/ChaosClothAssetEditor/Source/ChaosClothAssetEditorTools/Private/ChaosClothAsset/ClothWeightMapPaintTool.cpp
@@ -2205,11 +2205,11 @@ void UClothEditorWeightMapPaintTool::UpdateBrushType(EClothEditorWeightMapPaintB
 	FText BaseMessage;
 	if (BrushType == EClothEditorWeightMapPaintBrushType::Paint)
 	{
-		BaseMessage = LOCTEXT("OnStartTool", "Hold Shift to Erase. Use [/] and S/D keys to change brush size (+Shift to small-step). W/E to change Value (+Shift to small-step). Shift-G to get current Value under cursor. Q/A to cycle through brush modes.");
+		BaseMessage = LOCTEXT("OnStartPaintMode", "Hold Shift to Erase. Use [/] and S/D keys to change brush size (+Shift to small-step). W/E to change Value (+Shift to small-step). Shift-G to get current Value under cursor. Q/A to cycle through brush modes.");
 	}
 	else if (BrushType == EClothEditorWeightMapPaintBrushType::Smooth)
 	{
-		BaseMessage = LOCTEXT("OnStartTool", "Hold Shift to Erase. Use [/] and S/D keys to change brush size (+Shift to small-step). Q/A to cycle through brush modes.");
+		BaseMessage = LOCTEXT("OnStartBrushMode", "Hold Shift to Erase. Use [/] and S/D keys to change brush size (+Shift to small-step). Q/A to cycle through brush modes.");
 	}
 
 	FTextBuilder Builder;
