CIS - Fixed the namespace and key for the Button_ReimportAsset_Tooltip localization text.

#changelist validated
#virtualized

diff --git a/Engine/Plugins/Experimental/ChaosClothAssetEditor/Source/ChaosClothAssetDataflowNodes/Private/ChaosClothAsset/TerminalNodeRefreshAssetCustomization.cpp b/Engine/Plugins/Experimental/ChaosClothAssetEditor/Source/ChaosClothAssetDataflowNodes/Private/ChaosClothAsset/TerminalNodeRefreshAssetCustomization.cpp
--- a/Engine/Plugins/Experimental/ChaosClothAssetEditor/Source/ChaosClothAssetDataflowNodes/Private/ChaosClothAsset/TerminalNodeRefreshAssetCustomization.cpp
+++ b/Engine/Plugins/Experimental/ChaosClothAssetEditor/Source/ChaosClothAssetDataflowNodes/Private/ChaosClothAsset/TerminalNodeRefreshAssetCustomization.cpp
@@ -12,7 +12,7 @@
 #include "EditorDirectories.h"
 #include "PropertyHandle.h"
 
-#define LOCTEXT_NAMESPACE "ImportFilePathCustomization"
+#define LOCTEXT_NAMESPACE "TerminalNodeRefreshAssetCustomization"
 
 namespace UE::Chaos::ClothAsset
 {
@@ -37,7 +37,7 @@ namespace UE::Chaos::ClothAsset
 			[
 				SNew(SButton)
 				.ButtonStyle(FAppStyle::Get(), "RoundButton")
-				.ToolTipText(LOCTEXT("Button_ReimportAsset_Tooltip", "Refresh asset"))
+				.ToolTipText(LOCTEXT("Button_RefreshAsset_Tooltip", "Refresh asset"))
 				.OnClicked(this, &FTerminalNodeRefreshAssetCustomization::OnClicked)
 				[
 					SNew(SHorizontalBox)
