Unshelved from pending changelist '31088053':

USD Fabric used in config nodes

diff --git a/Engine/Plugins/Experimental/ChaosClothAsset/Source/ChaosClothAsset/Private/ChaosClothAsset/ClothCollection.cpp b/Engine/Plugins/Experimental/ChaosClothAsset/Source/ChaosClothAsset/Private/ChaosClothAsset/ClothCollection.cpp
--- a/Engine/Plugins/Experimental/ChaosClothAsset/Source/ChaosClothAsset/Private/ChaosClothAsset/ClothCollection.cpp
+++ b/Engine/Plugins/Experimental/ChaosClothAsset/Source/ChaosClothAsset/Private/ChaosClothAsset/ClothCollection.cpp
@@ -15,6 +15,27 @@ namespace UE::Chaos::ClothAsset::Private
 		SkeletalMeshPathNameAttribute
 	};
 
+	// Fabrics Group
+	static const FName FabricBendingStiffnessAttribute(TEXT("FabricBendingStiffness"));
+	static const FName FabricBucklingStiffnessAttribute(TEXT("FabricBucklingStiffness"));
+	static const FName FabricStretchStiffnessAttribute(TEXT("FabricStretchStiffness"));
+	static const FName FabricBucklingRatioAttribute(TEXT("FabricBucklingRatio"));
+	static const FName FabricClothDensityAttribute(TEXT("FabricClothDensity"));
+	static const FName FabricClothFrictionAttribute(TEXT("FabricClothFriction"));
+	static const FName FabricClothThicknessAttribute(TEXT("FabricClothThickness"));
+	static const FName FabricClothDampingAttribute(TEXT("FabricClothDamping"));
+	static const TArray<FName> SimFabricGroupAttributes =
+	{
+		FabricBendingStiffnessAttribute,
+		FabricBucklingStiffnessAttribute,
+		FabricStretchStiffnessAttribute,
+		FabricBucklingRatioAttribute,
+		FabricClothDensityAttribute,
+		FabricClothFrictionAttribute,
+		FabricClothThicknessAttribute,
+		FabricClothDampingAttribute
+	};
+
 	// Seam Group
 	static const FName SeamStitchStartAttribute(TEXT("SeamStitchStart"));
 	static const FName SeamStitchEndAttribute(TEXT("SeamStitchEnd"));
@@ -38,12 +59,14 @@ namespace UE::Chaos::ClothAsset::Private
 	static const FName SimVertices2DEndAttribute(TEXT("SimVertices2DEnd"));
 	static const FName SimFacesStartAttribute(TEXT("SimFacesStart"));
 	static const FName SimFacesEndAttribute(TEXT("SimFacesEnd"));
+	static const FName SimPatternFabricAttribute(TEXT("SimPatternFabric"));
 	static const TArray<FName> SimPatternsGroupAttributes =
 	{ 
 		SimVertices2DStartAttribute,
 		SimVertices2DEndAttribute,
 		SimFacesStartAttribute,
 		SimFacesEndAttribute,
+		SimPatternFabricAttribute
 	};
 
 	// RenderPatterns Group
@@ -140,7 +163,8 @@ namespace UE::Chaos::ClothAsset::Private
 		{ ClothCollectionGroup::SimVertices2D, SimVertices2DGroupAttributes },
 		{ ClothCollectionGroup::SimVertices3D, SimVertices3DGroupAttributes },
 		{ ClothCollectionGroup::RenderFaces, RenderFacesGroupAttributes },
-		{ ClothCollectionGroup::RenderVertices, RenderVerticesGroupAttributes }
+		{ ClothCollectionGroup::RenderVertices, RenderVerticesGroupAttributes },
+		{ ClothCollectionGroup::Fabrics, SimFabricGroupAttributes }
 	};
 }  // End namespace UE::Chaos::ClothAsset::Private
 
@@ -168,6 +192,7 @@ namespace UE::Chaos::ClothAsset
 		SimVertices2DEnd = ManagedArrayCollection->FindAttribute<int32>(SimVertices2DEndAttribute, ClothCollectionGroup::SimPatterns);
 		SimFacesStart = ManagedArrayCollection->FindAttribute<int32>(SimFacesStartAttribute, ClothCollectionGroup::SimPatterns);
 		SimFacesEnd = ManagedArrayCollection->FindAttribute<int32>(SimFacesEndAttribute, ClothCollectionGroup::SimPatterns);
+		SimPatternFabric = ManagedArrayCollection->FindAttribute<int32>(SimPatternFabricAttribute, ClothCollectionGroup::SimPatterns);
 
 		// Render Patterns Group
 		RenderVerticesStart = ManagedArrayCollection->FindAttribute<int32>(RenderVerticesStartAttribute, ClothCollectionGroup::RenderPatterns);
@@ -176,6 +201,16 @@ namespace UE::Chaos::ClothAsset
 		RenderFacesEnd = ManagedArrayCollection->FindAttribute<int32>(RenderFacesEndAttribute, ClothCollectionGroup::RenderPatterns);
 		RenderMaterialPathName = ManagedArrayCollection->FindAttribute<FString>(RenderMaterialPathNameAttribute, ClothCollectionGroup::RenderPatterns);
 
+		//~ Fabric Group
+		FabricBendingStiffness = ManagedArrayCollection->FindAttribute<FVector3f>(FabricBendingStiffnessAttribute, ClothCollectionGroup::Fabrics);
+		FabricBucklingStiffness = ManagedArrayCollection->FindAttribute<FVector3f>(FabricBucklingStiffnessAttribute, ClothCollectionGroup::Fabrics);
+		FabricStretchStiffness = ManagedArrayCollection->FindAttribute<FVector3f>(FabricStretchStiffnessAttribute, ClothCollectionGroup::Fabrics);
+		FabricBucklingRatio = ManagedArrayCollection->FindAttribute<float>(FabricBucklingRatioAttribute, ClothCollectionGroup::Fabrics);
+		FabricClothDensity = ManagedArrayCollection->FindAttribute<float>(FabricClothDensityAttribute, ClothCollectionGroup::Fabrics);
+		FabricClothFriction = ManagedArrayCollection->FindAttribute<float>(FabricClothFrictionAttribute, ClothCollectionGroup::Fabrics);
+		FabricClothThickness = ManagedArrayCollection->FindAttribute<float>(FabricClothThicknessAttribute, ClothCollectionGroup::Fabrics);
+		FabricClothDamping = ManagedArrayCollection->FindAttribute<float>(FabricClothDampingAttribute, ClothCollectionGroup::Fabrics);
+		
 		// Sim Faces Group
 		SimIndices2D = ManagedArrayCollection->FindAttribute<FIntVector3>(SimIndices2DAttribute, ClothCollectionGroup::SimFaces);
 		SimIndices3D = ManagedArrayCollection->FindAttribute<FIntVector3>(SimIndices3DAttribute, ClothCollectionGroup::SimFaces);
@@ -228,6 +263,7 @@ namespace UE::Chaos::ClothAsset
 			SimVertices2DEnd &&
 			SimFacesStart &&
 			SimFacesEnd &&
+			SimPatternFabric &&
 
 			// Render Patterns Group
 			RenderVerticesStart &&
@@ -240,6 +276,16 @@ namespace UE::Chaos::ClothAsset
 			SimIndices2D &&
 			SimIndices3D &&
 
+			// Fabrics Group
+			FabricBendingStiffness &&
+			FabricBucklingStiffness &&
+			FabricStretchStiffness &&
+			FabricBucklingRatio &&
+			FabricClothDensity &&
+			FabricClothFriction &&
+			FabricClothThickness &&
+			FabricClothDamping &&
+
 			// Sim Vertices 2D Group
 			SimPosition2D  &&
 			SimVertex3DLookup &&
@@ -279,6 +325,7 @@ namespace UE::Chaos::ClothAsset
 		FManagedArrayCollection::FConstructionParameters RenderFacesDependency(ClothCollectionGroup::RenderFaces, bSaved, bAllowCircularDependency);
 		FManagedArrayCollection::FConstructionParameters RenderVerticesDependency(ClothCollectionGroup::RenderVertices, bSaved, bAllowCircularDependency);
 		FManagedArrayCollection::FConstructionParameters SimFacesDependency(ClothCollectionGroup::SimFaces, bSaved, bAllowCircularDependency);
+		FManagedArrayCollection::FConstructionParameters SimFabricsDependency(ClothCollectionGroup::Fabrics, bSaved, bAllowCircularDependency);
 		FManagedArrayCollection::FConstructionParameters SimVertices2DDependency(ClothCollectionGroup::SimVertices2D, bSaved, bAllowCircularDependency);
 		FManagedArrayCollection::FConstructionParameters SimVertices3DDependency(ClothCollectionGroup::SimVertices3D, bSaved, bAllowCircularDependency);  // Any attribute with this dependency must handle welding and splitting in FCollectionClothSeamFacade
 
@@ -299,6 +346,7 @@ namespace UE::Chaos::ClothAsset
 		SimVertices2DEnd = &ManagedArrayCollection->AddAttribute<int32>(SimVertices2DEndAttribute, ClothCollectionGroup::SimPatterns, SimVertices2DDependency);
 		SimFacesStart = &ManagedArrayCollection->AddAttribute<int32>(SimFacesStartAttribute, ClothCollectionGroup::SimPatterns, SimFacesDependency);
 		SimFacesEnd = &ManagedArrayCollection->AddAttribute<int32>(SimFacesEndAttribute, ClothCollectionGroup::SimPatterns, SimFacesDependency);
+		SimPatternFabric = &ManagedArrayCollection->AddAttribute<int32>(SimPatternFabricAttribute, ClothCollectionGroup::SimPatterns, SimFabricsDependency);
 
 		// Render Patterns Group
 		RenderVerticesStart = &ManagedArrayCollection->AddAttribute<int32>(RenderVerticesStartAttribute, ClothCollectionGroup::RenderPatterns, RenderVerticesDependency);
@@ -307,6 +355,16 @@ namespace UE::Chaos::ClothAsset
 		RenderFacesEnd = &ManagedArrayCollection->AddAttribute<int32>(RenderFacesEndAttribute, ClothCollectionGroup::RenderPatterns, RenderFacesDependency);
 		RenderMaterialPathName = &ManagedArrayCollection->AddAttribute<FString>(RenderMaterialPathNameAttribute, ClothCollectionGroup::RenderPatterns);
 
+		//~ Fabric Group
+		FabricBendingStiffness = &ManagedArrayCollection->AddAttribute<FVector3f>(FabricBendingStiffnessAttribute, ClothCollectionGroup::Fabrics);
+		FabricBucklingStiffness = &ManagedArrayCollection->AddAttribute<FVector3f>(FabricBucklingStiffnessAttribute, ClothCollectionGroup::Fabrics);
+		FabricStretchStiffness = &ManagedArrayCollection->AddAttribute<FVector3f>(FabricStretchStiffnessAttribute, ClothCollectionGroup::Fabrics);
+		FabricBucklingRatio = &ManagedArrayCollection->AddAttribute<float>(FabricBucklingRatioAttribute, ClothCollectionGroup::Fabrics);
+		FabricClothDensity = &ManagedArrayCollection->AddAttribute<float>(FabricClothDensityAttribute, ClothCollectionGroup::Fabrics);
+		FabricClothFriction = &ManagedArrayCollection->AddAttribute<float>(FabricClothFrictionAttribute, ClothCollectionGroup::Fabrics);
+		FabricClothThickness = &ManagedArrayCollection->AddAttribute<float>(FabricClothThicknessAttribute, ClothCollectionGroup::Fabrics);
+		FabricClothDamping = &ManagedArrayCollection->AddAttribute<float>(FabricClothDampingAttribute, ClothCollectionGroup::Fabrics);
+		
 		// Sim Faces Group
 		SimIndices2D = &ManagedArrayCollection->AddAttribute<FIntVector3>(SimIndices2DAttribute, ClothCollectionGroup::SimFaces, SimVertices2DDependency);
 		SimIndices3D = &ManagedArrayCollection->AddAttribute<FIntVector3>(SimIndices3DAttribute, ClothCollectionGroup::SimFaces, SimVertices3DDependency);
diff --git a/Engine/Plugins/Experimental/ChaosClothAsset/Source/ChaosClothAsset/Private/ChaosClothAsset/ClothCollection.h b/Engine/Plugins/Experimental/ChaosClothAsset/Source/ChaosClothAsset/Private/ChaosClothAsset/ClothCollection.h
--- a/Engine/Plugins/Experimental/ChaosClothAsset/Source/ChaosClothAsset/Private/ChaosClothAsset/ClothCollection.h
+++ b/Engine/Plugins/Experimental/ChaosClothAsset/Source/ChaosClothAsset/Private/ChaosClothAsset/ClothCollection.h
@@ -164,11 +164,22 @@ namespace UE::Chaos::ClothAsset
 		const TManagedArray<FIntVector2>* GetSeamStitch2DEndIndices() const { return SeamStitch2DEndIndices; }
 		const TManagedArray<int32>* GetSeamStitch3DIndex() const { return SeamStitch3DIndex; }
 
+		//~ Fabric Group
+		const TManagedArray<FVector3f>* GetFabricBendingStiffness() const { return FabricBendingStiffness; }
+		const TManagedArray<FVector3f>* GetFabricBucklingStiffness() const { return FabricBucklingStiffness; }
+		const TManagedArray<FVector3f>* GetFabricStretchStiffness() const { return FabricStretchStiffness; }
+		const TManagedArray<float>* GetFabricBucklingRatio() const { return FabricBucklingRatio; }
+		const TManagedArray<float>* GetFabricDensity() const { return FabricClothDensity; }
+		const TManagedArray<float>* GetFabricFriction() const { return FabricClothFriction; }
+		const TManagedArray<float>* GetFabricThickness() const { return FabricClothThickness; }
+		const TManagedArray<float>* GetFabricDamping() const { return FabricClothDamping; }
+		
 		//~ Sim Patterns Group
 		const TManagedArray<int32>* GetSimVertices2DStart() const { return SimVertices2DStart; }
 		const TManagedArray<int32>* GetSimVertices2DEnd() const { return SimVertices2DEnd; }
 		const TManagedArray<int32>* GetSimFacesStart() const { return SimFacesStart; }
 		const TManagedArray<int32>* GetSimFacesEnd() const { return SimFacesEnd; }
+		const TManagedArray<int32>* GetSimPatternFabric() const { return SimPatternFabric; }
 
 		//~ Render Patterns Group
 		const TManagedArray<int32>* GetRenderVerticesStart() const { return RenderVerticesStart; }
@@ -212,6 +223,16 @@ namespace UE::Chaos::ClothAsset
 		TManagedArray<FString>* GetPhysicsAssetPathName(){ return PhysicsAssetPathName; }
 		TManagedArray<FString>* GetSkeletalMeshPathName() { return SkeletalMeshPathName; }
 
+		//~ Fabric Group
+		TManagedArray<FVector3f>* GetFabricBendingStiffness() { return FabricBendingStiffness; }
+		TManagedArray<FVector3f>* GetFabricBucklingStiffness() { return FabricBucklingStiffness; }
+		TManagedArray<FVector3f>* GetFabricStretchStiffness() { return FabricStretchStiffness; }
+		TManagedArray<float>* GetFabricBucklingRatio() { return FabricBucklingRatio; }
+		TManagedArray<float>* GetFabricDensity() { return FabricClothDensity; }
+		TManagedArray<float>* GetFabricFriction() { return FabricClothFriction; }
+		TManagedArray<float>* GetFabricThickness() { return FabricClothThickness; }
+		TManagedArray<float>* GetFabricDamping() { return FabricClothDamping; }
+		
 		//~ Seam Group
 		TManagedArray<int32>* GetSeamStitchStart() { return SeamStitchStart; }
 		TManagedArray<int32>* GetSeamStitchEnd() { return SeamStitchEnd; }
@@ -225,6 +246,7 @@ namespace UE::Chaos::ClothAsset
 		TManagedArray<int32>* GetSimVertices2DEnd() { return SimVertices2DEnd; }
 		TManagedArray<int32>* GetSimFacesStart() { return SimFacesStart; }
 		TManagedArray<int32>* GetSimFacesEnd() { return SimFacesEnd; }
+		TManagedArray<int32>* GetSimPatternFabric() { return SimPatternFabric; }
 
 		//~ Render Patterns Group
 		TManagedArray<int32>* GetRenderVerticesStart() { return RenderVerticesStart; }
@@ -272,6 +294,16 @@ namespace UE::Chaos::ClothAsset
 		TManagedArray<FString>* PhysicsAssetPathName;
 		TManagedArray<FString>* SkeletalMeshPathName;
 
+		//! Fabrics Group
+		TManagedArray<FVector3f>* FabricBendingStiffness;
+		TManagedArray<FVector3f>* FabricBucklingStiffness;
+		TManagedArray<FVector3f>* FabricStretchStiffness;
+		TManagedArray<float>* FabricBucklingRatio;
+		TManagedArray<float>* FabricClothDensity;
+		TManagedArray<float>* FabricClothFriction;
+		TManagedArray<float>* FabricClothThickness;
+		TManagedArray<float>* FabricClothDamping;
+
 		//~ Seam Group
 		TManagedArray<int32>* SeamStitchStart;
 		TManagedArray<int32>* SeamStitchEnd;
@@ -285,6 +317,7 @@ namespace UE::Chaos::ClothAsset
 		TManagedArray<int32>* SimVertices2DEnd;
 		TManagedArray<int32>* SimFacesStart;
 		TManagedArray<int32>* SimFacesEnd;
+		TManagedArray<int32>* SimPatternFabric;
 
 		//~ Render Patterns Group
 		TManagedArray<int32>* RenderVerticesStart;
diff --git a/Engine/Plugins/Experimental/ChaosClothAsset/Source/ChaosClothAsset/Private/ChaosClothAsset/CollectionClothFabricFacade.cpp b/Engine/Plugins/Experimental/ChaosClothAsset/Source/ChaosClothAsset/Private/ChaosClothAsset/CollectionClothFabricFacade.cpp
new file mode 100644
--- /dev/null
+++ b/Engine/Plugins/Experimental/ChaosClothAsset/Source/ChaosClothAsset/Private/ChaosClothAsset/CollectionClothFabricFacade.cpp
@@ -0,0 +1,118 @@
+// Copyright Epic Games, Inc. All Rights Reserved.
+
+#include "ChaosClothAsset/CollectionClothFabricFacade.h"
+#include "ChaosClothAsset/CollectionClothFacade.h"
+#include "ChaosClothAsset/ClothCollection.h"
+#include "ChaosClothAsset/ClothCollectionGroup.h"
+
+namespace UE::Chaos::ClothAsset
+{
+	FCollectionClothFabricConstFacade::FAnisotropicData::FAnisotropicData(const float WeftValue, const float WarpValue, const float BiasValue) :
+		Weft(WeftValue), Warp(WarpValue), Bias(BiasValue)
+	{}
+			
+	FCollectionClothFabricConstFacade::FAnisotropicData::FAnisotropicData(const FVector3f& VectorDatas) :
+		Weft(VectorDatas.X), Warp(VectorDatas.Y), Bias(VectorDatas.Z) 
+	{}
+
+	FCollectionClothFabricConstFacade::FAnisotropicData::FAnisotropicData(const float& FloatDatas) :
+		Weft(FloatDatas), Warp(FloatDatas), Bias(FloatDatas)
+	{}
+			
+	FVector3f FCollectionClothFabricConstFacade::FAnisotropicData::GetVectorDatas() const
+	{
+		return FVector3f(Weft, Warp, Bias);
+	}
+	
+	FCollectionClothFabricConstFacade::FAnisotropicData FCollectionClothFabricConstFacade::GetXPBDAnisoBendingStiffness() const
+	{
+		return FCollectionClothFabricConstFacade::FAnisotropicData(
+			ClothCollection->GetElements(ClothCollection->GetFabricBendingStiffness())[GetElementIndex()]);
+	}
+	
+	float FCollectionClothFabricConstFacade::GetXPBDAnisoBucklingRatio() const
+	{
+		return ClothCollection->GetElements(ClothCollection->GetFabricBucklingRatio())[GetElementIndex()];
+	}
+	
+	FCollectionClothFabricConstFacade::FAnisotropicData FCollectionClothFabricConstFacade::GetXPBDAnisoBucklingStiffness() const
+	{
+		return FCollectionClothFabricConstFacade::FAnisotropicData(
+			ClothCollection->GetElements(ClothCollection->GetFabricBucklingStiffness())[GetElementIndex()]);
+	}
+	
+	FCollectionClothFabricConstFacade::FAnisotropicData FCollectionClothFabricConstFacade::GetXPBDAnisoStretchStiffness() const
+	{
+		return FCollectionClothFabricConstFacade::FAnisotropicData(
+			ClothCollection->GetElements(ClothCollection->GetFabricStretchStiffness())[GetElementIndex()]);
+	}
+	
+	float FCollectionClothFabricConstFacade::GetDensityWeighted() const
+	{
+		return ClothCollection->GetElements(ClothCollection->GetFabricDensity())[GetElementIndex()];
+	}
+	
+	float FCollectionClothFabricConstFacade::GetCollisionThickness() const
+	{
+		return ClothCollection->GetElements(ClothCollection->GetFabricThickness())[GetElementIndex()];
+	}
+	
+	float FCollectionClothFabricConstFacade::GetXPBDAnisoDamping() const
+	{
+		return ClothCollection->GetElements(ClothCollection->GetFabricDamping())[GetElementIndex()];
+	}
+	
+	float FCollectionClothFabricConstFacade::GetFrictionCoefficient() const
+	{
+		return ClothCollection->GetElements(ClothCollection->GetFabricFriction())[GetElementIndex()];
+	}
+	
+	FCollectionClothFabricConstFacade::FCollectionClothFabricConstFacade(const TSharedRef<const FClothCollection>& ClothCollection, int32 InFabricIndex)
+		: ClothCollection(ClothCollection)
+		, FabricIndex(InFabricIndex)
+	{
+		check(ClothCollection->IsValid());
+		check(FabricIndex >= 0 && FabricIndex < ClothCollection->GetNumElements(ClothCollectionGroup::Fabrics));
+	}
+
+	void FCollectionClothFabricFacade::Initialize(const FAnisotropicData& BendingStiffness, const float BucklingRatio,
+			const FAnisotropicData& BucklingStiffness, const FAnisotropicData& StretchStiffness,
+			const float ClothDensity, const float ClothFriction, const float ClothDamping, const float ClothThickness)
+	{
+		const int32 ElementIndex = GetElementIndex();
+		check((ElementIndex >= 0) && (ElementIndex < ClothCollection->GetNumElements(ClothCollectionGroup::Fabrics)));
+		
+		const TSharedRef<class FClothCollection> ClothCol = GetClothCollection();
+		
+		ClothCol->GetElements(ClothCol->GetFabricBendingStiffness())[ElementIndex] = BendingStiffness.GetVectorDatas();
+		ClothCol->GetElements(ClothCol->GetFabricBucklingRatio())[ElementIndex] = BucklingRatio;
+		ClothCol->GetElements(ClothCol->GetFabricStretchStiffness())[ElementIndex] = StretchStiffness.GetVectorDatas();
+		ClothCol->GetElements(ClothCol->GetFabricBucklingStiffness())[ElementIndex] = BucklingStiffness.GetVectorDatas();
+		ClothCol->GetElements(ClothCol->GetFabricFriction())[ElementIndex] = ClothFriction;
+		ClothCol->GetElements(ClothCol->GetFabricDensity())[ElementIndex] = ClothDensity;
+		ClothCol->GetElements(ClothCol->GetFabricThickness())[ElementIndex] = ClothThickness;
+		ClothCol->GetElements(ClothCol->GetFabricDamping())[ElementIndex] = ClothDamping;
+	}
+	
+	void FCollectionClothFabricFacade::Initialize(const FCollectionClothFabricConstFacade& OtherFabricFacade)
+	{
+		Initialize( OtherFabricFacade.GetXPBDAnisoBendingStiffness(), OtherFabricFacade.GetXPBDAnisoBucklingRatio(), OtherFabricFacade.GetXPBDAnisoBucklingStiffness(),
+			OtherFabricFacade.GetXPBDAnisoStretchStiffness(), OtherFabricFacade.GetDensityWeighted(),
+			OtherFabricFacade.GetFrictionCoefficient(), OtherFabricFacade.GetXPBDAnisoDamping(), OtherFabricFacade.GetCollisionThickness());
+	}
+
+	FCollectionClothFabricFacade::FCollectionClothFabricFacade(const TSharedRef<FClothCollection>& ClothCollection, int32 InFabricIndex)
+		: FCollectionClothFabricConstFacade(ClothCollection, InFabricIndex)
+	{
+	}
+
+	void FCollectionClothFabricFacade::Reset()
+	{
+		SetDefaults();
+	}
+
+	void FCollectionClothFabricFacade::SetDefaults()
+	{
+		Initialize(100.0f, 0.5f, 50.0f, 100.0f, 0.35f, 0.1f, 0.01f, 0.8f);
+	}
+}  // End namespace UE::Chaos::ClothAsset
diff --git a/Engine/Plugins/Experimental/ChaosClothAsset/Source/ChaosClothAsset/Private/ChaosClothAsset/CollectionClothFacade.cpp b/Engine/Plugins/Experimental/ChaosClothAsset/Source/ChaosClothAsset/Private/ChaosClothAsset/CollectionClothFacade.cpp
--- a/Engine/Plugins/Experimental/ChaosClothAsset/Source/ChaosClothAsset/Private/ChaosClothAsset/CollectionClothFacade.cpp
+++ b/Engine/Plugins/Experimental/ChaosClothAsset/Source/ChaosClothAsset/Private/ChaosClothAsset/CollectionClothFacade.cpp
@@ -58,6 +58,7 @@ namespace UE::Chaos::ClothAsset
 		ResultHash = HashCombineFast(ResultHash, ClothCollection->GetElementsTypeHash(ClothCollection->GetSimVertices2DEnd()));
 		ResultHash = HashCombineFast(ResultHash, ClothCollection->GetElementsTypeHash(ClothCollection->GetSimFacesStart()));
 		ResultHash = HashCombineFast(ResultHash, ClothCollection->GetElementsTypeHash(ClothCollection->GetSimFacesEnd()));
+		ResultHash = HashCombineFast(ResultHash, ClothCollection->GetElementsTypeHash(ClothCollection->GetSimPatternFabric()));
 
 		//~ Render Patterns Group
 		ResultHash = HashCombineFast(ResultHash, ClothCollection->GetElementsTypeHash(ClothCollection->GetRenderVerticesStart()));
@@ -75,18 +76,28 @@ namespace UE::Chaos::ClothAsset
 		ResultHash = HashCombineFast(ResultHash, ClothCollection->GetElementsTypeHash(ClothCollection->GetSimVertex3DLookup()));
 		
 		//~ Sim Vertices 3D Group
-		ResultHash = HashCombineFast(ResultHash, ClothCollection->GetElementsTypeHash(ClothCollection->GetSimPosition3D()));
-		ResultHash = HashCombineFast(ResultHash, ClothCollection->GetElementsTypeHash(ClothCollection->GetSimNormal()));
+    	ResultHash = HashCombineFast(ResultHash, ClothCollection->GetElementsTypeHash(ClothCollection->GetSimPosition3D()));
+    	ResultHash = HashCombineFast(ResultHash, ClothCollection->GetElementsTypeHash(ClothCollection->GetSimNormal()));
 		ResultHash = HashCombineFast(ResultHash, ClothCollection->GetElementsTypeHash(ClothCollection->GetSimBoneIndices()));
 		ResultHash = HashCombineFast(ResultHash, ClothCollection->GetElementsTypeHash(ClothCollection->GetSimBoneWeights()));
 		ResultHash = HashCombineFast(ResultHash, ClothCollection->GetElementsTypeHash(ClothCollection->GetTetherKinematicIndex()));
 		ResultHash = HashCombineFast(ResultHash, ClothCollection->GetElementsTypeHash(ClothCollection->GetTetherReferenceLength()));
 		ResultHash = HashCombineFast(ResultHash, ClothCollection->GetElementsTypeHash(ClothCollection->GetSimVertex2DLookup()));
 		ResultHash = HashCombineFast(ResultHash, ClothCollection->GetElementsTypeHash(ClothCollection->GetSeamStitchLookup()));
+		
+		//~ Fabrics Group
+		ResultHash = HashCombineFast(ResultHash, ClothCollection->GetElementsTypeHash(ClothCollection->GetFabricBendingStiffness()));
+		ResultHash = HashCombineFast(ResultHash, ClothCollection->GetElementsTypeHash(ClothCollection->GetFabricBucklingRatio()));
+		ResultHash = HashCombineFast(ResultHash, ClothCollection->GetElementsTypeHash(ClothCollection->GetFabricBucklingStiffness()));
+		ResultHash = HashCombineFast(ResultHash, ClothCollection->GetElementsTypeHash(ClothCollection->GetFabricDamping()));
+		ResultHash = HashCombineFast(ResultHash, ClothCollection->GetElementsTypeHash(ClothCollection->GetFabricDensity()));
+		ResultHash = HashCombineFast(ResultHash, ClothCollection->GetElementsTypeHash(ClothCollection->GetFabricFriction()));
+		ResultHash = HashCombineFast(ResultHash, ClothCollection->GetElementsTypeHash(ClothCollection->GetFabricThickness()));
+		ResultHash = HashCombineFast(ResultHash, ClothCollection->GetElementsTypeHash(ClothCollection->GetFabricStretchStiffness()));
 
 		//~ Render Faces Group
 		ResultHash = HashCombineFast(ResultHash, ClothCollection->GetElementsTypeHash(ClothCollection->GetRenderIndices()));
-
+		
 		//~ Render Vertices Group
 		ResultHash = HashCombineFast(ResultHash, ClothCollection->GetElementsTypeHash(ClothCollection->GetRenderPosition()));
 		ResultHash = HashCombineFast(ResultHash, ClothCollection->GetElementsTypeHash(ClothCollection->GetRenderNormal()));
@@ -274,6 +285,16 @@ namespace UE::Chaos::ClothAsset
 		return FCollectionClothSeamConstFacade(ClothCollection, SeamIndex);
 	}
 
+	int32 FCollectionClothConstFacade::GetNumFabrics() const
+	{
+		return ClothCollection->GetNumElements(ClothCollectionGroup::Fabrics);
+	}
+
+	FCollectionClothFabricConstFacade FCollectionClothConstFacade::GetFabric(int32 FabricIndex) const
+	{
+		return FCollectionClothFabricConstFacade(ClothCollection, FabricIndex);
+	}
+
 	int32 FCollectionClothConstFacade::GetNumRenderVertices() const
 	{
 		return ClothCollection->GetNumElements(ClothCollectionGroup::RenderVertices);
@@ -472,16 +493,24 @@ namespace UE::Chaos::ClothAsset
 		FClothCollection::CopyArrayViewDataAndApplyOffset(GetTetherKinematicIndex().Right(OtherNumSimVertices3D), Other.GetTetherKinematicIndex(), StartNumSimVertices3D);
 		FClothCollection::CopyArrayViewData(GetTetherReferenceLength().Right(OtherNumSimVertices3D), Other.GetTetherReferenceLength());
 
+		// Fabrics Group
+		const int32 StartNumFabrics = GetNumFabrics();
+		const int32 OtherNumFabrics = Other.GetNumFabrics();
+		SetNumFabrics(StartNumFabrics + OtherNumFabrics);
+		for (int32 FabricIndex = 0; FabricIndex < OtherNumFabrics; ++FabricIndex)
+		{
+			GetFabric(FabricIndex + StartNumFabrics).Initialize(Other.GetFabric(FabricIndex));
+		}
+
 		// Sim Patterns Group
 		const int32 StartNumSimVertices2D = GetNumSimVertices2D();
 		const int32 StartNumSimPatterns = GetNumSimPatterns();
-		const int32 StartNumSimFaces = GetNumSimFaces();
 		const int32 OtherNumSimPatterns = Other.GetNumSimPatterns();
 		const int32 OtherNumSimFaces = Other.GetNumSimFaces();
 		SetNumSimPatterns(StartNumSimPatterns + OtherNumSimPatterns);
 		for (int32 PatternIndex = 0; PatternIndex < OtherNumSimPatterns; ++PatternIndex)
 		{
-			GetSimPattern(StartNumSimPatterns + PatternIndex).Initialize(Other.GetSimPattern(PatternIndex), StartNumSimVertices3D);
+			GetSimPattern(StartNumSimPatterns + PatternIndex).Initialize(Other.GetSimPattern(PatternIndex), StartNumSimVertices3D, StartNumFabrics);
 		}
 
 		// Seams Group
@@ -492,7 +521,7 @@ namespace UE::Chaos::ClothAsset
 		{
 			GetSeam(SeamIndex + StartNumSeams).Initialize(Other.GetSeam(SeamIndex), StartNumSimVertices2D, StartNumSimVertices3D);
 		}
-
+		
 		// Sim Vertices 3D Group (lookups)
 		FClothCollection::CopyArrayViewDataAndApplyOffset(
 			GetClothCollection()->GetElements(GetClothCollection()->GetSimVertex2DLookup()).Right(OtherNumSimVertices3D), 
@@ -751,6 +780,44 @@ namespace UE::Chaos::ClothAsset
 		GetClothCollection()->RemoveElements(ClothCollectionGroup::Seams, SortedDeletionList);
 	}
 
+	void FCollectionClothFacade::SetNumFabrics(int32 InNumFabrics)
+	{
+		const int32 NumFabrics = GetNumFabrics();
+
+		for (int32 FabricIndex = InNumFabrics; FabricIndex < NumFabrics; ++FabricIndex)
+		{
+			GetFabric(FabricIndex).Reset();
+		}
+
+		GetClothCollection()->SetNumElements(InNumFabrics, ClothCollectionGroup::Fabrics);
+
+		for (int32 FabricIndex = NumFabrics; FabricIndex < InNumFabrics; ++FabricIndex)
+		{
+			GetFabric(FabricIndex).SetDefaults();
+		}
+	}
+
+	int32 FCollectionClothFacade::AddFabric()
+	{
+		const int32 FabricIndex = GetNumFabrics();
+		SetNumFabrics(FabricIndex + 1);
+		return FabricIndex;
+	}
+
+	FCollectionClothFabricFacade FCollectionClothFacade::GetFabric(int32 FabricIndex)
+	{
+		return FCollectionClothFabricFacade(GetClothCollection(), FabricIndex);
+	}
+
+	void FCollectionClothFacade::RemoveFabrics(const TArray<int32>& SortedDeletionList)
+	{
+		for (const int32 FabricToRemove : SortedDeletionList)
+		{
+			GetFabric(FabricToRemove).Reset();
+		}
+		GetClothCollection()->RemoveElements(ClothCollectionGroup::Fabrics, SortedDeletionList);
+	}
+
 	//~ Render Vertices Group
 	TArrayView<FVector3f> FCollectionClothFacade::GetRenderPosition()
 	{
diff --git a/Engine/Plugins/Experimental/ChaosClothAsset/Source/ChaosClothAsset/Private/ChaosClothAsset/CollectionClothSimPatternFacade.cpp b/Engine/Plugins/Experimental/ChaosClothAsset/Source/ChaosClothAsset/Private/ChaosClothAsset/CollectionClothSimPatternFacade.cpp
--- a/Engine/Plugins/Experimental/ChaosClothAsset/Source/ChaosClothAsset/Private/ChaosClothAsset/CollectionClothSimPatternFacade.cpp
+++ b/Engine/Plugins/Experimental/ChaosClothAsset/Source/ChaosClothAsset/Private/ChaosClothAsset/CollectionClothSimPatternFacade.cpp
@@ -60,6 +60,11 @@ namespace UE::Chaos::ClothAsset
 			ClothCollection->GetSimFacesEnd(),
 			GetElementIndex());
 	}
+	
+	int32 FCollectionClothSimPatternConstFacade::GetFabricIndex() const
+	{
+		return ClothCollection->GetElements(ClothCollection->GetSimPatternFabric())[GetElementIndex()];
+	}
 
 	int32 FCollectionClothSimPatternConstFacade::GetSimFacesOffset() const
 	{
@@ -108,10 +113,12 @@ namespace UE::Chaos::ClothAsset
 	}
 
 	template<typename IndexType, typename TEnableIf<TIsIndexType<IndexType>::Value, int>::type>
-	void FCollectionClothSimPatternFacade::Initialize(const TArray<FVector2f>& Positions2D, const TArray<FVector3f>& Positions3D, const TArray<IndexType>& Indices)
+	void FCollectionClothSimPatternFacade::Initialize(const TArray<FVector2f>& Positions2D, const TArray<FVector3f>& Positions3D, const TArray<IndexType>& Indices, const int32 FabricIndex)
 	{
 		Reset();
 
+		SetFabricIndex(FabricIndex);
+
 		const int32 NumSimVertices = Positions2D.Num();
 		check(NumSimVertices == Positions3D.Num());
 
@@ -173,14 +180,16 @@ namespace UE::Chaos::ClothAsset
 			SimNormal[SimVertexIndex + StartVertexIndex[1]] = SimNormal[SimVertexIndex + StartVertexIndex[1]].GetSafeNormal(UE_SMALL_NUMBER, FVector3f::XAxisVector);
 		}
 	}
-	template CHAOSCLOTHASSET_API void FCollectionClothSimPatternFacade::Initialize(const TArray<FVector2f>& Positions2D, const TArray<FVector3f>& Positions3D, const TArray<int32>& Indices);
-	template CHAOSCLOTHASSET_API void FCollectionClothSimPatternFacade::Initialize(const TArray<FVector2f>& Positions2D, const TArray<FVector3f>& Positions3D, const TArray<uint32>& Indices);
-	template CHAOSCLOTHASSET_API void FCollectionClothSimPatternFacade::Initialize(const TArray<FVector2f>& Positions2D, const TArray<FVector3f>& Positions3D, const TArray<FIntVector3>& Indices);
+	template CHAOSCLOTHASSET_API void FCollectionClothSimPatternFacade::Initialize(const TArray<FVector2f>& Positions2D, const TArray<FVector3f>& Positions3D, const TArray<int32>& Indices, const int32 FabricIndex);
+	template CHAOSCLOTHASSET_API void FCollectionClothSimPatternFacade::Initialize(const TArray<FVector2f>& Positions2D, const TArray<FVector3f>& Positions3D, const TArray<uint32>& Indices, const int32 FabricIndex);
+	template CHAOSCLOTHASSET_API void FCollectionClothSimPatternFacade::Initialize(const TArray<FVector2f>& Positions2D, const TArray<FVector3f>& Positions3D, const TArray<FIntVector3>& Indices, const int32 FabricIndex);
 
-	void FCollectionClothSimPatternFacade::Initialize(const FCollectionClothSimPatternConstFacade& Other, const int32 SimVertex3DOffset)
+	void FCollectionClothSimPatternFacade::Initialize(const FCollectionClothSimPatternConstFacade& Other, const int32 SimVertex3DOffset, const int32 FabricsOffset)
 	{
 		Reset();
 
+		SetFabricIndex(FabricsOffset+Other.GetFabricIndex());
+
 		// Sim Vertices 2D Group
 		SetNumSimVertices2D(Other.GetNumSimVertices2D());
 		FClothCollection::CopyArrayViewData(GetSimPosition2D(), Other.GetSimPosition2D());
@@ -263,6 +272,11 @@ namespace UE::Chaos::ClothAsset
 			GetElementIndex());
 	}
 
+	void FCollectionClothSimPatternFacade::SetFabricIndex(const int32 FabricIndex)
+	{
+		GetClothCollection()->GetElements(GetClothCollection()->GetSimPatternFabric())[GetElementIndex()] = FabricIndex;
+	}
+
 	TArrayView<FIntVector3> FCollectionClothSimPatternFacade::GetSimIndices2D()
 	{
 		return GetClothCollection()->GetElements(
@@ -312,6 +326,7 @@ namespace UE::Chaos::ClothAsset
 		(*GetClothCollection()->GetSimVertices2DEnd())[ElementIndex] = INDEX_NONE;
 		(*GetClothCollection()->GetSimFacesStart())[ElementIndex] = INDEX_NONE;
 		(*GetClothCollection()->GetSimFacesEnd())[ElementIndex] = INDEX_NONE;
+		(*GetClothCollection()->GetSimPatternFabric())[ElementIndex] = INDEX_NONE;
 	}
 
 	void FCollectionClothSimPatternFacade::SetNumSimVertices2D(int32 NumSimVertices)
diff --git a/Engine/Plugins/Experimental/ChaosClothAsset/Source/ChaosClothAsset/Public/ChaosClothAsset/ClothCollectionGroup.h b/Engine/Plugins/Experimental/ChaosClothAsset/Source/ChaosClothAsset/Public/ChaosClothAsset/ClothCollectionGroup.h
--- a/Engine/Plugins/Experimental/ChaosClothAsset/Source/ChaosClothAsset/Public/ChaosClothAsset/ClothCollectionGroup.h
+++ b/Engine/Plugins/Experimental/ChaosClothAsset/Source/ChaosClothAsset/Public/ChaosClothAsset/ClothCollectionGroup.h
@@ -31,5 +31,7 @@ namespace UE::Chaos::ClothAsset
 		inline const FName RenderFaces = FName(TEXT("RenderFaces"));
 		/** Contains 3D render model. */
 		inline const FName RenderVertices = FName(TEXT("RenderVertices"));
+		/** Contains all the fabrics simulation parameters used by the sim patterns . */
+		inline const FName Fabrics = FName(TEXT("Fabrics"));
 	}
 }  // End namespace UE::Chaos::ClothAsset
diff --git a/Engine/Plugins/Experimental/ChaosClothAsset/Source/ChaosClothAsset/Public/ChaosClothAsset/CollectionClothFabricFacade.h b/Engine/Plugins/Experimental/ChaosClothAsset/Source/ChaosClothAsset/Public/ChaosClothAsset/CollectionClothFabricFacade.h
new file mode 100644
--- /dev/null
+++ b/Engine/Plugins/Experimental/ChaosClothAsset/Source/ChaosClothAsset/Public/ChaosClothAsset/CollectionClothFabricFacade.h
@@ -0,0 +1,123 @@
+// Copyright Epic Games, Inc. All Rights Reserved.
+
+#pragma once
+
+#include "Templates/SharedPointer.h"
+
+namespace UE::Chaos::ClothAsset
+{
+	/**
+	 * Cloth Asset collection fabric facade class to access cloth fabric data.
+	 * Constructed from FCollectionClothConstFacade.
+	 * Const access (read only) version.
+	 */
+	class CHAOSCLOTHASSET_API FCollectionClothFabricConstFacade
+	{
+	public:
+		FCollectionClothFabricConstFacade() = delete;
+
+		FCollectionClothFabricConstFacade(const FCollectionClothFabricConstFacade&) = delete;
+		FCollectionClothFabricConstFacade& operator=(const FCollectionClothFabricConstFacade&) = delete;
+
+		FCollectionClothFabricConstFacade(FCollectionClothFabricConstFacade&&) = default;
+		FCollectionClothFabricConstFacade& operator=(FCollectionClothFabricConstFacade&&) = default;
+
+		virtual ~FCollectionClothFabricConstFacade() = default;
+
+		/** Anisotropic fabric datas structure (Weft,Warp,Bias) */
+		struct CHAOSCLOTHASSET_API FAnisotropicData
+		{
+			FAnisotropicData(const float WeftValue, const float WarpValue, const float BiasValue);
+			FAnisotropicData(const FVector3f& VectorDatas);
+			FAnisotropicData(const float& FloatDatas);
+			
+			FVector3f GetVectorDatas() const;
+			
+			
+			float Weft;
+			float Warp;
+			float Bias;
+		};
+
+		/** Return the anisotropic bending stiffness */
+		FAnisotropicData GetXPBDAnisoBendingStiffness() const;
+
+		/** Return the buckling ratio */
+		float GetXPBDAnisoBucklingRatio() const;
+
+		/** Return the anisotropic buckling stiffness */
+		FAnisotropicData GetXPBDAnisoBucklingStiffness() const;
+
+		/** Return the anisotropic stretch stiffness */
+		FAnisotropicData GetXPBDAnisoStretchStiffness() const;
+
+		/** Return the cloth density */
+		float GetDensityWeighted() const;
+
+		/** Return the cloth thickness */
+		float GetCollisionThickness() const;
+
+		/** Return the cloth damping */
+		float GetXPBDAnisoDamping() const;
+
+		/** Return the cloth friction */
+		float GetFrictionCoefficient() const;
+
+	protected:
+		friend class FCollectionClothFabricFacade;  // For other instances access
+		friend class FCollectionClothConstFacade;
+		FCollectionClothFabricConstFacade(const TSharedRef<const class FClothCollection>& InClothCollection, int32 InFabricIndex);
+
+		static constexpr int32 GetBaseElementIndex() { return 0; }
+
+		/** Get the global element index */
+		int32 GetElementIndex() const { return GetBaseElementIndex() + FabricIndex; }
+
+		/** Cloth collection modified by  the fabric facade */
+		TSharedRef<const class FClothCollection> ClothCollection;
+
+		/** Fabric index that will be referred in the sim patterns */
+		int32 FabricIndex;
+	};
+
+	/**
+	 * Cloth Asset collection fabric facade class to access cloth fabric data.
+	 * Constructed from FCollectionClothFacade.
+	 * Non-const access (read/write) version.
+	 */
+	class CHAOSCLOTHASSET_API FCollectionClothFabricFacade final : public FCollectionClothFabricConstFacade
+	{
+	public:
+		FCollectionClothFabricFacade() = delete;
+
+		FCollectionClothFabricFacade(const FCollectionClothFabricFacade&) = delete;
+		FCollectionClothFabricFacade& operator=(const FCollectionClothFabricFacade&) = delete;
+
+		FCollectionClothFabricFacade(FCollectionClothFabricFacade&&) = default;
+		FCollectionClothFabricFacade& operator=(FCollectionClothFabricFacade&&) = default;
+
+		virtual ~FCollectionClothFabricFacade() override = default;
+
+		/** Initialize the cloth fabric with simulation parameters. */
+		void Initialize(const FAnisotropicData& BendingStiffness, const float BucklingRatio,
+			const FAnisotropicData& BucklingStiffness, const FAnisotropicData& StretchStiffness,
+			const float ClothDensity, const float ClothFriction,
+			const float ClothDamping, const float ClothThickness);
+
+		/** Initialize the cloth fabric with another one. */
+		void Initialize(const FCollectionClothFabricConstFacade& OtherFabricFacade);
+
+	private:
+		friend class FCollectionClothFacade;
+		FCollectionClothFabricFacade(const TSharedRef<class FClothCollection>& InClothCollection, int32 InFabricIndex);
+
+		/** Set default values to the fabric properties */
+		void SetDefaults();
+
+		/** Reset the fabric values properties */
+		void Reset();
+
+		/** Get the non const cloth collection */
+		TSharedRef<class FClothCollection> GetClothCollection() { return ConstCastSharedRef<class FClothCollection>(ClothCollection); }
+	};
+}  // End namespace UE::Chaos::ClothAsset
diff --git a/Engine/Plugins/Experimental/ChaosClothAsset/Source/ChaosClothAsset/Public/ChaosClothAsset/CollectionClothFacade.h b/Engine/Plugins/Experimental/ChaosClothAsset/Source/ChaosClothAsset/Public/ChaosClothAsset/CollectionClothFacade.h
--- a/Engine/Plugins/Experimental/ChaosClothAsset/Source/ChaosClothAsset/Public/ChaosClothAsset/CollectionClothFacade.h
+++ b/Engine/Plugins/Experimental/ChaosClothAsset/Source/ChaosClothAsset/Public/ChaosClothAsset/CollectionClothFacade.h
@@ -6,6 +6,7 @@
 #include "ChaosClothAsset/CollectionClothRenderPatternFacade.h"
 #include "ChaosClothAsset/CollectionClothSeamFacade.h"
 #include "ChaosClothAsset/IsUserAttributeType.h"
+#include "ChaosClothAsset/CollectionClothFabricFacade.h"
 
 namespace Chaos
 {
@@ -106,6 +107,12 @@ namespace UE::Chaos::ClothAsset
 		/** Return a seam facade for the specified seam index. */
 		FCollectionClothSeamConstFacade GetSeam(int32 SeamIndex) const;
 
+		//~ Fabric Group
+		/** Return the number of fabrics in this collection. */
+		int32 GetNumFabrics() const;
+		/** Return a fabric facade for the specified fabric index. */
+		FCollectionClothFabricConstFacade GetFabric(int32 FabricIndex) const;
+
 		//~ Render Vertices Group
 		/** Return the total number of render vertices for this collection. */
 		int32 GetNumRenderVertices() const;
@@ -145,6 +152,9 @@ namespace UE::Chaos::ClothAsset
 
 		friend class FCollectionClothSeamFacade;
 		friend class FCollectionClothSeamConstFacade;
+		friend class FCollectionClothFabricFacade;
+		friend class FCollectionClothFabricConstFacade;
+		
 		explicit FCollectionClothConstFacade(const TSharedRef<const class FClothCollection>& ClothCollection);
 	};
 
@@ -248,6 +258,18 @@ namespace UE::Chaos::ClothAsset
 		FCollectionClothSeamFacade AddGetSeam() { return GetSeam(AddSeam()); }
 		/** Remove a sorted list of seams. */
 		void RemoveSeams(const TArray<int32>& SortedDeletionList);
+		
+		//~ Fabric Group
+		/** Set the new number of fabrics to this cloth. */
+		void SetNumFabrics(int32 NumFabrics);
+		/** Add a new fabric to this cloth and return its index in the fabric list. */
+		int32 AddFabric();
+		/** Return a fabric facade for the specified fabric index. */
+		FCollectionClothFabricFacade GetFabric(int32 FabricIndex);
+		/** Add a new fabric to this cloth and return the fabric facade set to its index. */
+		FCollectionClothFabricFacade AddGetFabric() { return GetFabric(AddFabric()); }
+		/** Remove a sorted list of fabrics. */
+		void RemoveFabrics(const TArray<int32>& SortedDeletionList);
 
 		//~ Render Vertices Group
 		/** SetNumRenderVertices per pattern within pattern facade. */
@@ -287,6 +309,7 @@ namespace UE::Chaos::ClothAsset
 
 		friend class FCollectionClothSeamFacade;
 		friend class FCollectionClothSimPatternFacade;
+		friend class FCollectionClothFabricFacade;
 
 		explicit FCollectionClothFacade(const TSharedRef<class FClothCollection>& InClothCollection);
 
diff --git a/Engine/Plugins/Experimental/ChaosClothAsset/Source/ChaosClothAsset/Public/ChaosClothAsset/CollectionClothSimPatternFacade.h b/Engine/Plugins/Experimental/ChaosClothAsset/Source/ChaosClothAsset/Public/ChaosClothAsset/CollectionClothSimPatternFacade.h
--- a/Engine/Plugins/Experimental/ChaosClothAsset/Source/ChaosClothAsset/Public/ChaosClothAsset/CollectionClothSimPatternFacade.h
+++ b/Engine/Plugins/Experimental/ChaosClothAsset/Source/ChaosClothAsset/Public/ChaosClothAsset/CollectionClothSimPatternFacade.h
@@ -52,6 +52,9 @@ namespace UE::Chaos::ClothAsset
 
 		/** Return the Pattern index this facade has been created with. */
 		int32 GetPatternIndex() const { return PatternIndex; }
+		
+		/** Return the Pattern index this facade has been created with. */
+        int32 GetFabricIndex() const;
 
 	protected:
 		friend class FCollectionClothSimPatternFacade;  // For other instances access
@@ -88,10 +91,10 @@ namespace UE::Chaos::ClothAsset
 
 		/** Initialize the cloth pattern using the specified 3D and 2D positions, and topology. */
 		template<typename IndexType, TEMPLATE_REQUIRES(TIsIndexType<IndexType>::Value)>
-		void Initialize(const TArray<FVector2f>& Positions2D, const TArray<FVector3f>& Positions3D, const TArray<IndexType>& Indices);
+		void Initialize(const TArray<FVector2f>& Positions2D, const TArray<FVector3f>& Positions3D, const TArray<IndexType>& Indices, const int32 FabricIndex = INDEX_NONE);
 
 		/** Initialize this pattern using another pattern collection. */
-		void Initialize(const FCollectionClothSimPatternConstFacade& Other, const int32 SimVertex3DOffset);
+		void Initialize(const FCollectionClothSimPatternConstFacade& Other, const int32 SimVertex3DOffset, const int32 FabricsOffset = 0);
 
 		//~ Sim Vertices 2D Group
 		// Note: Use the FCollectionClothConstFacade accessors instead of these for the array indices to match the SimIndices2D values
@@ -110,6 +113,9 @@ namespace UE::Chaos::ClothAsset
 		TArrayView<FIntVector3> GetSimIndices3D();
 		void RemoveSimFaces(const TArray<int32>& SortedDeletionList);
 
+		/** Set the fabric index used by this pattern */
+		void SetFabricIndex(const int32 InFabricIndex);
+
 	private:
 		friend class FCollectionClothFacade;
 		FCollectionClothSimPatternFacade(const TSharedRef<class FClothCollection>& InClothCollection, int32 InPatternIndex);
diff --git a/Engine/Plugins/Experimental/ChaosClothAssetEditor/Source/ChaosClothAssetDataflowNodes/Private/ChaosClothAsset/ChaosClothAssetDataflowNodesModule.cpp b/Engine/Plugins/Experimental/ChaosClothAssetEditor/Source/ChaosClothAssetDataflowNodes/Private/ChaosClothAsset/ChaosClothAssetDataflowNodesModule.cpp
--- a/Engine/Plugins/Experimental/ChaosClothAssetEditor/Source/ChaosClothAssetDataflowNodes/Private/ChaosClothAsset/ChaosClothAssetDataflowNodesModule.cpp
+++ b/Engine/Plugins/Experimental/ChaosClothAssetEditor/Source/ChaosClothAssetDataflowNodes/Private/ChaosClothAsset/ChaosClothAssetDataflowNodesModule.cpp
@@ -23,6 +23,8 @@
 #include "ChaosClothAsset/SimulationDampingConfigNode.h"
 #include "ChaosClothAsset/SimulationDefaultConfigNode.h"
 #include "ChaosClothAsset/SimulationGravityConfigNode.h"
+#include "ChaosClothAsset/SimulationBendingConfigNode.h"
+#include "ChaosClothAsset/SimulationStretchConfigNode.h"
 #include "ChaosClothAsset/SimulationLongRangeAttachmentConfigNode.h"
 #include "ChaosClothAsset/SimulationMassConfigNode.h"
 #include "ChaosClothAsset/SimulationMaxDistanceConfigNode.h"
@@ -79,6 +81,8 @@ namespace UE::Chaos::ClothAsset
 			DATAFLOW_NODE_REGISTER_CREATION_FACTORY(FChaosClothAssetSimulationAnimDriveConfigNode);
 			DATAFLOW_NODE_REGISTER_CREATION_FACTORY(FChaosClothAssetSimulationBackstopConfigNode);
 			DATAFLOW_NODE_REGISTER_CREATION_FACTORY(FChaosClothAssetSimulationCollisionConfigNode);
+			DATAFLOW_NODE_REGISTER_CREATION_FACTORY(FChaosClothAssetSimulationBendingConfigNode);
+			DATAFLOW_NODE_REGISTER_CREATION_FACTORY(FChaosClothAssetSimulationStretchConfigNode);
 			DATAFLOW_NODE_REGISTER_CREATION_FACTORY(FChaosClothAssetSimulationDampingConfigNode);
 			DATAFLOW_NODE_REGISTER_CREATION_FACTORY(FChaosClothAssetSimulationDefaultConfigNode);
 			DATAFLOW_NODE_REGISTER_CREATION_FACTORY(FChaosClothAssetSimulationGravityConfigNode);
diff --git a/Engine/Plugins/Experimental/ChaosClothAssetEditor/Source/ChaosClothAssetDataflowNodes/Private/ChaosClothAsset/SimulationBaseConfigNode.cpp b/Engine/Plugins/Experimental/ChaosClothAssetEditor/Source/ChaosClothAssetDataflowNodes/Private/ChaosClothAsset/SimulationBaseConfigNode.cpp
--- a/Engine/Plugins/Experimental/ChaosClothAssetEditor/Source/ChaosClothAssetDataflowNodes/Private/ChaosClothAsset/SimulationBaseConfigNode.cpp
+++ b/Engine/Plugins/Experimental/ChaosClothAssetEditor/Source/ChaosClothAssetDataflowNodes/Private/ChaosClothAsset/SimulationBaseConfigNode.cpp
@@ -71,7 +71,7 @@ PRAGMA_DISABLE_DEPRECATION_WARNINGS
 		AddProperties(Context, Properties);  // Deprecated 5.4
 PRAGMA_ENABLE_DEPRECATION_WARNINGS
 
-		FPropertyHelper PropertyHelper(*this, Context, Properties);
+		FPropertyHelper PropertyHelper(*this, Context, Properties, ClothCollection);
 		AddProperties(PropertyHelper);
 
 		if (FCollectionClothFacade(ClothCollection).IsValid())  // Can only act on the collection if it is a valid cloth collection
@@ -118,10 +118,11 @@ int32 FChaosClothAssetSimulationBaseConfigNode::AddPropertyHelper(
 	return KeyIndex;
 }
 
-FChaosClothAssetSimulationBaseConfigNode::FPropertyHelper::FPropertyHelper(const FChaosClothAssetSimulationBaseConfigNode& InConfigNode, Dataflow::FContext& InContext, ::Chaos::Softs::FCollectionPropertyMutableFacade& InProperties)
+FChaosClothAssetSimulationBaseConfigNode::FPropertyHelper::FPropertyHelper(const FChaosClothAssetSimulationBaseConfigNode& InConfigNode, Dataflow::FContext& InContext, ::Chaos::Softs::FCollectionPropertyMutableFacade& InProperties, const TSharedRef<FManagedArrayCollection>& InClothCollection)
 	: ConfigNode(InConfigNode)
 	, Context(InContext)
 	, Properties(InProperties)
+	, ClothCollection(InClothCollection) 
 {}
 
 int32 FChaosClothAssetSimulationBaseConfigNode::FPropertyHelper::SetPropertyBool(
@@ -147,14 +148,23 @@ int32 FChaosClothAssetSimulationBaseConfigNode::FPropertyHelper::SetPropertyStri
 	return PropertyKeyIndex;
 }
 
+FString FChaosClothAssetSimulationBaseConfigNode::FPropertyHelper::GetPropertyString(const FString* PropertyReference) const
+{
+	return ConfigNode.GetValue<FString>(Context, PropertyReference);
+}
+
 int32 FChaosClothAssetSimulationBaseConfigNode::FPropertyHelper::SetPropertyWeighted(
 	const FName& PropertyName,
-	const FChaosClothAssetWeightedValue& PropertyValue,
+	const bool bIsAnimatable,
+	const float& PropertyLow,
+	const float& PropertyHigh,
+	const FString& WeightMap,
+	FString& MapOverride,
 	const TArray<FName>& SimilarPropertyNames,
-	ECollectionPropertyFlags PropertyFlags)
+	ECollectionPropertyFlags PropertyFlags) const
 {
 	ensureMsgf(!EnumHasAnyFlags(PropertyFlags, ECollectionPropertyFlags::Animatable), TEXT("Animatable flag ignored. Weighted properties are set animatable through FChaosClothAssetWeightedValue::bIsAnimatable."));
-	if (PropertyValue.bIsAnimatable)
+	if (bIsAnimatable)
 	{
 		EnumAddFlags(PropertyFlags, ECollectionPropertyFlags::Animatable);  // Animatable
 	}
@@ -163,25 +173,30 @@ int32 FChaosClothAssetSimulationBaseConfigNode::FPropertyHelper::SetPropertyWeig
 		EnumRemoveFlags(PropertyFlags, ECollectionPropertyFlags::Animatable);  // Non-animatable
 	}
 	const int32 PropertyKeyIndex = ConfigNode.AddPropertyHelper(Properties, PropertyName, SimilarPropertyNames, PropertyFlags);
-	Properties.SetWeightedValue(PropertyKeyIndex, PropertyValue.Low, PropertyValue.High);
-	Properties.SetStringValue(PropertyKeyIndex, ConfigNode.GetValue<FString>(Context, &PropertyValue.WeightMap));
-	PropertyValue.WeightMap_Override = ConfigNode.GetValue<FString>(Context, &PropertyValue.WeightMap, UE::Chaos::ClothAsset::FWeightMapTools::NotOverridden);
+	Properties.SetWeightedValue(PropertyKeyIndex, PropertyLow, PropertyHigh);
+	Properties.SetStringValue(PropertyKeyIndex, ConfigNode.GetValue<FString>(Context, &WeightMap));
+	MapOverride = ConfigNode.GetValue<FString>(Context, &WeightMap, UE::Chaos::ClothAsset::FWeightMapTools::NotOverridden);
 	return PropertyKeyIndex;
 }
 
+int32 FChaosClothAssetSimulationBaseConfigNode::FPropertyHelper::SetPropertyWeighted(
+	const FName& PropertyName,
+	const FChaosClothAssetWeightedValue& PropertyValue,
+	const TArray<FName>& SimilarPropertyNames,
+	ECollectionPropertyFlags PropertyFlags)
+{
+	return SetPropertyWeighted(PropertyName, PropertyValue.bIsAnimatable, PropertyValue.Low,
+		PropertyValue.High, PropertyValue.WeightMap, PropertyValue.WeightMap_Override, SimilarPropertyNames, PropertyFlags);
+}
+
 int32 FChaosClothAssetSimulationBaseConfigNode::FPropertyHelper::SetPropertyWeighted(
 	const FName& PropertyName,
 	const FChaosClothAssetWeightedValueNonAnimatable& PropertyValue,
 	const TArray<FName>& SimilarPropertyNames,
 	ECollectionPropertyFlags PropertyFlags)
 {
-	ensureMsgf(!EnumHasAnyFlags(PropertyFlags, ECollectionPropertyFlags::Animatable), TEXT("Animatable flag ignored. This code is for a non animatable weighted property."));
-	EnumRemoveFlags(PropertyFlags, ECollectionPropertyFlags::Animatable);  // Non animatable
-	const int32 PropertyKeyIndex = ConfigNode.AddPropertyHelper(Properties, PropertyName, SimilarPropertyNames, PropertyFlags);
-	Properties.SetWeightedValue(PropertyKeyIndex, PropertyValue.Low, PropertyValue.High);
-	Properties.SetStringValue(PropertyKeyIndex, ConfigNode.GetValue<FString>(Context, &PropertyValue.WeightMap));
-	PropertyValue.WeightMap_Override = ConfigNode.GetValue<FString>(Context, &PropertyValue.WeightMap, UE::Chaos::ClothAsset::FWeightMapTools::NotOverridden);
-	return PropertyKeyIndex;
+	return SetPropertyWeighted(PropertyName, false, PropertyValue.Low,
+		PropertyValue.High, PropertyValue.WeightMap, PropertyValue.WeightMap_Override, SimilarPropertyNames, PropertyFlags);
 }
 
 int32 FChaosClothAssetSimulationBaseConfigNode::FPropertyHelper::SetPropertyWeighted(
@@ -190,15 +205,8 @@ int32 FChaosClothAssetSimulationBaseConfigNode::FPropertyHelper::SetPropertyWeig
 	const TArray<FName>& SimilarPropertyNames,
 	ECollectionPropertyFlags PropertyFlags)
 {
-	ensureMsgf(!EnumHasAnyFlags(PropertyFlags, ECollectionPropertyFlags::Animatable), TEXT("Animatable flag ignored. This code is for a non animatable weighted property."));
-	EnumRemoveFlags(PropertyFlags, ECollectionPropertyFlags::Animatable);  // Non animatable
-	const int32 PropertyKeyIndex = ConfigNode.AddPropertyHelper(Properties, PropertyName, SimilarPropertyNames, PropertyFlags);
-	constexpr float Low = 0.f;
-	constexpr float High = 1.f;
-	Properties.SetWeightedValue(PropertyKeyIndex, Low, High);
-	Properties.SetStringValue(PropertyKeyIndex, ConfigNode.GetValue<FString>(Context, &PropertyValue.WeightMap));
-	PropertyValue.WeightMap_Override = ConfigNode.GetValue<FString>(Context, &PropertyValue.WeightMap, UE::Chaos::ClothAsset::FWeightMapTools::NotOverridden);
-	return PropertyKeyIndex;
+	return SetPropertyWeighted(PropertyName, false, 0.0f,
+		1.0f, PropertyValue.WeightMap, PropertyValue.WeightMap_Override, SimilarPropertyNames, PropertyFlags);
 }
 
 int32 FChaosClothAssetSimulationBaseConfigNode::FPropertyHelper::SetPropertyString(
diff --git a/Engine/Plugins/Experimental/ChaosClothAssetEditor/Source/ChaosClothAssetDataflowNodes/Private/ChaosClothAsset/SimulationBendingConfigNode.cpp b/Engine/Plugins/Experimental/ChaosClothAssetEditor/Source/ChaosClothAssetDataflowNodes/Private/ChaosClothAsset/SimulationBendingConfigNode.cpp
new file mode 100644
--- /dev/null
+++ b/Engine/Plugins/Experimental/ChaosClothAssetEditor/Source/ChaosClothAssetDataflowNodes/Private/ChaosClothAsset/SimulationBendingConfigNode.cpp
@@ -0,0 +1,237 @@
+// Copyright Epic Games, Inc. All Rights Reserved.
+
+#include "ChaosClothAsset/SimulationBendingConfigNode.h"
+#include "Chaos/CollectionPropertyFacade.h"
+
+#include UE_INLINE_GENERATED_CPP_BY_NAME(SimulationBendingConfigNode)
+
+FChaosClothAssetSimulationBendingConfigNode::FChaosClothAssetSimulationBendingConfigNode(const Dataflow::FNodeParameters& InParam, FGuid InGuid)
+	: FChaosClothAssetSimulationFabricConfigNode(InParam, InGuid)
+{
+	RegisterCollectionConnections();
+	RegisterInputConnection(&FlatnessRatio.WeightMap);
+	RegisterInputConnection(&RestAngle.WeightMap);
+	RegisterInputConnection(&BendingStiffness.WeightMap);
+	RegisterInputConnection(&BendingStiffnessWarp.WeightMap);
+	RegisterInputConnection(&BendingStiffnessWeft.WeightMap);
+	RegisterInputConnection(&BendingStiffnessBias.WeightMap);
+	RegisterInputConnection(&BendingDamping.WeightMap);
+	RegisterInputConnection(&BucklingStiffness.WeightMap);
+	RegisterInputConnection(&BucklingStiffnessWarp.WeightMap);
+	RegisterInputConnection(&BucklingStiffnessWeft.WeightMap);
+	RegisterInputConnection(&BucklingStiffnessBias.WeightMap);
+}
+
+void FChaosClothAssetSimulationBendingConfigNode::AddProperties(FPropertyHelper& PropertyHelper) const
+{
+	UE::Chaos::ClothAsset::FCollectionClothFacade ClothFacade(PropertyHelper.GetClothCollection());
+	if(!CanUseFabrics(ClothFacade))
+	{
+		if(SolverType == EChaosClothAssetConstraintSolverType::XPBD)
+		{
+			if(DistributionType == EChaosClothAssetConstraintDistributionType::Isotropic)
+			{
+				if(ConstraintType == EChaosClothAssetBendingConstraintType::FacesSpring)
+				{
+					PropertyHelper.SetPropertyWeighted(FName(TEXT("XPBDBendingSpringStiffness")), BendingStiffness,{
+						FName(TEXT("XPBDBendingElementStiffness")),
+						FName(TEXT("XPBDAnisoBendingStiffnessWarp")),
+						FName(TEXT("BendingSpringStiffness")),
+						FName(TEXT("BendingElementStiffness"))});
+					
+					PropertyHelper.SetPropertyWeighted(FName(TEXT("XPBDBendingSpringDamping")), BendingDamping,{
+						FName(TEXT("XPBDBendingElementDamping")),
+						FName(TEXT("XPBDAnisoBendingDamping"))});
+				}
+				else
+				{
+					PropertyHelper.SetPropertyWeighted(FName(TEXT("XPBDBendingElementStiffness")), BendingStiffness,{
+						FName(TEXT("XPBDBendingSpringStiffness")),
+						FName(TEXT("XPBDAnisoBendingStiffnessWarp")),
+						FName(TEXT("BendingSpringStiffness")),
+						FName(TEXT("BendingElementStiffness"))});
+					
+					PropertyHelper.SetPropertyWeighted(FName(TEXT("XPBDBendingElementDamping")), BendingDamping,{
+						FName(TEXT("XPBDBendingSpringDamping")),
+						FName(TEXT("XPBDAnisoBendingDamping"))});
+					
+					PropertyHelper.SetPropertyWeighted(FName(TEXT("XPBDBucklingStiffness")), BucklingStiffness,{
+						FName(TEXT("XPBDAnisoBucklingStiffnessWarp")),
+						FName(TEXT("BucklingStiffness"))});
+					
+					PropertyHelper.SetProperty(FName(TEXT("XPBDBucklingRatio")), BucklingRatio, {
+						FName(TEXT("XPBDAnisoBucklingRatio")),
+						FName(TEXT("BucklingRatios"))});
+				}
+				PropertyHelper.SetPropertyEnum(FName(TEXT("XPBDRestAngleType")), RestAngleType, {
+					FName(TEXT("XPBDAnisoRestAngleType")),
+					FName(TEXT("RestAngleType")) 
+				}, ECollectionPropertyFlags::None);
+				
+				PropertyHelper.SetPropertyWeighted(FName(TEXT("XPBDFlatnessRatio")), FlatnessRatio, {
+					FName(TEXT("XPBDAnisoFlatnessRatio")),
+					FName(TEXT("FlatnessRatio"))       
+				});
+				
+				PropertyHelper.SetPropertyWeighted(FName(TEXT("XPBDRestAngle")), RestAngle, {
+					FName(TEXT("XPBDAnisoRestAngle")),
+					FName(TEXT("RestAngle"))       
+				});
+			}
+			else
+			{
+				PropertyHelper.SetPropertyWeighted(FName(TEXT("XPBDAnisoBendingStiffnessWarp")), BendingStiffnessWarp, {
+					FName(TEXT("BendingSpringStiffness")),     
+					FName(TEXT("BendingElementStiffness")),     
+					FName(TEXT("XPBDBendingSpringStiffness")),  
+					FName(TEXT("XPBDBendingElementStiffness"))});
+				
+				PropertyHelper.SetPropertyWeighted(FName(TEXT("XPBDAnisoBendingStiffnessWeft")), BendingStiffnessWeft, {});
+				PropertyHelper.SetPropertyWeighted(FName(TEXT("XPBDAnisoBendingStiffnessBias")), BendingStiffnessBias, {});
+
+				PropertyHelper.SetPropertyWeighted(FName(TEXT("XPBDAnisoBucklingStiffnessWarp")), BucklingStiffnessWarp, {
+					FName(TEXT("BucklingStiffness")),    
+					FName(TEXT("XPBDBucklingStiffness"))});
+				
+				PropertyHelper.SetPropertyWeighted(FName(TEXT("XPBDAnisoBucklingStiffnessWeft")), BucklingStiffnessWeft, {});
+				PropertyHelper.SetPropertyWeighted(FName(TEXT("XPBDAnisoBucklingStiffnessBias")), BucklingStiffnessBias, {});
+
+				PropertyHelper.SetPropertyWeighted(FName(TEXT("XPBDAnisoBendingDamping")), BendingDamping, {
+					FName(TEXT("XPBDBendingSpringDamping")),  
+					FName(TEXT("XPBDBendingElementDamping"))});
+				
+				PropertyHelper.SetProperty(FName(TEXT("XPBDAnisoBucklingRatio")), BucklingRatio, {
+					FName(TEXT("BucklingRatio")),
+					FName(TEXT("XPBDBucklingRatio"))});
+
+				PropertyHelper.SetPropertyEnum(FName(TEXT("XPBDAnisoRestAngleType")), RestAngleType,{
+					FName(TEXT("XPBDRestAngleType")),
+					FName(TEXT("RestAngleType"))  
+				}, ECollectionPropertyFlags::None);
+				
+				PropertyHelper.SetPropertyWeighted(FName(TEXT("XPBDAnisoFlatnessRatio")), FlatnessRatio, {
+					FName(TEXT("XPBDFlatnessRatio")),
+					FName(TEXT("FlatnessRatio"))       
+				});
+				PropertyHelper.SetPropertyWeighted(FName(TEXT("XPBDAnisoRestAngle")), RestAngle, {
+					FName(TEXT("XPBDRestAngle")),
+					FName(TEXT("RestAngle"))       
+				});
+			}
+		}
+		else
+		{
+			if(ConstraintType == EChaosClothAssetBendingConstraintType::FacesSpring)
+			{
+				PropertyHelper.SetPropertyWeighted(FName(TEXT("BendingSpringStiffness")), BendingStiffness, {
+					FName(TEXT("BendingElementStiffness")),
+					FName(TEXT("XPBDBendingSpringStiffness")),    
+					FName(TEXT("XPBDBendingElementStiffness")),   
+					FName(TEXT("XPBDAnisoBendingStiffnessWarp"))});
+			}
+			else
+			{
+				PropertyHelper.SetPropertyWeighted(FName(TEXT("BendingElementStiffness")), BendingStiffness, {
+					FName(TEXT("BendingSpringStiffness")),        // Existing properties to warn against
+					FName(TEXT("XPBDBendingSpringStiffness")),
+					FName(TEXT("XPBDBendingElementStiffness")),   
+					FName(TEXT("XPBDAnisoBendingStiffnessWarp"))});
+				
+				PropertyHelper.SetPropertyWeighted(FName(TEXT("BucklingStiffness")), BucklingStiffness, {
+					FName(TEXT("XPBDBucklingStiffness")),
+					FName(TEXT("XPBDAnisoBucklingStiffnessWarp"))});
+				
+				PropertyHelper.SetProperty(FName(TEXT("BucklingRatios")), BucklingRatio, {
+					FName(TEXT("XPBDBucklingRatio")),
+					FName(TEXT("XPBDAnisoBucklingRatio"))});
+			}
+			
+			PropertyHelper.SetPropertyEnum(FName(TEXT("RestAngleType")), RestAngleType, {
+				FName(TEXT("XPBDAnisoRestAngleType")),
+				FName(TEXT("XPBDRestAngleType")) 
+			}, ECollectionPropertyFlags::None);
+			PropertyHelper.SetPropertyWeighted(FName(TEXT("FlatnessRatio")), FlatnessRatio, {
+					FName(TEXT("XPBDAnisoFlatnessRatio")),
+					FName(TEXT("XPBDFlatnessRatio"))       
+				});
+			PropertyHelper.SetPropertyWeighted(FName(TEXT("RestAngle")), RestAngle, {
+				FName(TEXT("XPBDAnisoRestAngle")), 
+				FName(TEXT("XPBDRestAngle"))       
+			});
+		}
+	}
+	else
+	{
+		SetFabricPropertyWeighted(FName(TEXT("XPBDAnisoBendingStiffnessWarp")), BendingStiffnessWarp, ClothFacade, PropertyHelper, [](
+			const UE::Chaos::ClothAsset::FCollectionClothFabricFacade& FabricFacade)-> float
+		 {
+			 return FabricFacade.GetXPBDAnisoBendingStiffness().Warp;
+		 }, {
+			FName(TEXT("BendingSpringStiffness")),     
+			FName(TEXT("BendingElementStiffness")),     
+			FName(TEXT("XPBDBendingSpringStiffness")),  
+			FName(TEXT("XPBDBendingElementStiffness"))});
+
+		SetFabricPropertyWeighted(FName(TEXT("XPBDAnisoBendingStiffnessWeft")), BendingStiffnessWeft, ClothFacade, PropertyHelper, [](
+			const UE::Chaos::ClothAsset::FCollectionClothFabricFacade& FabricFacade)-> float
+		 {
+			 return FabricFacade.GetXPBDAnisoBendingStiffness().Weft;
+		 }, {});
+		
+		SetFabricPropertyWeighted(FName(TEXT("XPBDAnisoBendingStiffnessBias")), BendingStiffnessBias,ClothFacade, PropertyHelper, [](
+			const UE::Chaos::ClothAsset::FCollectionClothFabricFacade& FabricFacade)-> float
+		 {
+			 return FabricFacade.GetXPBDAnisoBendingStiffness().Bias;
+		 }, {});
+		
+		SetFabricPropertyWeighted(FName(TEXT("XPBDAnisoBucklingStiffnessWarp")), BucklingStiffnessWarp,ClothFacade, PropertyHelper, [](
+			const UE::Chaos::ClothAsset::FCollectionClothFabricFacade& FabricFacade)-> float
+		 {
+			 return FabricFacade.GetXPBDAnisoBucklingStiffness().Warp;
+		 }, {
+			FName(TEXT("BucklingStiffness")),    
+			FName(TEXT("XPBDBucklingStiffness"))});
+		
+		SetFabricPropertyWeighted(FName(TEXT("XPBDAnisoBucklingStiffnessWeft")), BucklingStiffnessWeft, ClothFacade, PropertyHelper, [](
+			const UE::Chaos::ClothAsset::FCollectionClothFabricFacade& FabricFacade)-> float
+		{
+			return FabricFacade.GetXPBDAnisoBucklingStiffness().Weft;
+		},{});
+		
+		SetFabricPropertyWeighted(FName(TEXT("XPBDAnisoBucklingStiffnessBias")), BucklingStiffnessBias, ClothFacade, PropertyHelper, [](
+			const UE::Chaos::ClothAsset::FCollectionClothFabricFacade& FabricFacade)-> float
+		{
+			return FabricFacade.GetXPBDAnisoBucklingStiffness().Bias;
+		},{});
+
+		SetFabricPropertyWeighted(FName(TEXT("XPBDAnisoBendingDamping")), BendingDamping, ClothFacade, PropertyHelper, [](
+			const UE::Chaos::ClothAsset::FCollectionClothFabricFacade& FabricFacade)-> float
+		{
+			return FabricFacade.GetXPBDAnisoDamping();
+		},{
+			FName(TEXT("XPBDBendingSpringDamping")),  
+			FName(TEXT("XPBDBendingElementDamping"))});
+		
+		SetFabricProperty(FName(TEXT("XPBDAnisoBucklingRatio")), ClothFacade, PropertyHelper, [](const UE::Chaos::ClothAsset::FCollectionClothFabricFacade& FabricFacade)-> float
+		 {
+			return FabricFacade.GetXPBDAnisoBucklingRatio();
+		 }, {
+			FName(TEXT("BucklingRatio")),
+			FName(TEXT("XPBDBucklingRatio"))});
+
+		PropertyHelper.SetPropertyEnum(FName(TEXT("XPBDAnisoRestAngleType")), RestAngleType,{
+				FName(TEXT("XPBDRestAngleType")),
+				FName(TEXT("RestAngleType"))  
+			}, ECollectionPropertyFlags::None);
+		
+		PropertyHelper.SetPropertyWeighted(FName(TEXT("XPBDAnisoFlatnessRatio")), FlatnessRatio, {
+			FName(TEXT("XPBDFlatnessRatio")),
+			FName(TEXT("FlatnessRatio"))       
+		});
+		
+		PropertyHelper.SetPropertyWeighted(FName(TEXT("XPBDAnisoRestAngle")), RestAngle, {
+			FName(TEXT("XPBDRestAngle")),
+			FName(TEXT("RestAngle"))       
+		});
+	}
+}
diff --git a/Engine/Plugins/Experimental/ChaosClothAssetEditor/Source/ChaosClothAssetDataflowNodes/Private/ChaosClothAsset/SimulationCollisionConfigNode.cpp b/Engine/Plugins/Experimental/ChaosClothAssetEditor/Source/ChaosClothAssetDataflowNodes/Private/ChaosClothAsset/SimulationCollisionConfigNode.cpp
--- a/Engine/Plugins/Experimental/ChaosClothAssetEditor/Source/ChaosClothAssetDataflowNodes/Private/ChaosClothAsset/SimulationCollisionConfigNode.cpp
+++ b/Engine/Plugins/Experimental/ChaosClothAssetEditor/Source/ChaosClothAssetDataflowNodes/Private/ChaosClothAsset/SimulationCollisionConfigNode.cpp
@@ -6,15 +6,30 @@
 #include UE_INLINE_GENERATED_CPP_BY_NAME(SimulationCollisionConfigNode)
 
 FChaosClothAssetSimulationCollisionConfigNode::FChaosClothAssetSimulationCollisionConfigNode(const Dataflow::FNodeParameters& InParam, FGuid InGuid)
-	: FChaosClothAssetSimulationBaseConfigNode(InParam, InGuid)
+	: FChaosClothAssetSimulationFabricConfigNode(InParam, InGuid)
 {
 	RegisterCollectionConnections();
 }
 
 void FChaosClothAssetSimulationCollisionConfigNode::AddProperties(FPropertyHelper& PropertyHelper) const
 {
-	PropertyHelper.SetProperty(this, &CollisionThickness);
-	PropertyHelper.SetProperty(this, &FrictionCoefficient);
+	UE::Chaos::ClothAsset::FCollectionClothFacade ClothFacade(PropertyHelper.GetClothCollection());
+	if(!CanUseFabrics(ClothFacade))
+	{
+		PropertyHelper.SetProperty(this, &FrictionCoefficient);
+		PropertyHelper.SetProperty(this, &CollisionThickness);
+	}
+	else
+	{
+		SetFabricProperty(FName(TEXT("FrictionCoefficient")), ClothFacade, PropertyHelper, [](const UE::Chaos::ClothAsset::FCollectionClothFabricFacade& FabricFacade)-> float
+		 {
+			 return FabricFacade.GetFrictionCoefficient();
+		 }, {});
+		SetFabricProperty(FName(TEXT("CollisionThickness")), ClothFacade, PropertyHelper, [](const UE::Chaos::ClothAsset::FCollectionClothFabricFacade& FabricFacade)-> float
+		 {
+			 return FabricFacade.GetCollisionThickness();
+		 }, {});
+	}
 	PropertyHelper.SetPropertyBool(this, &bUseCCD);
 	PropertyHelper.SetProperty(this, &ProximityStiffness);
 }
diff --git a/Engine/Plugins/Experimental/ChaosClothAssetEditor/Source/ChaosClothAssetDataflowNodes/Private/ChaosClothAsset/SimulationFabricConfigNode.cpp b/Engine/Plugins/Experimental/ChaosClothAssetEditor/Source/ChaosClothAssetDataflowNodes/Private/ChaosClothAsset/SimulationFabricConfigNode.cpp
new file mode 100644
--- /dev/null
+++ b/Engine/Plugins/Experimental/ChaosClothAssetEditor/Source/ChaosClothAssetDataflowNodes/Private/ChaosClothAsset/SimulationFabricConfigNode.cpp
@@ -0,0 +1,131 @@
+// Copyright Epic Games, Inc. All Rights Reserved.
+
+#include "ChaosClothAsset/SimulationFabricConfigNode.h"
+#include "ClothConfig.h"
+#include "ChaosClothAsset/SimulationBaseConfigNodePrivate.h"
+#include "ChaosClothAsset/CollectionClothFacade.h"
+
+#include UE_INLINE_GENERATED_CPP_BY_NAME(SimulationFabricConfigNode)
+
+namespace UE::Chaos::ClothAsset::Private
+{
+	static FWeightedValueBounds ComputeWeightedValueBounds(FCollectionClothFacade& ClothFacade,
+		 TArray<float>& PatternValues, const TFunction<float(const FCollectionClothFabricFacade&)>& FabricValueFunction)
+	{
+		const int32 NumPatterns = ClothFacade.GetNumSimPatterns();
+		
+		PatternValues.Init(0.0f, NumPatterns);
+		float MinValue = FLT_MAX, MaxValue = 0.0f;
+		
+		for(int32 PatternIndex = 0; PatternIndex < NumPatterns; ++PatternIndex)
+		{
+			FCollectionClothSimPatternFacade PatternFacade = ClothFacade.GetSimPattern(PatternIndex);
+			const int32 FabricIndex = PatternFacade.GetFabricIndex();
+
+			if(FabricIndex >= 0 && FabricIndex < ClothFacade.GetNumFabrics())
+			{
+				FCollectionClothFabricFacade FabricFacade = ClothFacade.GetFabric(FabricIndex);
+				const float FabricValue = FabricValueFunction(FabricFacade);
+
+				MinValue = FMath::Min(MinValue, FabricValue);
+				MaxValue = FMath::Max(MaxValue, FabricValue);
+
+				PatternValues[PatternIndex] = FabricValue;
+			}
+		}
+		return {MinValue, MaxValue};
+	}
+
+	static FWeightedValueBounds BuildFabricWeightedValue(FCollectionClothFacade& ClothFacade,
+					  const FString& WeightMapName, const TFunction<float(const FCollectionClothFabricFacade&)>& FabricValueFunction)
+	{
+		const int32 NumPatterns = ClothFacade.GetNumSimPatterns();
+
+		TArray<float> PatternValues;
+		const FWeightedValueBounds WeightValueBounds = ComputeWeightedValueBounds(ClothFacade, PatternValues, FabricValueFunction);
+
+		
+		if(WeightValueBounds.Low != WeightValueBounds.High)
+		{
+			const int32 NumVertices = ClothFacade.GetNumSimVertices3D();
+		
+			TArray<int32> NumValues;
+			NumValues.Init(0, NumVertices);
+
+			TArray<float> VertexValues;
+			VertexValues.Init(0.0f, NumVertices);
+			
+			ClothFacade.AddWeightMap(*WeightMapName);
+			const TArrayView<float> ValueWeightMap = ClothFacade.GetWeightMap(*WeightMapName);
+			
+			for(int32 PatternIndex = 0; PatternIndex < NumPatterns; ++PatternIndex)
+			{
+				FCollectionClothSimPatternFacade PatternFacade = ClothFacade.GetSimPattern(PatternIndex);
+				
+				const TConstArrayView<int32> SimVertex3DLookup =
+					static_cast<FCollectionClothSimPatternConstFacade&>(PatternFacade).GetSimVertex3DLookup();
+
+				// Average of the pattern values at the seam
+				for(const int32& SimVertex3DIndex : SimVertex3DLookup)
+				{
+					VertexValues[SimVertex3DIndex] += PatternValues[PatternIndex];
+					NumValues[SimVertex3DIndex]++;
+				}
+			}
+			for(int32 VertexIndex = 0; VertexIndex < NumVertices; ++VertexIndex)
+            {
+				if(NumValues[VertexIndex] > 0)
+				{
+					ValueWeightMap[VertexIndex] = (VertexValues[VertexIndex] / NumValues[VertexIndex]  - WeightValueBounds.Low) /
+						(WeightValueBounds.High - WeightValueBounds.Low);
+				}
+            }
+		}
+		
+		return WeightValueBounds;
+	}
+
+	static FWeightedValueBounds BuildFabricValue(FCollectionClothFacade& ClothFacade,
+					  const TFunction<float(const FCollectionClothFabricFacade&)>& FabricValueFunction)
+	{
+		TArray<float> PatternValues;
+		return ComputeWeightedValueBounds(ClothFacade, PatternValues, FabricValueFunction);
+	}
+}
+
+FChaosClothAssetSimulationFabricConfigNode::FChaosClothAssetSimulationFabricConfigNode(const Dataflow::FNodeParameters& InParam, FGuid InGuid)
+	: FChaosClothAssetSimulationBaseConfigNode(InParam, InGuid)
+{}
+
+void FChaosClothAssetSimulationFabricConfigNode::SetFabricProperty(const FName& PropertyName, UE::Chaos::ClothAsset::FCollectionClothFacade& ClothFacade,
+	FPropertyHelper& PropertyHelper, const TFunction<float(const UE::Chaos::ClothAsset::FCollectionClothFabricFacade&)>& FabricValueFunction,
+	const TArray<FName>& SimilarPropertyNames, ECollectionPropertyFlags PropertyFlags)
+{
+	const UE::Chaos::ClothAsset::FWeightedValueBounds WeightedValueBounds =
+			UE::Chaos::ClothAsset::Private::BuildFabricValue(ClothFacade, FabricValueFunction);
+
+	const float AveragedValue = 0.5f * (WeightedValueBounds.Low + WeightedValueBounds.High);
+	PropertyHelper.SetProperty(PropertyName, AveragedValue, SimilarPropertyNames, PropertyFlags);
+}
+
+template<typename PropertyType>
+void FChaosClothAssetSimulationFabricConfigNode::SetFabricPropertyWeighted(
+	const FName& PropertyName, const PropertyType& PropertyValue,  UE::Chaos::ClothAsset::FCollectionClothFacade& ClothFacade,
+	FPropertyHelper& PropertyHelper, const TFunction<float(const UE::Chaos::ClothAsset::FCollectionClothFabricFacade&)>& FabricValueFunction,
+	const TArray<FName>& SimilarPropertyNames, ECollectionPropertyFlags PropertyFlags)
+{
+	const UE::Chaos::ClothAsset::FWeightedValueBounds WeightedValueBounds = UE::Chaos::ClothAsset::Private::BuildFabricWeightedValue(ClothFacade, 
+			 PropertyHelper.GetPropertyString(&PropertyValue.WeightMap), FabricValueFunction);
+
+	PropertyHelper.SetPropertyWeighted(PropertyName, WeightedValueBounds, PropertyValue, SimilarPropertyNames, PropertyFlags);
+}
+
+template void FChaosClothAssetSimulationFabricConfigNode::SetFabricPropertyWeighted<FChaosClothAssetWeightedValue>(const FName& PropertyName,
+	const FChaosClothAssetWeightedValue& PropertyValue,  UE::Chaos::ClothAsset::FCollectionClothFacade& ClothFacade,
+	FPropertyHelper& PropertyHelper, const TFunction<float(const UE::Chaos::ClothAsset::FCollectionClothFabricFacade&)>& FabricValueFunction,
+	const TArray<FName>& SimilarPropertyNames, ECollectionPropertyFlags PropertyFlags);
+
+template void FChaosClothAssetSimulationFabricConfigNode::SetFabricPropertyWeighted<FChaosClothAssetWeightedValueNonAnimatable>(const FName& PropertyName,
+	const FChaosClothAssetWeightedValueNonAnimatable& PropertyValue,  UE::Chaos::ClothAsset::FCollectionClothFacade& ClothFacade,
+	FPropertyHelper& PropertyHelper, const TFunction<float(const UE::Chaos::ClothAsset::FCollectionClothFabricFacade&)>& FabricValueFunction,
+	const TArray<FName>& SimilarPropertyNames, ECollectionPropertyFlags PropertyFlags);
diff --git a/Engine/Plugins/Experimental/ChaosClothAssetEditor/Source/ChaosClothAssetDataflowNodes/Private/ChaosClothAsset/SimulationMassConfigNode.cpp b/Engine/Plugins/Experimental/ChaosClothAssetEditor/Source/ChaosClothAssetDataflowNodes/Private/ChaosClothAsset/SimulationMassConfigNode.cpp
--- a/Engine/Plugins/Experimental/ChaosClothAssetEditor/Source/ChaosClothAssetDataflowNodes/Private/ChaosClothAsset/SimulationMassConfigNode.cpp
+++ b/Engine/Plugins/Experimental/ChaosClothAssetEditor/Source/ChaosClothAssetDataflowNodes/Private/ChaosClothAsset/SimulationMassConfigNode.cpp
@@ -1,13 +1,15 @@
 // Copyright Epic Games, Inc. All Rights Reserved.
 
 #include "ChaosClothAsset/SimulationMassConfigNode.h"
+#include "ChaosClothAsset/SimulationFabricConfigNode.h"
+#include "ChaosClothAsset/CollectionClothFacade.h"
 #include "Chaos/CollectionPropertyFacade.h"
 #include "UObject/FortniteValkyrieBranchObjectVersion.h"
 
 #include UE_INLINE_GENERATED_CPP_BY_NAME(SimulationMassConfigNode)
 
 FChaosClothAssetSimulationMassConfigNode::FChaosClothAssetSimulationMassConfigNode(const Dataflow::FNodeParameters& InParam, FGuid InGuid)
-	: FChaosClothAssetSimulationBaseConfigNode(InParam, InGuid)
+	: FChaosClothAssetSimulationFabricConfigNode(InParam, InGuid)
 {
 	RegisterCollectionConnections();
 	RegisterInputConnection(&UniformMassWeighted.WeightMap);
@@ -16,27 +18,40 @@ FChaosClothAssetSimulationMassConfigNode::FChaosClothAssetSimulationMassConfigNo
 
 void FChaosClothAssetSimulationMassConfigNode::AddProperties(FPropertyHelper& PropertyHelper) const
 {
-	PropertyHelper.SetPropertyEnum(this, &MassMode, {}, ECollectionPropertyFlags::Intrinsic);
-	switch (MassMode)
+	UE::Chaos::ClothAsset::FCollectionClothFacade ClothFacade(PropertyHelper.GetClothCollection());
+	if(!CanUseFabrics(ClothFacade))
 	{
-	default:
-	case EClothMassMode::UniformMass:
-	{
-		PropertyHelper.SetPropertyWeighted(FName(TEXT("MassValue")), UniformMassWeighted, {}, ECollectionPropertyFlags::Intrinsic);
-	}
-	break;
-	case EClothMassMode::TotalMass:
-	{
-		PropertyHelper.SetProperty(FName(TEXT("MassValue")), TotalMass, {}, ECollectionPropertyFlags::Intrinsic);
+		PropertyHelper.SetPropertyEnum(this, &MassMode, {}, ECollectionPropertyFlags::Intrinsic);
+		switch (MassMode)
+		{
+		default:
+		case EClothMassMode::UniformMass:
+			{
+				PropertyHelper.SetPropertyWeighted(FName(TEXT("MassValue")), UniformMassWeighted, {}, ECollectionPropertyFlags::Intrinsic);
+			}
+			break;
+		case EClothMassMode::TotalMass:
+			{
+				PropertyHelper.SetProperty(FName(TEXT("MassValue")), TotalMass, {}, ECollectionPropertyFlags::Intrinsic);
+			}
+			break;
+		case EClothMassMode::Density:
+			{
+				PropertyHelper.SetPropertyWeighted(FName(TEXT("MassValue")), DensityWeighted, {}, ECollectionPropertyFlags::Intrinsic);
+			}
+			break;
+		}
 	}
-	break;
-	case EClothMassMode::Density:
+	else
 	{
-		PropertyHelper.SetPropertyWeighted(FName(TEXT("MassValue")), DensityWeighted, {}, ECollectionPropertyFlags::Intrinsic);
-	}
-	break;
-	}
+		constexpr EClothMassMode MassModeProperty = EClothMassMode::Density;
+		PropertyHelper.SetPropertyEnum(FName(TEXT("MassMode")), MassModeProperty, {}, ECollectionPropertyFlags::Intrinsic);
 
+		SetFabricPropertyWeighted(FName(TEXT("MassValue")), DensityWeighted, ClothFacade, PropertyHelper, [](const UE::Chaos::ClothAsset::FCollectionClothFabricFacade& FabricFacade)-> float
+		 {
+			 return FabricFacade.GetDensityWeighted();
+		 }, {}, ECollectionPropertyFlags::Intrinsic);
+	}
 	PropertyHelper.SetProperty(this, &MinPerParticleMass, {}, ECollectionPropertyFlags::Intrinsic);
 }
 
diff --git a/Engine/Plugins/Experimental/ChaosClothAssetEditor/Source/ChaosClothAssetDataflowNodes/Private/ChaosClothAsset/SimulationStretchConfigNode.cpp b/Engine/Plugins/Experimental/ChaosClothAssetEditor/Source/ChaosClothAssetDataflowNodes/Private/ChaosClothAsset/SimulationStretchConfigNode.cpp
new file mode 100644
--- /dev/null
+++ b/Engine/Plugins/Experimental/ChaosClothAssetEditor/Source/ChaosClothAssetDataflowNodes/Private/ChaosClothAsset/SimulationStretchConfigNode.cpp
@@ -0,0 +1,145 @@
+// Copyright Epic Games, Inc. All Rights Reserved.
+
+#include "ChaosClothAsset/SimulationStretchConfigNode.h"
+#include "Chaos/CollectionPropertyFacade.h"
+
+#include UE_INLINE_GENERATED_CPP_BY_NAME(SimulationStretchConfigNode)
+
+FChaosClothAssetSimulationStretchConfigNode::FChaosClothAssetSimulationStretchConfigNode(const Dataflow::FNodeParameters& InParam, FGuid InGuid)
+	: FChaosClothAssetSimulationFabricConfigNode(InParam, InGuid)
+{
+	RegisterCollectionConnections();
+	RegisterInputConnection(&StretchStiffness.WeightMap);
+	RegisterInputConnection(&StretchStiffnessWarp.WeightMap);
+	RegisterInputConnection(&StretchStiffnessWeft.WeightMap);
+	RegisterInputConnection(&StretchStiffnessBias.WeightMap);
+	RegisterInputConnection(&StretchDamping.WeightMap);
+	RegisterInputConnection(&StretchWarpScale.WeightMap);
+	RegisterInputConnection(&StretchWeftScale.WeightMap);
+	RegisterInputConnection(&AreaStiffness.WeightMap);
+}
+
+void FChaosClothAssetSimulationStretchConfigNode::AddProperties(FPropertyHelper& PropertyHelper) const
+{
+	UE::Chaos::ClothAsset::FCollectionClothFacade ClothFacade(PropertyHelper.GetClothCollection());
+	if(!CanUseFabrics(ClothFacade))
+	{
+		if(SolverType == EChaosClothAssetConstraintSolverType::XPBD)
+		{
+			if(DistributionType == EChaosClothAssetConstraintDistributionType::Anisotropic)
+			{
+				if(ConstraintType == EChaosClothAssetStretchConstraintType::StretchShear)
+				{
+					PropertyHelper.SetPropertyBool(FName(TEXT("XPBDAnisoStretchUse3dRestLengths")), bStretchUse3dRestLengths, {
+						FName(TEXT("XPBDAnisoSpringUse3dRestLengths"))}, ECollectionPropertyFlags::None);  // Non animatable
+
+					PropertyHelper.SetPropertyWeighted(FName(TEXT("XPBDAnisoStretchStiffnessWarp")), StretchStiffnessWarp, {
+						FName(TEXT("EdgeSpringStiffness")),
+						FName(TEXT("XPBDEdgeSpringStiffness")),
+						FName(TEXT("XPBDAnisoSpringStiffnessWarp"))});
+				
+					PropertyHelper.SetPropertyWeighted(FName(TEXT("XPBDAnisoStretchStiffnessWeft")), StretchStiffnessWeft, {
+						FName(TEXT("XPBDAnisoSpringStiffnessWeft"))});
+					PropertyHelper.SetPropertyWeighted(FName(TEXT("XPBDAnisoStretchStiffnessBias")), StretchStiffnessBias, {
+						FName(TEXT("XPBDAnisoSpringStiffnessBias"))});
+					PropertyHelper.SetPropertyWeighted(FName(TEXT("XPBDAnisoStretchDamping")), StretchDamping, {
+						FName(TEXT("XPBDEdgeSpringDamping")),
+						FName(TEXT("XPBDAnisoSpringDamping"))});
+				
+					PropertyHelper.SetPropertyWeighted(FName(TEXT("XPBDAnisoStretchWarpScale")), StretchWarpScale, {
+						FName(TEXT("XPBDAnisoSpringWarpScale"))});
+					PropertyHelper.SetPropertyWeighted(FName(TEXT("XPBDAnisoStretchWeftScale")), StretchWeftScale, {
+						FName(TEXT("XPBDAnisoSpringWeftScale"))});
+				}
+				else
+				{
+					PropertyHelper.SetPropertyBool(FName(TEXT("XPBDAnisoSpringUse3dRestLengths")), bStretchUse3dRestLengths, {
+						FName(TEXT("XPBDAnisoStretchUse3dRestLengths"))}, ECollectionPropertyFlags::None);  // Non animatable
+
+					PropertyHelper.SetPropertyWeighted(FName(TEXT("XPBDAnisoSpringStiffnessWarp")), StretchStiffnessWarp, {
+						FName(TEXT("EdgeSpringStiffness")),
+						FName(TEXT("XPBDEdgeSpringStiffness")),
+						FName(TEXT("XPBDAnisoStretchStiffnessWarp"))});
+				
+					PropertyHelper.SetPropertyWeighted(FName(TEXT("XPBDAnisoSpringStiffnessWeft")), StretchStiffnessWeft, {
+						FName(TEXT("XPBDAnisoStretchStiffnessWeft"))});
+					PropertyHelper.SetPropertyWeighted(FName(TEXT("XPBDAnisoSpringStiffnessBias")), StretchStiffnessBias, {
+						FName(TEXT("XPBDAnisoStretchStiffnessBias"))});
+					PropertyHelper.SetPropertyWeighted(FName(TEXT("XPBDAnisoSpringDamping")), StretchDamping, {
+						FName(TEXT("XPBDAnisoStretchDamping")),
+						FName(TEXT("XPBDEdgeSpringDamping"))});
+				
+					PropertyHelper.SetPropertyWeighted(FName(TEXT("XPBDAnisoSpringWarpScale")), StretchWarpScale, {
+						FName(TEXT("XPBDAnisoStretchWarpScale"))});
+					PropertyHelper.SetPropertyWeighted(FName(TEXT("XPBDAnisoSpringWeftScale")), StretchWeftScale, {
+						FName(TEXT("XPBDAnisoStretchWeftScale"))});
+				}
+			}
+			else
+			{
+				PropertyHelper.SetPropertyWeighted(FName(TEXT("XPBDEdgeSpringStiffness")), StretchStiffness, {
+					FName(TEXT("EdgeSpringStiffness")),
+					FName(TEXT("XPBDAnisoStretchStiffnessWarp")),
+					FName(TEXT("XPBDAnisoSpringStiffnessWarp"))});
+				
+				PropertyHelper.SetPropertyWeighted(FName(TEXT("XPBDEdgeSpringDamping")), StretchDamping, {
+					FName(TEXT("XPBDAnisoStretchDamping")),
+					FName(TEXT("XPBDAnisoSpringDamping"))});
+
+				if(bAddAreaConstraint)
+				{
+					PropertyHelper.SetPropertyWeighted(FName(TEXT("XPBDAreaSpringStiffness")), AreaStiffness,{
+						FName(TEXT("AreaSpringStiffness"))});
+				}
+			}
+		}
+		else
+		{
+			PropertyHelper.SetPropertyWeighted(FName(TEXT("EdgeSpringStiffness")), StretchStiffness, {
+				FName(TEXT("XPBDEdgeSpringStiffness")), 
+				FName(TEXT("XPBDAnisoStretchStiffnessWarp")),
+				FName(TEXT("XPBDAnisoSpringStiffnessWarp"))});
+			
+			if(bAddAreaConstraint)
+			{
+				PropertyHelper.SetPropertyWeighted(FName(TEXT("AreaSpringStiffness")), AreaStiffness,{
+					FName(TEXT("XPBDAreaSpringStiffness"))});
+			}
+		}
+	}
+	else
+	{
+		PropertyHelper.SetPropertyBool(FName(TEXT("XPBDAnisoStretchUse3dRestLengths")), bStretchUse3dRestLengths, {
+						FName(TEXT("XPBDAnisoSpringUse3dRestLengths"))}, ECollectionPropertyFlags::None);  // Non animatable
+
+		SetFabricPropertyWeighted(FName(TEXT("XPBDAnisoStretchStiffnessWarp")), StretchStiffnessWarp,ClothFacade, PropertyHelper, [](
+			const UE::Chaos::ClothAsset::FCollectionClothFabricFacade& FabricFacade)-> float
+		 {
+			 return FabricFacade.GetXPBDAnisoStretchStiffness().Warp;
+		 }, {
+			FName(TEXT("EdgeSpringStiffness")),
+			FName(TEXT("XPBDEdgeSpringStiffness"))});
+
+		SetFabricPropertyWeighted(FName(TEXT("XPBDAnisoStretchStiffnessWeft")), StretchStiffnessWeft, ClothFacade, PropertyHelper, [](
+			const UE::Chaos::ClothAsset::FCollectionClothFabricFacade& FabricFacade)-> float
+		 {
+			 return FabricFacade.GetXPBDAnisoStretchStiffness().Weft;
+		 }, {});
+		
+		SetFabricPropertyWeighted(FName(TEXT("XPBDAnisoStretchStiffnessBias")), StretchStiffnessBias,ClothFacade, PropertyHelper, [](
+			const UE::Chaos::ClothAsset::FCollectionClothFabricFacade& FabricFacade)-> float
+		 {
+			 return FabricFacade.GetXPBDAnisoStretchStiffness().Bias;
+		 }, {});
+
+		SetFabricPropertyWeighted(FName(TEXT("XPBDAnisoStretchDamping")), StretchDamping,ClothFacade, PropertyHelper, [](
+			const UE::Chaos::ClothAsset::FCollectionClothFabricFacade& FabricFacade)-> float
+		 {
+			 return FabricFacade.GetXPBDAnisoDamping();
+		 }, {
+			FName(TEXT("XPBDEdgeSpringDamping"))});
+
+		PropertyHelper.SetPropertyWeighted(FName(TEXT("XPBDAnisoStretchWarpScale")), StretchWarpScale, {});
+		PropertyHelper.SetPropertyWeighted(FName(TEXT("XPBDAnisoStretchWeftScale")), StretchWeftScale, {});
+	}
+}
diff --git a/Engine/Plugins/Experimental/ChaosClothAssetEditor/Source/ChaosClothAssetDataflowNodes/Private/ChaosClothAsset/USDImportNode.cpp b/Engine/Plugins/Experimental/ChaosClothAssetEditor/Source/ChaosClothAssetDataflowNodes/Private/ChaosClothAsset/USDImportNode.cpp
--- a/Engine/Plugins/Experimental/ChaosClothAssetEditor/Source/ChaosClothAssetDataflowNodes/Private/ChaosClothAsset/USDImportNode.cpp
+++ b/Engine/Plugins/Experimental/ChaosClothAssetEditor/Source/ChaosClothAssetDataflowNodes/Private/ChaosClothAsset/USDImportNode.cpp
@@ -25,7 +25,9 @@
 #include "USDStageImportOptions.h"
 #include "USDTypesConversion.h"
 #include "USDValueConversion.h"
+#include "Chaos/CollectionPropertyFacade.h"
 #include "UsdWrappers/UsdPrim.h"
+#include "UsdWrappers/UsdRelationship.h"
 #include "UsdWrappers/VtValue.h"
 
 #include UE_INLINE_GENERATED_CPP_BY_NAME(USDImportNode)
@@ -566,6 +568,127 @@ bool FChaosClothAssetUSDImportNode::ImportFromFile(const FString& UsdFilePath, c
 		}
 	}
 
+	auto FillIntDatas = [](const UE::FUsdPrim& FabricPrim, const FString& DatasName, uint32& IntDatas) 
+	{
+		const UE::FUsdAttribute IntDatasAttr = FabricPrim.GetAttribute(*DatasName);
+		if (IntDatasAttr.HasValue() && IntDatasAttr.GetTypeName() == TEXT("uint"))
+		{
+			UE::FVtValue Value;
+			IntDatasAttr.Get(Value);
+			const TOptional<uint32> Optional = UsdUtils::GetUnderlyingValue<uint32>(Value);
+			IntDatas = Optional.IsSet() ? Optional.GetValue() : 0;
+		}
+	};
+
+	// Fabrics
+	TArray<uint32> FabricIds;
+	const UE::FSdfPath FabricsPath = UE::FSdfPath(UE::FSdfPath::AbsoluteRootPath()).AppendChild(TEXT("SimulationData")).AppendChild(TEXT("Fabrics"));
+	if (const UE::FUsdPrim FabricsPrim = UsdStage.GetPrimAtPath(FabricsPath))
+	{
+		auto FillFloatDatas = [](const UE::FUsdPrim& FabricPrim, const FString& DatasName, float& FloatDatas) 
+		{
+			const UE::FUsdAttribute FloatDatasAttr = FabricPrim.GetAttribute(*DatasName);
+			if (FloatDatasAttr.HasValue() && FloatDatasAttr.GetTypeName() == TEXT("float"))
+			{
+				UE::FVtValue Value;
+				FloatDatasAttr.Get(Value);
+				const TOptional<float> Optional = UsdUtils::GetUnderlyingValue<float>(Value);
+				FloatDatas = Optional.IsSet() ? Optional.GetValue() : 0.0f;
+			}
+		};
+		for (const UE::FUsdPrim& FabricPrim : FabricsPrim.GetChildren())
+		{
+			float BendingBiasLeft = 0.0f;
+			FillFloatDatas(FabricPrim, TEXT("BendingBiasLeft"), BendingBiasLeft);
+
+			float BendingBiasRight= 0.0f;
+			FillFloatDatas(FabricPrim, TEXT("BendingBiasRight"), BendingBiasRight);
+
+			float BendingWarp = 0.0f;
+			FillFloatDatas(FabricPrim, TEXT("BendingWarp"), BendingWarp);
+
+			float BendingWeft = 0.0f;
+			FillFloatDatas(FabricPrim, TEXT("BendingWeft"), BendingWeft);
+
+			float BucklingRatioBiasLeft = 0.0f;
+			FillFloatDatas(FabricPrim, TEXT("BucklingRatioBiasLeft"), BucklingRatioBiasLeft);
+
+			float BucklingRatioBiasRight = 0.0f;
+			FillFloatDatas(FabricPrim, TEXT("BucklingRatioBiasRight"), BucklingRatioBiasRight);
+			
+			float BucklingRatioWarp = 0.0f;
+			FillFloatDatas(FabricPrim, TEXT("BucklingRatioWarp"), BucklingRatioWarp);
+
+			float BucklingRatioWeft = 0.0f;
+			FillFloatDatas(FabricPrim, TEXT("BucklingRatioWeft"), BucklingRatioWeft);
+
+			float BucklingStiffnessBiasLeft = 0.0f;
+			FillFloatDatas(FabricPrim, TEXT("BucklingStiffnessBiasLeft"), BucklingStiffnessBiasLeft);
+
+			float BucklingStiffnessBiasRight = 0.0f;
+			FillFloatDatas(FabricPrim, TEXT("BucklingStiffnessBiasRight"), BucklingStiffnessBiasRight);
+
+			float BucklingStiffnessWarp = 0.0f;
+			FillFloatDatas(FabricPrim, TEXT("BucklingStiffnessWarp"), BucklingStiffnessWarp);
+
+			float BucklingStiffnessWeft = 0.0f;
+			FillFloatDatas(FabricPrim, TEXT("BucklingStiffnessWeft"), BucklingStiffnessWeft);
+
+			float Density = 0.0f;
+			FillFloatDatas(FabricPrim, TEXT("Density"), Density);
+
+			float Friction = 0.0f;
+			FillFloatDatas(FabricPrim, TEXT("Friction"), Friction);
+
+			float InternalDamping = 0.0f;
+			FillFloatDatas(FabricPrim, TEXT("InternalDamping"), InternalDamping);
+
+			float Thickness = 0.0f;
+			FillFloatDatas(FabricPrim, TEXT("Thickness"), Thickness);
+
+			float ShearLeft = 0.0f;
+			FillFloatDatas(FabricPrim, TEXT("ShearLeft"), ShearLeft);
+
+			float ShearRight = 0.0f;
+			FillFloatDatas(FabricPrim, TEXT("ShearRight"), ShearRight);
+
+			float StretchWarp = 0.0f;
+			FillFloatDatas(FabricPrim, TEXT("StretchWarp"), StretchWarp);
+
+			float StretchWeft = 0.0f;
+			FillFloatDatas(FabricPrim, TEXT("StretchWeft"), StretchWeft);
+
+			uint32 FabricId = 0;
+			FillIntDatas(FabricPrim, TEXT("FabricId"), FabricId);
+
+			FCollectionClothFabricFacade Fabric = ClothFacade.AddGetFabric();
+
+			static constexpr float BendingScaling = 1e-5f; // from g.mm2/s2 to kg.cm2/s2
+			static constexpr float StretchShearScaling = 1e-3f; // from g/s2 to kg.cm/s2
+			static constexpr float DensityScaling = 1e+3f; // from g/mm2 to kg/m2
+			static constexpr float ThicknessScaling = 1e-1f; // from mm to cm
+			
+			FCollectionClothFabricFacade::FAnisotropicData BendingStiffness(
+				  BendingWeft*BendingScaling, BendingWarp*BendingScaling, 0.5f * (BendingBiasLeft+BendingBiasRight)*BendingScaling);
+			
+			FCollectionClothFabricFacade::FAnisotropicData StretchStiffness(
+				StretchWeft*StretchShearScaling, StretchWarp*StretchShearScaling, 0.5f * (ShearLeft+ShearRight)*StretchShearScaling);
+			
+			FCollectionClothFabricFacade::FAnisotropicData BucklingStiffness(
+				BendingStiffness.Weft * BucklingStiffnessWeft, BendingStiffness.Warp * BucklingStiffnessWarp,
+				BendingStiffness.Bias * 0.5f * (BucklingStiffnessBiasLeft+BucklingStiffnessBiasRight));
+
+			// Only scalar value used in the solver right now
+			const float BucklingRatio = (BucklingRatioWeft + BucklingRatioWarp +
+				0.5f * (BucklingRatioBiasLeft+BucklingRatioBiasRight)) / 3.0f;
+
+			Fabric.Initialize(BendingStiffness, BucklingRatio, BucklingStiffness, StretchStiffness,
+				Density * DensityScaling, Friction, InternalDamping, Thickness * ThicknessScaling);
+
+			FabricIds.Add(FabricId);
+		}
+	}
+
 	// Patterns
 	const UE::FSdfPath PatternsPath = UE::FSdfPath(UE::FSdfPath::AbsoluteRootPath()).AppendChild(TEXT("SimulationData")).AppendChild(TEXT("Patterns"));
 	if (const UE::FUsdPrim PatternsPrim = UsdStage.GetPrimAtPath(PatternsPath))
@@ -573,33 +696,31 @@ bool FChaosClothAssetUSDImportNode::ImportFromFile(const FString& UsdFilePath, c
 		for (const UE::FUsdPrim& PatternPrim : PatternsPrim.GetChildren())
 		{
 			uint32 TriangleCount = 0;
-			const UE::FUsdAttribute TriangleCountAttr = PatternPrim.GetAttribute(TEXT("TriangleCount"));
-			if (TriangleCountAttr.HasValue() && TriangleCountAttr.GetTypeName() == TEXT("uint"))
-			{
-				UE::FVtValue Value;
-				TriangleCountAttr.Get(Value);
-				const TOptional<uint32> Optional = UsdUtils::GetUnderlyingValue<uint32>(Value);
-				TriangleCount = Optional.IsSet() ? Optional.GetValue() : 0;
-			}
+			FillIntDatas(PatternPrim, TEXT("TriangleCount"), TriangleCount);
 
 			uint32 VertexCount = 0;
-			const UE::FUsdAttribute VertexCountAttr = PatternPrim.GetAttribute(TEXT("VertexCount"));
-			if (VertexCountAttr.HasValue() && VertexCountAttr.GetTypeName() == TEXT("uint"))
-			{
-				UE::FVtValue Value;
-				VertexCountAttr.Get(Value);
-				const TOptional<uint32> Optional = UsdUtils::GetUnderlyingValue<uint32>(Value);
-				VertexCount = Optional.IsSet() ? Optional.GetValue() : 0;
-			}
+			FillIntDatas(PatternPrim, TEXT("VertexCount"), VertexCount);
 
 			uint32 PatternId = 0;
-			const UE::FUsdAttribute PatternIdAttr = PatternPrim.GetAttribute(TEXT("PatternId"));
-			if (PatternIdAttr.HasValue() && PatternIdAttr.GetTypeName() == TEXT("uint"))
+			FillIntDatas(PatternPrim, TEXT("PatternId"), PatternId);
+
+			uint32 FabricIndex= 0;
+			if(UE::FUsdRelationship Relationship = PatternPrim.GetRelationship(TEXT("fabric")))
 			{
-				UE::FVtValue Value;
-				PatternIdAttr.Get(Value);
-				const TOptional<uint32> Optional = UsdUtils::GetUnderlyingValue<uint32>(Value);
-				PatternId = Optional.IsSet() ? Optional.GetValue() : 0;
+				TArray<UE::FSdfPath> TargetsPath;
+				if(Relationship.GetTargets(TargetsPath))
+				{
+					if (TargetsPath.Num() > 0)
+                    {
+						UE::FUsdPrim FabricPrim = PatternPrim.GetStage().GetPrimAtPath(TargetsPath[0]);
+
+						uint32 FabricId = 0;
+						FillIntDatas(FabricPrim, TEXT("FabricId"), FabricId);
+
+						// The fabric index referenced in the pattern is the index of the fabric in the managed array collection
+						FabricIndex = FabricIds.Find(FabricId);
+                    }
+				}
 			}
 
 			UE_LOG(LogChaosClothAssetDataflowNodes,
@@ -713,7 +834,7 @@ bool FChaosClothAssetUSDImportNode::ImportFromFile(const FString& UsdFilePath, c
 				// Add the new pattern
 				const int32 SimPatternIndex = ClothFacade.AddSimPattern();
 				FCollectionClothSimPatternFacade SimPattern = ClothFacade.GetSimPattern(SimPatternIndex);
-				SimPattern.Initialize(OutRestPositions2D, OutDrapedPositions3D, OutTriangleToVertexIndex);
+				SimPattern.Initialize(OutRestPositions2D, OutDrapedPositions3D, OutTriangleToVertexIndex, FabricIndex);
 
 				// Remap this pattern's stitches
 				check(SeamPatterns.Num() == SeamStitches.Num());
@@ -777,7 +898,7 @@ bool FChaosClothAssetUSDImportNode::ImportFromFile(const FString& UsdFilePath, c
 		FCollectionClothSeamFacade Seam = ClothFacade.AddGetSeam();
 		Seam.Initialize(SeamStitches[SeamIndex]);
 	}
-
+	
 	SlowTask.EnterProgressFrame(1.f);
 	SlowTask.ForceRefresh();
 	return true;
diff --git a/Engine/Plugins/Experimental/ChaosClothAssetEditor/Source/ChaosClothAssetDataflowNodes/Public/ChaosClothAsset/SimulationBaseConfigNode.h b/Engine/Plugins/Experimental/ChaosClothAssetEditor/Source/ChaosClothAssetDataflowNodes/Public/ChaosClothAsset/SimulationBaseConfigNode.h
--- a/Engine/Plugins/Experimental/ChaosClothAssetEditor/Source/ChaosClothAssetDataflowNodes/Public/ChaosClothAsset/SimulationBaseConfigNode.h
+++ b/Engine/Plugins/Experimental/ChaosClothAssetEditor/Source/ChaosClothAssetDataflowNodes/Public/ChaosClothAsset/SimulationBaseConfigNode.h
@@ -32,7 +32,7 @@ protected:
 	struct FPropertyHelper
 	{
 	public:
-		FPropertyHelper(const FChaosClothAssetSimulationBaseConfigNode& InConfigNode, Dataflow::FContext& InContext, ::Chaos::Softs::FCollectionPropertyMutableFacade& InProperties);
+		FPropertyHelper(const FChaosClothAssetSimulationBaseConfigNode& InConfigNode, Dataflow::FContext& InContext, ::Chaos::Softs::FCollectionPropertyMutableFacade& InProperties, const TSharedRef<FManagedArrayCollection>& InClothCollection);
 
 		template<typename PropertyType, TEMPLATE_REQUIRES(::Chaos::Softs::TIsWeightedType<PropertyType>::Value && !std::is_same<PropertyType, bool>::value)>
 		inline int32 SetProperty(const FName& PropertyName, const PropertyType& PropertyValue, const TArray<FName>& SimilarPropertyNames = {}, ECollectionPropertyFlags PropertyFlags = ECollectionPropertyFlags::Animatable);
@@ -64,16 +64,30 @@ protected:
 		int32 SetPropertyWeighted(const FName& PropertyName, const FChaosClothAssetWeightedValueNonAnimatable& PropertyValue, const TArray<FName>& SimilarPropertyNames = {}, ECollectionPropertyFlags PropertyFlags = ECollectionPropertyFlags::None);
 		int32 SetPropertyWeighted(const FName& PropertyName, const FChaosClothAssetWeightedValueNonAnimatableNoLowHighRange& PropertyValue, const TArray<FName>& SimilarPropertyNames = {}, ECollectionPropertyFlags PropertyFlags = ECollectionPropertyFlags::None);
 
+		template<typename WeightedValueType>
+		inline int32 SetPropertyWeighted(const FName& PropertyName, const UE::Chaos::ClothAsset::FWeightedValueBounds& PropertyBounds, const WeightedValueType& PropertyValue, const TArray<FName>& SimilarPropertyNames = {}, ECollectionPropertyFlags PropertyFlags = ECollectionPropertyFlags::None);
+		
 		template<typename T, typename WeightedValueType, TEMPLATE_REQUIRES(TIsDerivedFrom<T, FChaosClothAssetSimulationBaseConfigNode>::Value)>
 		inline int32 SetPropertyWeighted(const T* ConfigStruct, const WeightedValueType* PropertyValue, const TArray<FName>& SimilarPropertyNames = {}, ECollectionPropertyFlags PropertyFlags = ECollectionPropertyFlags::None);
 
+		FString GetPropertyString(const FString* PropertyReference) const;
+
+		TSharedRef<FManagedArrayCollection> GetClothCollection() const
+        {
+        	return ClothCollection;
+        }
+		
 	private:
+		
 		template<typename T, TEMPLATE_REQUIRES(TIsDerivedFrom<T, FChaosClothAssetSimulationBaseConfigNode>::Value)>
 		inline FName FindPropertyNameByAddress(const T* ConfigStruct, const void* PropertyAddress);
+		
+		int32 SetPropertyWeighted(const FName& PropertyName, const bool bIsAnimatable, const float& PropertyLow, const float& PropertyHigh, const FString& WeightMap, FString& MapOverride, const TArray<FName>& SimilarPropertyNames, ECollectionPropertyFlags PropertyFlags) const;
 
 		const FChaosClothAssetSimulationBaseConfigNode& ConfigNode;
 		Dataflow::FContext& Context;
 		::Chaos::Softs::FCollectionPropertyMutableFacade& Properties;
+		TSharedRef<FManagedArrayCollection> ClothCollection;
 	};
 
 	virtual void Evaluate(Dataflow::FContext& Context, const FDataflowOutput* Out) const override;
@@ -210,6 +224,13 @@ inline int32 FChaosClothAssetSimulationBaseConfigNode::FPropertyHelper::SetPrope
 	return SetPropertyWeighted(PropertyName, *PropertyValue, SimilarPropertyNames, PropertyFlags);
 }
 
+template<typename WeightedValueType>
+inline int32 FChaosClothAssetSimulationBaseConfigNode::FPropertyHelper::SetPropertyWeighted(const FName& PropertyName, const UE::Chaos::ClothAsset::FWeightedValueBounds& PropertyBounds, const WeightedValueType& PropertyValue, const TArray<FName>& SimilarPropertyNames, ECollectionPropertyFlags PropertyFlags)
+{
+	return SetPropertyWeighted(PropertyName, PropertyValue.bIsAnimatable, PropertyBounds.Low,
+		PropertyBounds.High, PropertyValue.WeightMap, PropertyValue.WeightMap_Override, SimilarPropertyNames, PropertyFlags);
+}
+
 template<typename T, typename ConnectableStringValueType, typename TEnableIf<TIsDerivedFrom<T, FChaosClothAssetSimulationBaseConfigNode>::Value, int>::type>
 inline int32 FChaosClothAssetSimulationBaseConfigNode::FPropertyHelper::SetPropertyString(
 	const T* ConfigStruct, 
diff --git a/Engine/Plugins/Experimental/ChaosClothAssetEditor/Source/ChaosClothAssetDataflowNodes/Public/ChaosClothAsset/SimulationBendingConfigNode.h b/Engine/Plugins/Experimental/ChaosClothAssetEditor/Source/ChaosClothAssetDataflowNodes/Public/ChaosClothAsset/SimulationBendingConfigNode.h
new file mode 100644
--- /dev/null
+++ b/Engine/Plugins/Experimental/ChaosClothAssetEditor/Source/ChaosClothAssetDataflowNodes/Public/ChaosClothAsset/SimulationBendingConfigNode.h
@@ -0,0 +1,162 @@
+// Copyright Epic Games, Inc. All Rights Reserved.
+
+#pragma once 
+
+#include "ChaosClothAsset/SimulationFabricConfigNode.h"
+#include "ChaosClothAsset/WeightedValue.h"
+#include "ChaosClothAsset/SimulationConfigNodePropertyTypes.h"
+#include "SimulationBendingConfigNode.generated.h"
+
+/** Bending constraint property configuration node. */
+USTRUCT(Meta = (DataflowCloth))
+struct FChaosClothAssetSimulationBendingConfigNode : public FChaosClothAssetSimulationFabricConfigNode
+{
+	GENERATED_USTRUCT_BODY()
+	DATAFLOW_NODE_DEFINE_INTERNAL(FChaosClothAssetSimulationBendingConfigNode, "SimulationBendingConfig", "Cloth", "Cloth Simulation Bending Config")
+
+public:
+	/**
+	 * Method for calculating the rest angles of the constraints.
+	 */
+	UPROPERTY(EditAnywhere, Category = "Bending Types")
+	EChaosClothAssetRestAngleConstructionType RestAngleType = EChaosClothAssetRestAngleConstructionType::Use3DRestAngles;
+
+	/**
+	 * Calculate rest angles as a ratio between completely flat and whatever is the 3D rest angle.
+	 * When FlatnessRatio = 0, this is equivalent to "Use3DRestAngles".
+	 * When FlatnessRatio = 1, the rest angle will be 0 (completely flat).
+	 * If a valid weight map is found with the given Weight Map name, then both Low and High values
+	 * are interpolated with the per particle weight to make the final value used for the simulation.
+	 * Otherwise all particles are considered to have a zero weight, and only the Low value is meaningful.
+	 */
+	UPROPERTY(EditAnywhere, Category = "Bending Properties", Meta = (UIMin = "0", UIMax = "1", ClampMin = "0", ClampMax = "1", EditCondition = "RestAngleType == EChaosClothAssetRestAngleConstructionType::FlatnessRatio", EditConditionHides))
+	FChaosClothAssetWeightedValueNonAnimatable FlatnessRatio = { 0.f, 0.f, TEXT("FlatnessRatio") };
+
+	/**
+	 * Set rest angles to be the explicit value set here (in degrees).
+	 * 0 = Flat, Positive values fold away from the edge normal, Negative values fold toward the edge normal.
+	 * When converting vertex weight values to edge values, the value with the smallest absolute value is selected.
+	 * If a valid weight map is found with the given Weight Map name, then both Low and High values
+	 * are interpolated with the per particle weight to make the final value used for the simulation.
+	 * Otherwise all particles are considered to have a zero weight, and only the Low value is meaningful.
+	 */
+	UPROPERTY(EditAnywhere, Category = "Bending Properties", Meta = (UIMin = "-180", UIMax = "180", ClampMin = "-180", ClampMax = "180", EditCondition = "RestAngleType == EChaosClothAssetRestAngleConstructionType::RestAngle", EditConditionHides))
+	FChaosClothAssetWeightedValueNonAnimatable RestAngle = { 0.f, 0.f, TEXT("RestAngle") };
+
+	/** Constraint solver type */
+	UPROPERTY(EditAnywhere, Category = "Bending Types", Meta = (EditCondition = "bImportProperties == false", EditConditionHides))
+	EChaosClothAssetConstraintSolverType SolverType = EChaosClothAssetConstraintSolverType::PBD;
+	
+	/** Constraint distribution type */
+    UPROPERTY(EditAnywhere, Category = "Bending Types", Meta = (EditCondition = "bImportProperties == false && SolverType == EChaosClothAssetConstraintSolverType::XPBD", EditConditionHides))
+    EChaosClothAssetConstraintDistributionType DistributionType = EChaosClothAssetConstraintDistributionType::Isotropic;
+	
+	/** Constraint method type */
+	UPROPERTY(EditAnywhere, Category = "Bending Types", Meta = (EditCondition = "bImportProperties == false && DistributionType == EChaosClothAssetConstraintDistributionType::Anisotropic", EditConditionHides))
+	EChaosClothAssetBendingConstraintType ConstraintType = EChaosClothAssetBendingConstraintType::HingeAngles;
+
+	/**
+	 * The Stiffness of the bending constraints. Increase the iteration count for stiffer materials.
+	 * If a valid weight map is found with the given Weight Map name, then both Low and High values
+	 * are interpolated with the per particle weight to make the final value used for the simulation.
+	 * Otherwise all particles are considered to have a zero weight, and only the Low value is meaningful.
+	 */
+	UPROPERTY(EditAnywhere, Category = "Bending Properties", Meta = (UIMin = "0", UIMax = "1", ClampMin = "0", ClampMax = "1"), Meta = (EditCondition = "bImportProperties == false && ((DistributionType == EChaosClothAssetConstraintDistributionType::Isotropic && SolverType == EChaosClothAssetConstraintSolverType::XPBD) || SolverType == EChaosClothAssetConstraintSolverType::PBD)", EditConditionHides))
+	FChaosClothAssetWeightedValue BendingStiffness = { true, 1.f, 1.f, TEXT("BendingStiffness") };
+	
+	/**
+	 * The stiffness of the bending constraints in the warp (vertical) direction.
+	 * If a valid weight map is found with the given Weight Map name, then both Low and High values
+	 * are interpolated with the per particle weight to make the final value used for the simulation.
+	 * Otherwise all particles are considered to have a zero weight, and only the Low value is meaningful.
+	 */
+	UPROPERTY(EditAnywhere, Category = "Bending Properties", Meta = (UIMin = "0", UIMax = "10000", ClampMin = "0", ClampMax = "10000000", EditCondition = "DistributionType == EChaosClothAssetConstraintDistributionType::Anisotropic && bImportProperties == false && SolverType == EChaosClothAssetConstraintSolverType::XPBD", EditConditionHides))
+	FChaosClothAssetWeightedValue BendingStiffnessWarp = { true, 100.f, 100.f, TEXT("BendingStiffnessWarp") };
+
+	/**
+	 * The stiffness of the bending constraints in the weft (horizontal) direction.
+	 * If a valid weight map is found with the given Weight Map name, then both Low and High values
+	 * are interpolated with the per particle weight to make the final value used for the simulation.
+	 * Otherwise all particles are considered to have a zero weight, and only the Low value is meaningful.
+	 */
+	UPROPERTY(EditAnywhere, Category = "Bending Properties", Meta = (UIMin = "0", UIMax = "10000", ClampMin = "0", ClampMax = "10000000", EditCondition = "DistributionType == EChaosClothAssetConstraintDistributionType::Anisotropic && bImportProperties == false && SolverType == EChaosClothAssetConstraintSolverType::XPBD", EditConditionHides))
+	FChaosClothAssetWeightedValue BendingStiffnessWeft = { true, 100.f, 100.f, TEXT("BendingStiffnessWeft") };
+
+	/**
+	 * The stiffness of the bending constraints in the bias (diagonal) direction.
+	 * If a valid weight map is found with the given Weight Map name, then both Low and High values
+	 * are interpolated with the per particle weight to make the final value used for the simulation.
+	 * Otherwise all particles are considered to have a zero weight, and only the Low value is meaningful.
+	 */
+	UPROPERTY(EditAnywhere, Category = "Bending Properties", Meta = (UIMin = "0", UIMax = "10000", ClampMin = "0", ClampMax = "10000000", EditCondition = "DistributionType == EChaosClothAssetConstraintDistributionType::Anisotropic && bImportProperties == false && SolverType == EChaosClothAssetConstraintSolverType::XPBD", EditConditionHides))
+	FChaosClothAssetWeightedValue BendingStiffnessBias = { true, 100.f, 100.f, TEXT("BendingStiffnessBias") };
+
+	/**
+	 * The damping of the bending constraints, relative to critical damping.
+	 * If a valid weight map is found with the given Weight Map name, then both Low and High values
+	 * are interpolated with the per particle weight to make the final value used for the simulation.
+	 * Otherwise all particles are considered to have a zero weight, and only the Low value is meaningful.
+	 */
+	UPROPERTY(EditAnywhere, Category = "Bending Properties", Meta = (UIMin = "0", UIMax = "10", ClampMin = "0", ClampMax = "1000"), Meta = (EditCondition = "bImportProperties == false && SolverType == EChaosClothAssetConstraintSolverType::XPBD", EditConditionHides))
+	FChaosClothAssetWeightedValue BendingDamping = { true, 1.f, 1.f, TEXT("BendingDamping") };
+
+	/**
+	 * Once the element has bent such that it's folded more than this ratio from its rest angle ("buckled"), switch to using Buckling Stiffness instead of BendingElement Stiffness.
+	 * When Buckling Ratio = 0, the Buckling Stiffness will never be used. When BucklingRatio = 1, the Buckling Stiffness will be used as soon as its bent past its rest configuration.
+	 */
+	UPROPERTY(EditAnywhere, Category = "Bending Properties", Meta = (UIMin = "0", UIMax = "1", ClampMin = "0", ClampMax = "1"), Meta = (EditCondition = "bImportProperties == false && ConstraintType == EChaosClothAssetBendingConstraintType::HingeAngles", EditConditionHides))
+	float BucklingRatio = 0.5f;
+
+	/**
+	 * The stiffness after buckling.
+	 * The constraint will use this stiffness instead of bending Stiffness once the cloth has buckled, i.e., bent beyond a certain angle.
+	 * Typically, Buckling Stiffness is set to be less than Bending Stiffness.
+	 * Buckling Ratio determines the switch point between using BendingElement Stiffness and Buckling Stiffness.
+	 * If a valid weight map is found with the given Weight Map name, then both Low and High values
+	 * are interpolated with the per particle weight to make the final value used for the simulation.
+	 * Otherwise all particles are considered to have a zero weight, and only the Low value is meaningful.
+	 */
+	UPROPERTY(EditAnywhere, Category = "Bending Properties", Meta = (UIMin = "0", UIMax = "1", ClampMin = "0", ClampMax = "1", EditCondition = "BucklingRatio != 0 && ((DistributionType == EChaosClothAssetConstraintDistributionType::Isotropic && SolverType == EChaosClothAssetConstraintSolverType::XPBD) || SolverType == EChaosClothAssetConstraintSolverType::PBD) && bImportProperties == false && ConstraintType == EChaosClothAssetBendingConstraintType::HingeAngles", EditConditionHides))
+	FChaosClothAssetWeightedValue BucklingStiffness = { true, 0.9f, 0.9f, TEXT("BucklingStiffness") };
+	
+	/**
+	 * The stiffness after buckling in the warp (vertical) direction.
+	 * The constraint will use this stiffness instead of element Stiffness once the cloth has buckled, i.e., bent beyond a certain angle.
+	 * Typically, Buckling Stiffness is set to be less than BendingElement Stiffness.
+	 * Buckling Ratio determines the switch point between using BendingElement Stiffness and Buckling Stiffness.
+	 * If a valid weight map is found with the given Weight Map name, then both Low and High values
+	 * are interpolated with the per particle weight to make the final value used for the simulation.
+	 * Otherwise all particles are considered to have a zero weight, and only the Low value is meaningful.
+	 */
+	UPROPERTY(EditAnywhere, Category = "Bending Properties", Meta = (UIMin = "0", UIMax = "10000", ClampMin = "0", ClampMax = "10000000", EditCondition = "BucklingRatio != 0 && DistributionType == EChaosClothAssetConstraintDistributionType::Anisotropic && bImportProperties == false && SolverType == EChaosClothAssetConstraintSolverType::XPBD", EditConditionHides))
+	FChaosClothAssetWeightedValue BucklingStiffnessWarp = { true, 50.f, 50.f, TEXT("BucklingStiffnessWarp") };
+
+	/**
+	 * The stiffness after buckling in the weft (horizontal) direction.
+	 * The constraint will use this stiffness instead of element Stiffness once the cloth has buckled, i.e., bent beyond a certain angle.
+	 * Typically, Buckling Stiffness is set to be less than BendingElement Stiffness.
+	 * Buckling Ratio determines the switch point between using BendingElement Stiffness and Buckling Stiffness.
+	 * If a valid weight map is found with the given Weight Map name, then both Low and High values
+	 * are interpolated with the per particle weight to make the final value used for the simulation.
+	 * Otherwise all particles are considered to have a zero weight, and only the Low value is meaningful.
+	 */
+	UPROPERTY(EditAnywhere, Category = "Bending Properties", Meta = (UIMin = "0", UIMax = "10000", ClampMin = "0", ClampMax = "10000000", EditCondition = "BucklingRatio != 0 && DistributionType == EChaosClothAssetConstraintDistributionType::Anisotropic && bImportProperties == false && SolverType == EChaosClothAssetConstraintSolverType::XPBD", EditConditionHides))
+	FChaosClothAssetWeightedValue BucklingStiffnessWeft = { true, 50.f, 50.f, TEXT("BucklingStiffnessWeft") };
+
+	/**
+	 * The stiffness after buckling in the bias (diagonal) direction.
+	 * The constraint will use this stiffness instead of element Stiffness once the cloth has buckled, i.e., bent beyond a certain angle.
+	 * Typically, Buckling Stiffness is set to be less than BendingElement Stiffness.
+	 * Buckling Ratio determines the switch point between using BendingElement Stiffness and Buckling Stiffness.
+	 * If a valid weight map is found with the given Weight Map name, then both Low and High values
+	 * are interpolated with the per particle weight to make the final value used for the simulation.
+	 * Otherwise all particles are considered to have a zero weight, and only the Low value is meaningful.
+	 */
+	UPROPERTY(EditAnywhere, Category = "Bending Properties", Meta = (UIMin = "0", UIMax = "10000", ClampMin = "0", ClampMax = "10000000", EditCondition = "BucklingRatio != 0 && DistributionType == EChaosClothAssetConstraintDistributionType::Anisotropic && bImportProperties == false && SolverType == EChaosClothAssetConstraintSolverType::XPBD", EditConditionHides))
+	FChaosClothAssetWeightedValue BucklingStiffnessBias = { true, 50.f, 50.f, TEXT("BucklingStiffnessBias") };
+
+	FChaosClothAssetSimulationBendingConfigNode(const Dataflow::FNodeParameters& InParam, FGuid InGuid = FGuid::NewGuid());
+
+private:
+	virtual void AddProperties(FPropertyHelper& PropertyHelper) const override;
+};
diff --git a/Engine/Plugins/Experimental/ChaosClothAssetEditor/Source/ChaosClothAssetDataflowNodes/Public/ChaosClothAsset/SimulationCollisionConfigNode.h b/Engine/Plugins/Experimental/ChaosClothAssetEditor/Source/ChaosClothAssetDataflowNodes/Public/ChaosClothAsset/SimulationCollisionConfigNode.h
--- a/Engine/Plugins/Experimental/ChaosClothAssetEditor/Source/ChaosClothAssetDataflowNodes/Public/ChaosClothAsset/SimulationCollisionConfigNode.h
+++ b/Engine/Plugins/Experimental/ChaosClothAssetEditor/Source/ChaosClothAssetDataflowNodes/Public/ChaosClothAsset/SimulationCollisionConfigNode.h
@@ -2,23 +2,23 @@
 
 #pragma once 
 
-#include "ChaosClothAsset/SimulationBaseConfigNode.h"
+#include "ChaosClothAsset/SimulationFabricConfigNode.h"
 #include "SimulationCollisionConfigNode.generated.h"
 
 /** Physics mesh collision properties configuration node. */
 USTRUCT(Meta = (DataflowCloth))
-struct FChaosClothAssetSimulationCollisionConfigNode : public FChaosClothAssetSimulationBaseConfigNode
+struct FChaosClothAssetSimulationCollisionConfigNode : public FChaosClothAssetSimulationFabricConfigNode
 {
 	GENERATED_USTRUCT_BODY()
 	DATAFLOW_NODE_DEFINE_INTERNAL(FChaosClothAssetSimulationCollisionConfigNode, "SimulationCollisionConfig", "Cloth", "Cloth Simulation Collision Config")
 
 public:
 	/** The added thickness of collision shapes. */
-	UPROPERTY(EditAnywhere, Category = "Collision Properties", Meta = (UIMin = "0", UIMax = "100", ClampMin = "0", ClampMax = "1000"))
+	UPROPERTY(EditAnywhere, Category = "Collision Properties", Meta = (UIMin = "0", UIMax = "100", ClampMin = "0", ClampMax = "1000", EditCondition = "bImportProperties == false", EditConditionHides))
 	float CollisionThickness = 1.0f;
 
 	/** Friction coefficient for cloth - collider interaction. */
-	UPROPERTY(EditAnywhere, Category = "Collision Properties", Meta = (UIMin = "0", UIMax = "1", ClampMin = "0", ClampMax = "10"))
+	UPROPERTY(EditAnywhere, Category = "Collision Properties", Meta = (UIMin = "0", UIMax = "1", ClampMin = "0", ClampMax = "10", EditCondition = "bImportProperties == false", EditConditionHides))
 	float FrictionCoefficient = 0.8f;
 
 	/** Stiffness for proximity repulsion forces (Force-based solver only). Units = kg cm/ s^2 (same as XPBD springs)*/
diff --git a/Engine/Plugins/Experimental/ChaosClothAssetEditor/Source/ChaosClothAssetDataflowNodes/Public/ChaosClothAsset/SimulationConfigNodePropertyTypes.h b/Engine/Plugins/Experimental/ChaosClothAssetEditor/Source/ChaosClothAssetDataflowNodes/Public/ChaosClothAsset/SimulationConfigNodePropertyTypes.h
new file mode 100644
--- /dev/null
+++ b/Engine/Plugins/Experimental/ChaosClothAssetEditor/Source/ChaosClothAssetDataflowNodes/Public/ChaosClothAsset/SimulationConfigNodePropertyTypes.h
@@ -0,0 +1,55 @@
+// Copyright Epic Games, Inc. All Rights Reserved.
+
+#pragma once 
+
+#include "UObject/NameTypes.h"
+#include "UObject/ObjectMacros.h"
+#include "Chaos/PBDBendingConstraintsBase.h"
+#include "SimulationConfigNodePropertyTypes.generated.h"
+
+UENUM()
+enum class EChaosClothAssetRestAngleConstructionType : uint8
+{
+	/** Calculate rest angles using the 3D draped space simulation mesh. */
+	Use3DRestAngles = (uint8)Chaos::Softs::FPBDBendingConstraintsBase::ERestAngleConstructionType::Use3DRestAngles,
+	/** Calculate rest angles using the FlatnessRatio property. */
+	FlatnessRatio = (uint8)Chaos::Softs::FPBDBendingConstraintsBase::ERestAngleConstructionType::FlatnessRatio,
+	/** Calculate rest angles using the RestAngle property. */
+	RestAngle = (uint8)Chaos::Softs::FPBDBendingConstraintsBase::ERestAngleConstructionType::ExplicitRestAngles
+};
+
+UENUM()
+enum class EChaosClothAssetBendingConstraintType : uint8
+{
+	/** Add a spring in between the opposite vertices of the adjacent faces */
+	FacesSpring,
+	/** Add an angular constraint in between the 2 faces */
+	HingeAngles,
+};
+
+UENUM()
+enum class EChaosClothAssetStretchConstraintType : uint8
+{
+	/** Add a UV style stretch and shear constraint */
+	StretchShear,
+	/** Add edge and area springs in the 2d plane */
+	EdgeArea,
+};
+
+UENUM()
+enum class EChaosClothAssetConstraintDistributionType : uint8
+{
+	/** Having a separate stiffness along the warp, weft, bias direction */
+	Anisotropic,
+	/** Uniform stiffness in all directions */
+	Isotropic,
+};
+
+UENUM()
+enum class EChaosClothAssetConstraintSolverType : uint8
+{
+	/** Physically based PBD with compliance */
+	XPBD,
+	/** PBD style with no damping/stiffness */
+	PBD,
+};
diff --git a/Engine/Plugins/Experimental/ChaosClothAssetEditor/Source/ChaosClothAssetDataflowNodes/Public/ChaosClothAsset/SimulationFabricConfigNode.h b/Engine/Plugins/Experimental/ChaosClothAssetEditor/Source/ChaosClothAssetDataflowNodes/Public/ChaosClothAsset/SimulationFabricConfigNode.h
new file mode 100644
--- /dev/null
+++ b/Engine/Plugins/Experimental/ChaosClothAssetEditor/Source/ChaosClothAssetDataflowNodes/Public/ChaosClothAsset/SimulationFabricConfigNode.h
@@ -0,0 +1,47 @@
+// Copyright Epic Games, Inc. All Rights Reserved.
+
+#pragma once 
+
+#include "ChaosClothAsset/SimulationBaseConfigNode.h"
+#include "ChaosClothAsset/CollectionClothFacade.h"
+#include "SimulationFabricConfigNode.generated.h"
+
+/** Cloth configuration base node to import properties from fabrics . */
+USTRUCT(Meta = (DataflowCloth))
+struct FChaosClothAssetSimulationFabricConfigNode : public FChaosClothAssetSimulationBaseConfigNode
+{
+	GENERATED_USTRUCT_BODY()
+
+public:
+	
+	/** Import properties to boolean the  fabric override */
+	UPROPERTY(EditAnywhere, Category = "Fabric Properties")
+	bool bImportProperties = false;
+	
+	FChaosClothAssetSimulationFabricConfigNode() = default;
+	
+	FChaosClothAssetSimulationFabricConfigNode(const Dataflow::FNodeParameters& InParam, FGuid InGuid = FGuid::NewGuid());
+	
+protected :
+	
+	// FChaosClothAssetSimulationBaseConfigNode interface
+	virtual void AddProperties(struct FPropertyHelper& PropertyHelper) const override {}
+	
+	/** We can only use fabrics to set property values if the import settings is true and if the cloth facade has fabrics */
+	bool CanUseFabrics(const UE::Chaos::ClothAsset::FCollectionClothFacade& ClothFacade) const
+	{
+		return bImportProperties && ClothFacade.IsValid() && (ClothFacade.GetNumFabrics() > 0);
+	}
+	
+	/** Set an imported fabric value onto the weighted value property (animatable or not) */
+	template<typename PropertyType>
+	static void SetFabricPropertyWeighted(const FName& PropertyName, const PropertyType& PropertyValue, UE::Chaos::ClothAsset::FCollectionClothFacade& ClothFacade,
+		FPropertyHelper& PropertyHelper, const TFunction<float(const UE::Chaos::ClothAsset::FCollectionClothFabricFacade&)>& FabricValueFunction,
+		const TArray<FName>& SimilarPropertyNames, ECollectionPropertyFlags PropertyFlags = ECollectionPropertyFlags::None);
+	
+	/** Set an imported fabric value onto the float property */
+	static void SetFabricProperty(const FName& PropertyName, UE::Chaos::ClothAsset::FCollectionClothFacade& ClothFacade,
+		FPropertyHelper& PropertyHelper, const TFunction<float(const UE::Chaos::ClothAsset::FCollectionClothFabricFacade&)>& FabricValueFunction,
+		const TArray<FName>& SimilarPropertyNames, ECollectionPropertyFlags PropertyFlags = ECollectionPropertyFlags::None);
+};
+
diff --git a/Engine/Plugins/Experimental/ChaosClothAssetEditor/Source/ChaosClothAssetDataflowNodes/Public/ChaosClothAsset/SimulationMassConfigNode.h b/Engine/Plugins/Experimental/ChaosClothAssetEditor/Source/ChaosClothAssetDataflowNodes/Public/ChaosClothAsset/SimulationMassConfigNode.h
--- a/Engine/Plugins/Experimental/ChaosClothAssetEditor/Source/ChaosClothAssetDataflowNodes/Public/ChaosClothAsset/SimulationMassConfigNode.h
+++ b/Engine/Plugins/Experimental/ChaosClothAssetEditor/Source/ChaosClothAssetDataflowNodes/Public/ChaosClothAsset/SimulationMassConfigNode.h
@@ -2,13 +2,13 @@
 
 #pragma once 
 
-#include "ChaosClothAsset/SimulationBaseConfigNode.h"
+#include "ChaosClothAsset/SimulationFabricConfigNode.h"
 #include "ClothConfig.h"  // For EClothMassMode
 #include "SimulationMassConfigNode.generated.h"
 
 /** Mass properties configuration node. */
 USTRUCT(Meta = (DataflowCloth))
-struct FChaosClothAssetSimulationMassConfigNode : public FChaosClothAssetSimulationBaseConfigNode
+struct FChaosClothAssetSimulationMassConfigNode : public FChaosClothAssetSimulationFabricConfigNode
 {
 	GENERATED_USTRUCT_BODY()
 	DATAFLOW_NODE_DEFINE_INTERNAL(FChaosClothAssetSimulationMassConfigNode, "SimulationMassConfig", "Cloth", "Cloth Simulation Mass Config")
@@ -20,15 +20,15 @@ public:
 	 * -	Total Mass: The total mass is distributed equally over all the particles. Useful when referencing a specific garment size and feel.
 	 * -	Density: A constant mass density is used. Density is usually the preferred way of setting mass since it allows matching real life materials' density values.
 	 */
-	UPROPERTY(EditAnywhere, Category = "Mass Properties")
+	UPROPERTY(EditAnywhere, Category = "Mass Properties", meta = (EditCondition = "bImportProperties == false"))
 	EClothMassMode MassMode = EClothMassMode::Density;
 
 	/** The value used when the Mass Mode is set to Uniform Mass. */
-	UPROPERTY(EditAnywhere, Category = "Mass Properties", DisplayName = "Uniform Mass", Meta = (UIMin = "0.000001", UIMax = "0.001", ClampMin = "0", EditCondition = "MassMode == EClothMassMode::UniformMass"))
+	UPROPERTY(EditAnywhere, Category = "Mass Properties", DisplayName = "Uniform Mass", Meta = (UIMin = "0.000001", UIMax = "0.001", ClampMin = "0", EditCondition = "MassMode == EClothMassMode::UniformMass && bImportProperties == false"))
 	FChaosClothAssetWeightedValueNonAnimatable UniformMassWeighted = {0.00015f, 0.00015f, TEXT("UniformMass")};
 
 	/** The value used when Mass Mode is set to TotalMass. */
-	UPROPERTY(EditAnywhere, Category = "Mass Properties", Meta = (UIMin = "0.001", UIMax = "10", ClampMin = "0", EditCondition = "MassMode == EClothMassMode::TotalMass"))
+	UPROPERTY(EditAnywhere, Category = "Mass Properties", Meta = (UIMin = "0.001", UIMax = "10", ClampMin = "0", EditCondition = "MassMode == EClothMassMode::TotalMass && bImportProperties == false"))
 	float TotalMass = 0.5f;
 
 	/**
@@ -41,7 +41,7 @@ public:
 	 * Cotton: 0.2
 	 * Silk: 0.1
 	 */
-	UPROPERTY(EditAnywhere, Category = "Mass Properties", DisplayName = "Density", meta = (UIMin = "0.001", UIMax = "1", ClampMin = "0", EditCondition = "MassMode == EClothMassMode::Density"))
+	UPROPERTY(EditAnywhere, Category = "Mass Properties", DisplayName = "Density", meta = (UIMin = "0.001", UIMax = "1", ClampMin = "0", EditCondition = "MassMode == EClothMassMode::Density && bImportProperties == false"))
 	FChaosClothAssetWeightedValueNonAnimatable DensityWeighted = {0.35f, 0.35f, TEXT("Density")};
 
 	/** Calculated particle masses will be clamped to this minimum value (or 1e-8, whichever is larger). */
diff --git a/Engine/Plugins/Experimental/ChaosClothAssetEditor/Source/ChaosClothAssetDataflowNodes/Public/ChaosClothAsset/SimulationPBDBendingElementConfigNode.h b/Engine/Plugins/Experimental/ChaosClothAssetEditor/Source/ChaosClothAssetDataflowNodes/Public/ChaosClothAsset/SimulationPBDBendingElementConfigNode.h
--- a/Engine/Plugins/Experimental/ChaosClothAssetEditor/Source/ChaosClothAssetDataflowNodes/Public/ChaosClothAsset/SimulationPBDBendingElementConfigNode.h
+++ b/Engine/Plugins/Experimental/ChaosClothAssetEditor/Source/ChaosClothAssetDataflowNodes/Public/ChaosClothAsset/SimulationPBDBendingElementConfigNode.h
@@ -4,20 +4,9 @@
 
 #include "ChaosClothAsset/SimulationBaseConfigNode.h"
 #include "ChaosClothAsset/WeightedValue.h"
-#include "Chaos/PBDBendingConstraintsBase.h"
+#include "ChaosClothAsset/SimulationConfigNodePropertyTypes.h"
 #include "SimulationPBDBendingElementConfigNode.generated.h"
 
-UENUM()
-enum class EChaosClothAssetRestAngleConstructionType : uint8
-{
-	/** Calculate rest angles using the 3D draped space simulation mesh. */
-	Use3DRestAngles = (uint8)Chaos::Softs::FPBDBendingConstraintsBase::ERestAngleConstructionType::Use3DRestAngles,
-	/** Calculate rest angles using the FlatnessRatio property. */
-	FlatnessRatio = (uint8)Chaos::Softs::FPBDBendingConstraintsBase::ERestAngleConstructionType::FlatnessRatio,
-	/** Calculate rest angles using the RestAngle property. */
-	RestAngle = (uint8)Chaos::Softs::FPBDBendingConstraintsBase::ERestAngleConstructionType::ExplicitRestAngles
-};
-
 /** Bending element constraint property configuration node. */
 USTRUCT(Meta = (DataflowCloth))
 struct FChaosClothAssetSimulationPBDBendingElementConfigNode : public FChaosClothAssetSimulationBaseConfigNode
diff --git a/Engine/Plugins/Experimental/ChaosClothAssetEditor/Source/ChaosClothAssetDataflowNodes/Public/ChaosClothAsset/SimulationStretchConfigNode.h b/Engine/Plugins/Experimental/ChaosClothAssetEditor/Source/ChaosClothAssetDataflowNodes/Public/ChaosClothAsset/SimulationStretchConfigNode.h
new file mode 100644
--- /dev/null
+++ b/Engine/Plugins/Experimental/ChaosClothAssetEditor/Source/ChaosClothAssetDataflowNodes/Public/ChaosClothAsset/SimulationStretchConfigNode.h
@@ -0,0 +1,116 @@
+// Copyright Epic Games, Inc. All Rights Reserved.
+
+#pragma once 
+
+#include "ChaosClothAsset/SimulationFabricConfigNode.h"
+#include "ChaosClothAsset/WeightedValue.h"
+#include "ChaosClothAsset/SimulationConfigNodePropertyTypes.h"
+#include "SimulationStretchConfigNode.generated.h"
+
+/** Stretching constraint property configuration node. */
+USTRUCT(Meta = (DataflowCloth))
+struct FChaosClothAssetSimulationStretchConfigNode : public FChaosClothAssetSimulationFabricConfigNode
+{
+	GENERATED_USTRUCT_BODY()
+	DATAFLOW_NODE_DEFINE_INTERNAL(FChaosClothAssetSimulationStretchConfigNode, "SimulationStretchConfig", "Cloth", "Cloth Simulation Stretching Config")
+
+public:
+	/** Whether to use the 3D draped space as rest lengths, or use the 2D pattern space instead. */
+	UPROPERTY(EditAnywhere, Category = "Stretch Types")
+	bool bStretchUse3dRestLengths = true;
+
+	/** Constraint solver type.  */
+	UPROPERTY(EditAnywhere, Category = "Stretch Types", Meta = (EditCondition = "bImportProperties == false", EditConditionHides))
+	EChaosClothAssetConstraintSolverType SolverType = EChaosClothAssetConstraintSolverType::PBD;
+
+	/**  Constraint distribution type. */
+	UPROPERTY(EditAnywhere, Category = "Stretch Types", Meta = (EditCondition = "bImportProperties == false && SolverType == EChaosClothAssetConstraintSolverType::XPBD", EditConditionHides))
+	EChaosClothAssetConstraintDistributionType DistributionType = EChaosClothAssetConstraintDistributionType::Isotropic;
+	
+	/** Constraint stretching type.  */
+	UPROPERTY(EditAnywhere, Category = "Stretch Types", Meta = (EditCondition = "bImportProperties == false && SolverType == EChaosClothAssetConstraintSolverType::XPBD && DistributionType == EChaosClothAssetConstraintDistributionType::Anisotropic", EditConditionHides))
+	EChaosClothAssetStretchConstraintType ConstraintType = EChaosClothAssetStretchConstraintType::StretchShear;
+	
+	/** Add an area constraint in case of isotropic distribution  */
+	UPROPERTY(EditAnywhere, Category = "Stretch Types", Meta = (EditCondition = "bImportProperties == false && ((SolverType == EChaosClothAssetConstraintSolverType::XPBD && DistributionType == EChaosClothAssetConstraintDistributionType::Isotropic) || SolverType == EChaosClothAssetConstraintSolverType::PBD)", EditConditionHides))
+	bool bAddAreaConstraint = true;
+
+	/**
+	 * The stiffness of the stretch constraints.
+	 * If a valid weight map is found with the given Weight Map name, then both Low and High values
+	 * are interpolated with the per particle weight to make the final value used for the simulation.
+	 * Otherwise all particles are considered to have a zero weight, and only the Low value is meaningful.
+	 */
+	UPROPERTY(EditAnywhere, Category = "Stretch Properties", Meta = (UIMin = "0", UIMax = "10000", ClampMin = "0", ClampMax = "1000000000", EditCondition = "bImportProperties == false && (SolverType == EChaosClothAssetConstraintSolverType::PBD || DistributionType == EChaosClothAssetConstraintDistributionType::Isotropic)", EditConditionHides))
+	FChaosClothAssetWeightedValue StretchStiffness = { true, 100.f, 100.f, TEXT("StretchStiffness") };
+	
+	/**
+	 * The stiffness of the stretch constraints in the warp (vertical) direction.
+	 * If a valid weight map is found with the given Weight Map name, then both Low and High values
+	 * are interpolated with the per particle weight to make the final value used for the simulation.
+	 * Otherwise all particles are considered to have a zero weight, and only the Low value is meaningful.
+	 */
+	UPROPERTY(EditAnywhere, Category = "Stretch Properties", Meta = (UIMin = "0", UIMax = "10000", ClampMin = "0", ClampMax = "10000000", EditCondition = "bImportProperties == false && DistributionType == EChaosClothAssetConstraintDistributionType::Anisotropic && SolverType == EChaosClothAssetConstraintSolverType::XPBD", EditConditionHides))
+	FChaosClothAssetWeightedValue StretchStiffnessWarp = { true, 100.f, 100.f, TEXT("StretchStiffnessWarp") };
+
+	/**
+	 * The stiffness of the stretch constraints in the weft (horizontal) direction.
+	 * If a valid weight map is found with the given Weight Map name, then both Low and High values
+	 * are interpolated with the per particle weight to make the final value used for the simulation.
+	 * Otherwise all particles are considered to have a zero weight, and only the Low value is meaningful.
+	 */
+	UPROPERTY(EditAnywhere, Category = "Stretch Properties", Meta = (UIMin = "0", UIMax = "10000", ClampMin = "0", ClampMax = "10000000", EditCondition = "bImportProperties == false && DistributionType == EChaosClothAssetConstraintDistributionType::Anisotropic && SolverType == EChaosClothAssetConstraintSolverType::XPBD", EditConditionHides))
+	FChaosClothAssetWeightedValue StretchStiffnessWeft = { true, 100.f, 100.f, TEXT("StretchStiffnessWeft") };
+
+	/**
+	 * The stiffness of the stretch constraints in the bias (diagonal) direction.
+	 * If a valid weight map is found with the given Weight Map name, then both Low and High values
+	 * are interpolated with the per particle weight to make the final value used for the simulation.
+	 * Otherwise all particles are considered to have a zero weight, and only the Low value is meaningful.
+	 */
+	UPROPERTY(EditAnywhere, Category = "Stretch Properties", Meta = (UIMin = "0", UIMax = "10000", ClampMin = "0", ClampMax = "10000000", EditCondition = "bImportProperties == false && DistributionType == EChaosClothAssetConstraintDistributionType::Anisotropic && SolverType == EChaosClothAssetConstraintSolverType::XPBD", EditConditionHides))
+	FChaosClothAssetWeightedValue StretchStiffnessBias = { true, 100.f, 100.f, TEXT("StretchStiffnessBias") };
+
+	/**
+	 * 
+	 * The damping of the stretch constraints, relative to critical damping.
+	 * If a valid weight map is found with the given Weight Map name, then both Low and High values
+	 * are interpolated with the per particle weight to make the final value used for the simulation.
+	 * Otherwise all particles are considered to have a zero weight, and only the Low value is meaningful.
+	 */
+	UPROPERTY(EditAnywhere, Category = "Stretch Properties", Meta = (UIMin = "0", UIMax = "10", ClampMin = "0", ClampMax = "1000", EditCondition = "bImportProperties == false && SolverType == EChaosClothAssetConstraintSolverType::XPBD", EditConditionHides))
+	FChaosClothAssetWeightedValue StretchDamping = { true, 1.f, 1.f, TEXT("StretchDamping") };
+
+	/**
+	 * The scale of the stretch constraints at rest in the warp (vertical) direction.
+	 * If a valid weight map is found with the given Weight Map name, then both Low and High values
+	 * are interpolated with the per particle weight to make the final value used for the simulation.
+	 * Otherwise all particles are considered to have a zero weight, and only the Low value is meaningful.
+	 */
+	UPROPERTY(EditAnywhere, Category = "Stretch Properties", Meta = (UIMin = "0", UIMax = "10", ClampMin = "0", ClampMax = "10000000", EditCondition = "bImportProperties == true || (SolverType == EChaosClothAssetConstraintSolverType::XPBD && DistributionType == EChaosClothAssetConstraintDistributionType::Anisotropic)", EditConditionHides))
+	FChaosClothAssetWeightedValue StretchWarpScale = { true, 1.f, 1.f, TEXT("StretchWarpScale") };
+
+	/**
+	 * The scale of the stretch constraints at rest in the weft (horizontal) direction.
+	 * If a valid weight map is found with the given Weight Map name, then both Low and High values
+	 * are interpolated with the per particle weight to make the final value used for the simulation.
+	 * Otherwise all particles are considered to have a zero weight, and only the Low value is meaningful.
+	 */
+	UPROPERTY(EditAnywhere, Category = "Stretch Properties", Meta = (UIMin = "0", UIMax = "10", ClampMin = "0", ClampMax = "10000000", EditCondition = "bImportProperties == true || (SolverType == EChaosClothAssetConstraintSolverType::XPBD && DistributionType == EChaosClothAssetConstraintDistributionType::Anisotropic)", EditConditionHides))
+	FChaosClothAssetWeightedValue StretchWeftScale = { true, 1.f, 1.f, TEXT("StretchWeftScale") };
+
+	/**
+	 * The stiffness of the surface area preservation constraints.
+	 * Increase the solver iteration count for stiffer materials.
+	 * If a valid weight map is found with the given Weight Map name, then both Low and High values
+	 * are interpolated with the per particle weight to make the final value used for the simulation.
+	 * Otherwise all particles are considered to have a zero weight, and only the Low value is meaningful.
+	 */
+	UPROPERTY(EditAnywhere, Category = "Stretch Properties", Meta = (UIMin = "0", UIMax = "10000", ClampMin = "0", ClampMax = "10000000", EditCondition = "bImportProperties == false && bAddAreaConstraint == true && ((SolverType == EChaosClothAssetConstraintSolverType::XPBD && DistributionType == EChaosClothAssetConstraintDistributionType::Isotropic) || SolverType == EChaosClothAssetConstraintSolverType::PBD)", EditConditionHides))
+	FChaosClothAssetWeightedValue AreaStiffness = { true, 100.f, 100.f, TEXT("AreaStiffness") };
+
+	FChaosClothAssetSimulationStretchConfigNode(const Dataflow::FNodeParameters& InParam, FGuid InGuid = FGuid::NewGuid());
+
+private:
+	virtual void AddProperties(FPropertyHelper& PropertyHelper) const override;
+};
diff --git a/Engine/Plugins/Experimental/ChaosClothAssetEditor/Source/ChaosClothAssetDataflowNodes/Public/ChaosClothAsset/WeightedValue.h b/Engine/Plugins/Experimental/ChaosClothAssetEditor/Source/ChaosClothAssetDataflowNodes/Public/ChaosClothAsset/WeightedValue.h
--- a/Engine/Plugins/Experimental/ChaosClothAssetEditor/Source/ChaosClothAssetDataflowNodes/Public/ChaosClothAsset/WeightedValue.h
+++ b/Engine/Plugins/Experimental/ChaosClothAssetEditor/Source/ChaosClothAssetDataflowNodes/Public/ChaosClothAsset/WeightedValue.h
@@ -22,6 +22,13 @@ namespace UE::Chaos::ClothAsset
 		 */
 		static CHAOSCLOTHASSETDATAFLOWNODES_API void MakeWeightMapName(FString& InOutString);
 	};
+	
+	/** Weight value bounds that could be used to build a cloth asset weighted value */
+	struct FWeightedValueBounds
+	{
+		float Low;
+		float High;
+	};
 }
 
 USTRUCT()
