Cloth Editor Remesh Node: Add a Simplify option for Render mesh decimation. When enabled, this will use FSimplifyMeshOp rather than FRemeshMeshOp. The Simplifier uses QEM to prioritize removing vertices in flat areas to preserve curvature. The result is typically a lower poly count but less uniform mesh than FRemeshOp produces. Because of this, it's only enabled for the Render mesh.

diff --git a/Engine/Plugins/Experimental/ChaosClothAssetEditor/Source/ChaosClothAssetDataflowNodes/ChaosClothAssetDataflowNodes.Build.cs b/Engine/Plugins/Experimental/ChaosClothAssetEditor/Source/ChaosClothAssetDataflowNodes/ChaosClothAssetDataflowNodes.Build.cs
--- a/Engine/Plugins/Experimental/ChaosClothAssetEditor/Source/ChaosClothAssetDataflowNodes/ChaosClothAssetDataflowNodes.Build.cs
+++ b/Engine/Plugins/Experimental/ChaosClothAssetEditor/Source/ChaosClothAssetDataflowNodes/ChaosClothAssetDataflowNodes.Build.cs
@@ -39,6 +39,7 @@ public class ChaosClothAssetDataflowNodes : ModuleRules
 				"MeshConversion",
 				"MeshDescription",
 				"MeshUtilitiesCommon",
+				"ModelingOperatorsEditorOnly",	// TODO: Someday remove editor dependencies, see UE-206172
 				"ModelingOperators",
 				"RenderCore",
 				"SkeletalMeshDescription",
diff --git a/Engine/Plugins/Experimental/ChaosClothAssetEditor/Source/ChaosClothAssetDataflowNodes/Private/ChaosClothAsset/RemeshNode.cpp b/Engine/Plugins/Experimental/ChaosClothAssetEditor/Source/ChaosClothAssetDataflowNodes/Private/ChaosClothAsset/RemeshNode.cpp
--- a/Engine/Plugins/Experimental/ChaosClothAssetEditor/Source/ChaosClothAssetDataflowNodes/Private/ChaosClothAsset/RemeshNode.cpp
+++ b/Engine/Plugins/Experimental/ChaosClothAssetEditor/Source/ChaosClothAssetDataflowNodes/Private/ChaosClothAsset/RemeshNode.cpp
@@ -13,10 +13,12 @@
 #include "DynamicMesh/MeshNormals.h"
 #include "DynamicMesh/DynamicVertexSkinWeightsAttribute.h"
 #include "CleaningOps/RemeshMeshOp.h"
+#include "CleaningOps/SimplifyMeshOp.h"
 #include "MeshUVChannelInfo.h"
 #include "MeshBoundaryLoops.h"
 #include "Chaos/CollectionPropertyFacade.h"
 #include "Algo/Find.h"
+#include "IMeshReductionManagerModule.h"
 
 #include UE_INLINE_GENERATED_CPP_BY_NAME(RemeshNode)
 
@@ -603,7 +605,7 @@ namespace UE::Chaos::ClothAsset::Private
 		// (i.e. things we might wish to add to the node properties in the future)
 		//
 
-		constexpr bool bReprojectToInputMesh = false;
+		constexpr bool bReprojectToInputMesh = true;
 		constexpr bool bDiscardAttributes = false;
 		constexpr bool bUseFullRemeshPasses = true;
 		constexpr bool bAllowFlips = true;
@@ -765,6 +767,92 @@ namespace UE::Chaos::ClothAsset::Private
 		}
 	}
 
+
+
+	bool Simplify(UE::Geometry::FDynamicMesh3& Mesh, int TargetVertexCount, UE::Geometry::FCompactMaps* CompactMaps = nullptr)
+	{
+		using namespace UE::Geometry;
+
+		//
+		// These consts control overall remeshing behavior and are analogs of the properties exposed to the user in the actual Simplify tool
+		// (i.e. things we might wish to add to the node properties in the future)
+		//
+
+		constexpr ESimplifyType SimplifierType = ESimplifyType::Attribute;
+		constexpr bool bDiscardAttributes = false;
+		constexpr bool bPreventNormalFlips = true;
+		constexpr bool bPreserveSharpEdges = false;
+		constexpr bool bPreventTinyTriangles = false;
+		constexpr bool bReproject = true;
+		constexpr bool bAutoCompact = true;
+		constexpr bool bGeometricConstraint = false;
+
+		// Angle threshold in degrees used for testing if two triangles should be considered coplanar, or two lines collinear */
+		constexpr float MinimalAngleThreshold = 0.01f;
+
+		// Note PolyEdgeAngleTolerance is very similar to MinimalAngleThreshold, but not redundant b/c the useful ranges are very different (MinimalAngleThreshold should generally be kept very small)
+		// Threshold angle change (in degrees) along a polygroup edge, above which a vertex must be added
+		constexpr float PolyEdgeAngleTolerance = 0.1f;
+
+
+		FSimplifyMeshOp Op;
+
+		Op.bDiscardAttributes = bDiscardAttributes;
+		Op.bResultMustHaveAttributesEnabled = true;
+		Op.bPreventNormalFlips = bPreventNormalFlips;
+		Op.bPreserveSharpEdges = bPreserveSharpEdges;
+		Op.bAllowSeamCollapse = !bPreserveSharpEdges;
+		Op.bPreventTinyTriangles = bPreventTinyTriangles;
+		Op.bReproject = bReproject;
+		Op.SimplifierType = SimplifierType;
+		Op.MinimalPlanarAngleThresh = MinimalAngleThreshold;
+
+		Op.TargetMode = ESimplifyTargetType::VertexCount;
+		Op.TargetCount = TargetVertexCount;
+
+		Op.MeshBoundaryConstraint = EEdgeRefineFlags::CollapseOnly;
+		Op.GroupBoundaryConstraint = EEdgeRefineFlags::CollapseOnly;
+		Op.MaterialBoundaryConstraint = EEdgeRefineFlags::CollapseOnly;
+
+		Op.bGeometricDeviationConstraint = bGeometricConstraint;
+		Op.GeometricTolerance = 0.0f;
+		Op.PolyEdgeAngleTolerance = PolyEdgeAngleTolerance;
+
+		TSharedPtr<FDynamicMesh3> SourceMesh = MakeShared<FDynamicMesh3>(MoveTemp(Mesh));
+		TSharedPtr<FDynamicMeshAABBTree3> SourceSpatial;
+		if (bReproject)
+		{
+			// acceleration structure is only used for reprojecting
+			SourceSpatial = MakeShared<FDynamicMeshAABBTree3>(SourceMesh.Get(), true);
+		}
+		Op.OriginalMesh = SourceMesh;
+		Op.OriginalMeshSpatial = SourceSpatial;
+
+		IMeshReductionManagerModule& MeshReductionModule = FModuleManager::Get().LoadModuleChecked<IMeshReductionManagerModule>("MeshReductionInterface");
+		Op.MeshReduction = MeshReductionModule.GetStaticMeshReductionInterface();
+
+		constexpr FProgressCancel* Progress = nullptr;		// Don't allow cancel or report progress for now
+		Op.CalculateResult(Progress);
+
+		if (Op.GetResultInfo().Result == EGeometryResultType::Success)
+		{
+			TUniquePtr<FDynamicMesh3> ResultMesh = Op.ExtractResult();
+			Mesh = MoveTemp(*ResultMesh);
+		}
+		else
+		{
+			return false;
+		}
+
+		// compact the input mesh if enabled
+		if (bAutoCompact)
+		{
+			Mesh.CompactInPlace(CompactMaps);
+		}
+
+		return true;
+	}
+
 }
 
 
@@ -1076,84 +1164,73 @@ void FChaosClothAssetRemeshNode::EmptyRenderSelections(const TSharedRef<FManaged
 	}
 }
 
-void FChaosClothAssetRemeshNode::RemeshRenderPattern(const TSharedRef<const FManagedArrayCollection>& ClothCollection,
-	int32 PatternIndex,
+
+
+
+void FChaosClothAssetRemeshNode::RemeshRenderMesh(const TSharedRef<const FManagedArrayCollection>& ClothCollection,
 	const TSharedRef<FManagedArrayCollection>& OutClothCollection) const
 {
 	using namespace UE::Geometry;
 	using namespace UE::Chaos::ClothAsset;
 
-	// Get the source mesh for the pattern
+	// Get the source mesh
 	FClothPatternToDynamicMesh Converter;
-
 	FDynamicMesh3 DynamicMesh;
-	Converter.Convert(ClothCollection, PatternIndex, EClothPatternVertexType::Render, DynamicMesh);
 
-	const bool bHasUVs = DynamicMesh.HasAttributes() && DynamicMesh.Attributes() && DynamicMesh.Attributes()->PrimaryUV();
+	// NOTE: When applied to the Render mesh, this Convert function will assign PatternIDs to the MaterialID attribute of the DynamicMesh. 
+	// After remeshing we will use the MatrialID attribute to determine which triangles should go into which output pattern.
+	Converter.Convert(ClothCollection, INDEX_NONE, EClothPatternVertexType::Render, DynamicMesh);
 
-	const double PatternArea = TMeshQueries<FDynamicMesh3>::GetVolumeArea(DynamicMesh).Y;
-	const int PatternTriangleCount = DynamicMesh.TriangleCount();
-	const int TargetTriangleCount = FMath::RoundToInt(static_cast<float>(TargetPercentRender) / 100.0f * static_cast<float>(PatternTriangleCount));
-	const double TargetEdgeLength = FRemeshMeshOp::CalculateTargetEdgeLength(nullptr, TargetTriangleCount, PatternArea);
-
-	TArray<TArray<FIntVector2>> Seams;
-
-	for (int32 ResampleIter = 0; ResampleIter < IterationsRender; ++ResampleIter)
-	{
-		Private::RemeshBoundaries(DynamicMesh, Seams, TargetEdgeLength);
-	}
+	check(DynamicMesh.HasAttributes());
 
-	constexpr bool bUniformSmoothing = false;	// uniform smoothing can distort the UV layer pretty badly
-	const bool bSuccess = Private::Remesh(DynamicMesh, TargetEdgeLength, IterationsRender, SmoothingRender, bUniformSmoothing, Seams);
-	check(bSuccess);
+	const int InputMeshVertexCount = DynamicMesh.VertexCount();
+	const int InputMeshTriangleCount = DynamicMesh.TriangleCount();
 
+	const bool bHasUVs = (DynamicMesh.Attributes()->PrimaryUV() != nullptr);
 
-	// Set the output cloth collection
+	UE::Geometry::FCompactMaps CompactMaps;
+	if (RemeshMethodRender == EChaosClothAssetRemeshMethod::Remesh)
+	{
+		const double MeshArea = TMeshQueries<FDynamicMesh3>::GetVolumeArea(DynamicMesh).Y;
+		const int TargetTriangleCount = FMath::RoundToInt(static_cast<float>(TargetPercentRender) / 100.0f * static_cast<float>(InputMeshTriangleCount));
+		const double TargetEdgeLength = FRemeshMeshOp::CalculateTargetEdgeLength(nullptr, TargetTriangleCount, MeshArea);
 
-	// Collect info from the remeshed DynamicMesh data to reinitialize the Render Pattern
+		TArray<TArray<FIntVector2>> Seams;
 
-	TArray<FIndex3i> Indices;
-	for (const FIndex3i& Tri : DynamicMesh.TrianglesItr())
-	{
-		Indices.Add(Tri);
+		constexpr bool bUniformSmoothing = false;	// uniform smoothing can distort the UV layer pretty badly
+		const bool bSuccess = Private::Remesh(DynamicMesh, TargetEdgeLength, IterationsRender, SmoothingRender, bUniformSmoothing, Seams, &CompactMaps);
+		check(bSuccess);
 	}
-
-	TArray<FVector3f> Positions;
-	for (const FVector3d& Vertex : DynamicMesh.VerticesItr())
+	else
 	{
-		Positions.Add(FVector3f(Vertex));
+		const int TargetVertexCount = FMath::RoundToInt(static_cast<float>(TargetPercentRender) / 100.0f * static_cast<float>(InputMeshVertexCount));
+		Private::Simplify(DynamicMesh, TargetVertexCount, &CompactMaps);
 	}
 
+	// Collect outputs
+
 	//
 	// Normals
 	//
 
-	const bool bHasNormals = DynamicMesh.HasAttributes() && DynamicMesh.Attributes()->PrimaryNormals();
+	const bool bHasNormals = (DynamicMesh.Attributes()->PrimaryNormals() != nullptr);
 	check(bHasNormals);
 
 	TArray<FVector3f> Normals;
-	if (bHasNormals)
+	Normals.SetNum(DynamicMesh.VertexCount());
+	const FDynamicMeshNormalOverlay* const NormalOverlay = DynamicMesh.Attributes()->PrimaryNormals();
+	for (const int TriangleIndex : DynamicMesh.TriangleIndicesItr())
 	{
-		Normals.SetNum(Positions.Num());
-		const FDynamicMeshNormalOverlay* const NormalOverlay = DynamicMesh.Attributes()->PrimaryNormals();
-		for (const int TriangleIndex : DynamicMesh.TriangleIndicesItr())
-		{
-			const FIndex3i Tri = DynamicMesh.GetTriangle(TriangleIndex);
+		const FIndex3i Tri = DynamicMesh.GetTriangle(TriangleIndex);
 
-			for (int TriangleVertexIndex = 0; TriangleVertexIndex < 3; ++TriangleVertexIndex)
-			{
-				const int VertexIndex = Tri[TriangleVertexIndex];
+		for (int TriangleVertexIndex = 0; TriangleVertexIndex < 3; ++TriangleVertexIndex)
+		{
+			const int VertexIndex = Tri[TriangleVertexIndex];
 
-				// NOTE: This assumes one normal per vertex in the overlay (i.e. no "hard edges")
-				Normals[VertexIndex] = NormalOverlay->GetElementAtVertex(TriangleIndex, VertexIndex);
-			}
+			// NOTE: This assumes one normal per vertex in the overlay (i.e. no "hard edges")
+			Normals[VertexIndex] = NormalOverlay->GetElementAtVertex(TriangleIndex, VertexIndex);
 		}
 	}
-	else
-	{
-		// TODO: Compute normals
-
-	}
 
 	//
 	// Tangents
@@ -1162,14 +1239,14 @@ void FChaosClothAssetRemeshNode::RemeshRenderPattern(const TSharedRef<const FMan
 	TArray<FVector3f> TangentUs;
 	TArray<FVector3f> TangentVs;
 	{
-		const bool bHasTangentUs = DynamicMesh.HasAttributes() && DynamicMesh.Attributes()->PrimaryTangents();
-		const bool bHasTangentVs = DynamicMesh.HasAttributes() && DynamicMesh.Attributes()->PrimaryBiTangents();
-		if (!bHasTangentUs || bHasTangentVs)
+		const bool bHasTangentUs = (DynamicMesh.Attributes()->PrimaryTangents() != nullptr);
+		const bool bHasTangentVs = (DynamicMesh.Attributes()->PrimaryBiTangents() != nullptr);
+		if (!bHasTangentUs || !bHasTangentVs)
 		{
 			FMeshTangentsf::ComputeDefaultOverlayTangents(DynamicMesh);
 		}
-		TangentUs.SetNumZeroed(Positions.Num());
-		TangentVs.SetNumZeroed(Positions.Num());
+		TangentUs.SetNumZeroed(DynamicMesh.VertexCount());
+		TangentVs.SetNumZeroed(DynamicMesh.VertexCount());
 
 		const FDynamicMeshNormalOverlay* const TangentUOverlay = DynamicMesh.Attributes()->PrimaryTangents();
 		const FDynamicMeshNormalOverlay* const TangentVOverlay = DynamicMesh.Attributes()->PrimaryBiTangents();
@@ -1183,11 +1260,11 @@ void FChaosClothAssetRemeshNode::RemeshRenderPattern(const TSharedRef<const FMan
 				const int VertexIndex = Tri[TriangleVertexIndex];
 
 				TangentUs[VertexIndex] += TangentUOverlay->GetElementAtVertex(TriangleIndex, VertexIndex);
-				TangentVs[VertexIndex] += TangentUOverlay->GetElementAtVertex(TriangleIndex, VertexIndex);
+				TangentVs[VertexIndex] += TangentVOverlay->GetElementAtVertex(TriangleIndex, VertexIndex);
 			}
 		}
 
-		for (int32 VertexIndex = 0; VertexIndex < Positions.Num(); ++VertexIndex)
+		for (int32 VertexIndex = 0; VertexIndex < DynamicMesh.VertexCount(); ++VertexIndex)
 		{
 			TangentUs[VertexIndex].Normalize();
 			TangentVs[VertexIndex].Normalize();
@@ -1202,7 +1279,7 @@ void FChaosClothAssetRemeshNode::RemeshRenderPattern(const TSharedRef<const FMan
 	TArray<FVector2f> UVs;
 	if (bHasUVs)
 	{
-		UVs.SetNum(Positions.Num());
+		UVs.SetNum(DynamicMesh.VertexCount());
 
 		const FDynamicMeshUVOverlay* const UVOverlay = DynamicMesh.Attributes()->PrimaryUV();
 
@@ -1226,73 +1303,133 @@ void FChaosClothAssetRemeshNode::RemeshRenderPattern(const TSharedRef<const FMan
 	//
 
 	FDynamicMeshAttributeSet* Attributes = DynamicMesh.Attributes();
+	check(Attributes);
+
 	TArray<TArray<int32>> BoneIndices;
 	TArray<TArray<float>> BoneWeights;
+	BoneIndices.SetNum(DynamicMesh.VertexCount());
+	BoneWeights.SetNum(DynamicMesh.VertexCount());
 
-	if (Attributes)
+	for (const TPair<FName, TUniquePtr<FDynamicMeshVertexSkinWeightsAttribute>>& SkinWeightLayer : Attributes->GetSkinWeightsAttributes())
 	{
-		BoneIndices.SetNum(Positions.Num());
-		BoneWeights.SetNum(Positions.Num());
+		const FDynamicMeshVertexSkinWeightsAttribute* const SkinWeightAttribute = SkinWeightLayer.Value.Get();
 
-		for (const TPair<FName, TUniquePtr<FDynamicMeshVertexSkinWeightsAttribute>>& SkinWeightLayer : Attributes->GetSkinWeightsAttributes())
+		for (const int TriangleIndex : DynamicMesh.TriangleIndicesItr())
 		{
-			const FDynamicMeshVertexSkinWeightsAttribute* const SkinWeightAttribute = SkinWeightLayer.Value.Get();
+			const FIndex3i Tri = DynamicMesh.GetTriangle(TriangleIndex);
 
-			for (const int TriangleIndex : DynamicMesh.TriangleIndicesItr())
+			for (int TriangleVertexIndex = 0; TriangleVertexIndex < 3; ++TriangleVertexIndex)
 			{
-				const FIndex3i Tri = DynamicMesh.GetTriangle(TriangleIndex);
-
-				for (int TriangleVertexIndex = 0; TriangleVertexIndex < 3; ++TriangleVertexIndex)
-				{
-					const int VertexIndex = Tri[TriangleVertexIndex];
-					SkinWeightAttribute->GetValue(VertexIndex, BoneIndices[VertexIndex], BoneWeights[VertexIndex]);
-				}
+				const int VertexIndex = Tri[TriangleVertexIndex];
+				SkinWeightAttribute->GetValue(VertexIndex, BoneIndices[VertexIndex], BoneWeights[VertexIndex]);
 			}
 		}
 	}
 
 
-	FCollectionClothFacade OutClothFacade(OutClothCollection);
-	FCollectionClothRenderPatternFacade OutClothPatternFacade = OutClothFacade.GetRenderPattern(PatternIndex);
+	// Find the set of triangles per MaterialID
 
-	OutClothPatternFacade.SetNumRenderVertices(Positions.Num());
-	OutClothPatternFacade.SetNumRenderFaces(Indices.Num());
+	TMap<int32, TArray<int32>> MaterialTriangles;
 
-	TArrayView<FVector3f> RenderPosition = OutClothPatternFacade.GetRenderPosition();
-	TArrayView<FVector3f> RenderNormal = OutClothPatternFacade.GetRenderNormal();
-	TArrayView<FVector3f> RenderTangentU = OutClothPatternFacade.GetRenderTangentU();
-	TArrayView<FVector3f> RenderTangentV = OutClothPatternFacade.GetRenderTangentV();
-	TArrayView<TArray<FVector2f>> RenderUVs = OutClothPatternFacade.GetRenderUVs();
-	TArrayView<FLinearColor> RenderColor = OutClothPatternFacade.GetRenderColor();
-	TArrayView<TArray<int32>> RenderBoneIndices = OutClothPatternFacade.GetRenderBoneIndices();
-	TArrayView<TArray<float>> RenderBoneWeights = OutClothPatternFacade.GetRenderBoneWeights();
+	const FDynamicMeshMaterialAttribute* const MaterialAttribute = Attributes->GetMaterialID();
+	check(MaterialAttribute);
 
-	for (int32 VInd = 0; VInd < Positions.Num(); ++VInd)
+	for (const int32 TriangleID : DynamicMesh.TriangleIndicesItr())
 	{
-		RenderPosition[VInd] = Positions[VInd];
-		if (bHasUVs)
-		{
-			RenderUVs[VInd].SetNum(MAX_TEXCOORDS);
-			RenderUVs[VInd][0] = UVs[VInd];
-		}
-		if (bHasNormals)
+		const int32 MaterialID = MaterialAttribute->GetValue(TriangleID);
+		if (!MaterialTriangles.Contains(MaterialID))
 		{
-			RenderNormal[VInd] = Normals[VInd];
+			MaterialTriangles.Add(MaterialID);
 		}
-		RenderTangentU[VInd] = TangentUs[VInd];
-		RenderTangentV[VInd] = TangentVs[VInd];
-		RenderColor[VInd] = FLinearColor::White;
-		RenderBoneIndices[VInd] = BoneIndices[VInd];
-		RenderBoneWeights[VInd] = BoneWeights[VInd];
+		MaterialTriangles[MaterialID].Add(TriangleID);
 	}
 
-	const int32 VertexOffset = OutClothPatternFacade.GetRenderVerticesOffset();
-	TArrayView<FIntVector3> RenderIndices = OutClothPatternFacade.GetRenderIndices();
-	for (int32 TInd = 0; TInd < Indices.Num(); ++TInd)
+	TArray<int32> MaterialIDs;
+	MaterialTriangles.GetKeys(MaterialIDs);
+
+	const int32 NumMaterials = MaterialIDs.Num();
+
+	//
+	// Populate output cloth collection
+	// 
+
+	FClothGeometryTools::DeleteRenderMesh(OutClothCollection);
+	FCollectionClothFacade OutClothFacade(OutClothCollection);
+	OutClothFacade.SetNumRenderPatterns(NumMaterials);
+
+	for (int32 DestPatternID = 0; DestPatternID < MaterialIDs.Num(); ++DestPatternID)
 	{
-		RenderIndices[TInd][0] = VertexOffset + Indices[TInd][0];
-		RenderIndices[TInd][1] = VertexOffset + Indices[TInd][1];
-		RenderIndices[TInd][2] = VertexOffset + Indices[TInd][2];
+		FCollectionClothRenderPatternFacade OutClothPatternFacade = OutClothFacade.GetRenderPattern(DestPatternID);
+		check(OutClothPatternFacade.GetNumRenderFaces() == 0);
+		check(OutClothPatternFacade.GetNumRenderVertices() == 0);
+
+		const int32 SourceMaterialID = MaterialIDs[DestPatternID];
+		check(MaterialTriangles.Contains(SourceMaterialID));
+		const TArray<int32>& TriangleIDs = MaterialTriangles[SourceMaterialID];
+
+		TSet<int32> VertexIndices;
+		for (int32 PatternTriangleIndex = 0; PatternTriangleIndex < TriangleIDs.Num(); ++PatternTriangleIndex)
+		{
+			const int32 TInd = TriangleIDs[PatternTriangleIndex];
+			const FIndex3i Tri = DynamicMesh.GetTriangle(TInd);
+			VertexIndices.Add(Tri[0]);
+			VertexIndices.Add(Tri[1]);
+			VertexIndices.Add(Tri[2]);
+		}
+		const TArray<int32> SourceVertexIndicesArray = VertexIndices.Array();
+		const int32 NumVerticesThisPattern = SourceVertexIndicesArray.Num();
+
+		OutClothPatternFacade.SetNumRenderVertices(NumVerticesThisPattern);
+		TArrayView<FVector3f> RenderPosition = OutClothPatternFacade.GetRenderPosition();
+		TArrayView<FVector3f> RenderNormal = OutClothPatternFacade.GetRenderNormal();
+		TArrayView<FVector3f> RenderTangentU = OutClothPatternFacade.GetRenderTangentU();
+		TArrayView<FVector3f> RenderTangentV = OutClothPatternFacade.GetRenderTangentV();
+		TArrayView<TArray<FVector2f>> RenderUVs = OutClothPatternFacade.GetRenderUVs();
+		TArrayView<FLinearColor> RenderColor = OutClothPatternFacade.GetRenderColor();
+		TArrayView<TArray<int32>> RenderBoneIndices = OutClothPatternFacade.GetRenderBoneIndices();
+		TArrayView<TArray<float>> RenderBoneWeights = OutClothPatternFacade.GetRenderBoneWeights();
+		
+		TMap<int32, int32> SourceToDestVertexMap;
+		for (int32 PatternVertexIndex = 0; PatternVertexIndex < NumVerticesThisPattern; ++PatternVertexIndex)
+		{
+			const int32 SourceVertexIndex = SourceVertexIndicesArray[PatternVertexIndex];
+
+			SourceToDestVertexMap.Add(SourceVertexIndex, PatternVertexIndex);
+
+			RenderPosition[PatternVertexIndex] = FVector3f(DynamicMesh.GetVertex(SourceVertexIndex));
+			if (bHasUVs)
+			{
+				RenderUVs[PatternVertexIndex].SetNum(MAX_TEXCOORDS);
+				RenderUVs[PatternVertexIndex][0] = UVs[SourceVertexIndex];
+			}
+			if (bHasNormals)
+			{
+				RenderNormal[PatternVertexIndex] = Normals[SourceVertexIndex];
+			}
+			RenderTangentU[PatternVertexIndex] = TangentUs[SourceVertexIndex];
+			RenderTangentV[PatternVertexIndex] = TangentVs[SourceVertexIndex];
+			RenderColor[PatternVertexIndex] = FLinearColor::White;
+			RenderBoneIndices[PatternVertexIndex] = BoneIndices[SourceVertexIndex];
+			RenderBoneWeights[PatternVertexIndex] = BoneWeights[SourceVertexIndex];
+		}
+
+		OutClothPatternFacade.SetNumRenderFaces(TriangleIDs.Num());
+		for (int32 PatternTriangleIndex = 0; PatternTriangleIndex < TriangleIDs.Num(); ++PatternTriangleIndex)
+		{
+			const int32 VertexOffset = OutClothPatternFacade.GetRenderVerticesOffset();
+			TArrayView<FIntVector3> RenderIndices = OutClothPatternFacade.GetRenderIndices();
+
+			const int32 TInd = TriangleIDs[PatternTriangleIndex];
+			const FIndex3i SourceTri = DynamicMesh.GetTriangle(TInd);
+
+			RenderIndices[PatternTriangleIndex][0] = VertexOffset + SourceToDestVertexMap[SourceTri[0]];
+			RenderIndices[PatternTriangleIndex][1] = VertexOffset + SourceToDestVertexMap[SourceTri[1]];
+			RenderIndices[PatternTriangleIndex][2] = VertexOffset + SourceToDestVertexMap[SourceTri[2]];
+		}
+
+		FCollectionClothConstFacade InClothFacade(ClothCollection);
+		FCollectionClothRenderPatternConstFacade InPatternFacade = InClothFacade.GetRenderPattern(SourceMaterialID);
+		OutClothPatternFacade.SetRenderMaterialPathName(InPatternFacade.GetRenderMaterialPathName());
 	}
 }
 
@@ -1323,12 +1460,7 @@ void FChaosClothAssetRemeshNode::Evaluate(Dataflow::FContext& Context, const FDa
 			if (bRemeshRender)
 			{
 				EmptyRenderSelections(OutputClothCollection);
-
-				const int32 NumPatterns = ClothFacade.GetNumRenderPatterns();
-				for (int32 PatternIndex = 0; PatternIndex < NumPatterns; ++PatternIndex)
-				{
-					RemeshRenderPattern(ClothCollection, PatternIndex, OutputClothCollection);
-				}
+				RemeshRenderMesh(ClothCollection, OutputClothCollection);
 			}
 		}
 
diff --git a/Engine/Plugins/Experimental/ChaosClothAssetEditor/Source/ChaosClothAssetDataflowNodes/Public/ChaosClothAsset/RemeshNode.h b/Engine/Plugins/Experimental/ChaosClothAssetEditor/Source/ChaosClothAssetDataflowNodes/Public/ChaosClothAsset/RemeshNode.h
--- a/Engine/Plugins/Experimental/ChaosClothAssetEditor/Source/ChaosClothAssetDataflowNodes/Public/ChaosClothAsset/RemeshNode.h
+++ b/Engine/Plugins/Experimental/ChaosClothAssetEditor/Source/ChaosClothAssetDataflowNodes/Public/ChaosClothAsset/RemeshNode.h
@@ -6,6 +6,15 @@
 #include "GeometryCollection/ManagedArrayCollection.h"
 #include "RemeshNode.generated.h"
 
+
+UENUM(BlueprintType)
+enum class EChaosClothAssetRemeshMethod : uint8
+{
+	Remesh,
+	Simplify
+};
+
+
 /** Remesh the cloth surface(s) to get the specified mesh resolution(s).
  *  NOTE: Weight Maps, Skinning Data, Self Collision Spheres, and Long Range Attachment Constraints will be reconstructed on the output mesh, however all other Selections will be removed
  */
@@ -26,7 +35,7 @@ public:
 	UPROPERTY(EditAnywhere, Category = "Sim Mesh", meta=(UIMin = "1", UIMax = "200", ClampMin = "1", EditCondition = "bRemeshSim"))
 	int32 TargetPercentSim = 100;
 
-	UPROPERTY(EditAnywhere, Category = "Sim Mesh", meta = (UIMin = "0", UIMax = "20", ClampMin = "0", ClampMax = "100", EditCondition = "bRemeshSim"))
+	UPROPERTY(EditAnywhere, Category = "Sim Mesh", meta = (UIMin = "0", UIMax = "20", ClampMin = "0", ClampMax = "200", EditCondition = "bRemeshSim"))
 	int32 IterationsSim = 10;
 
 	UPROPERTY(EditAnywhere, Category = "Sim Mesh", meta = (UIMin = "0.0", UIMax = "1.0", ClampMin = "0.0", ClampMax = "1.0", EditCondition = "bRemeshSim"))
@@ -38,10 +47,13 @@ public:
 	UPROPERTY(EditAnywhere, Category = "Render Mesh", meta=(UIMin = "1", UIMax = "200", ClampMin = "1", EditCondition = "bRemeshRender"))
 	int32 TargetPercentRender = 100;
 
-	UPROPERTY(EditAnywhere, Category = "Render Mesh", meta = (UIMin = "0", UIMax = "20", ClampMin = "0", ClampMax = "100", EditCondition = "bRemeshRender"))
+	UPROPERTY(EditAnywhere, Category = "Render Mesh", meta = (UIMin = "0.0", UIMax = "1.0", ClampMin = "0.0", ClampMax = "1.0", EditCondition = "bRemeshRender"))
+	EChaosClothAssetRemeshMethod RemeshMethodRender = EChaosClothAssetRemeshMethod::Remesh;
+
+	UPROPERTY(EditAnywhere, Category = "Render Mesh", meta = (UIMin = "0", UIMax = "20", ClampMin = "0", ClampMax = "100", EditCondition = "bRemeshRender && RemeshMethodRender == EChaosClothAssetRemeshMethod::Remesh"))
 	int32 IterationsRender = 10;
 
-	UPROPERTY(EditAnywhere, Category = "Render Mesh", meta = (UIMin = "0.0", UIMax = "1.0", ClampMin = "0.0", ClampMax = "1.0", EditCondition = "bRemeshRender"))
+	UPROPERTY(EditAnywhere, Category = "Render Mesh", meta = (UIMin = "0.0", UIMax = "1.0", ClampMin = "0.0", ClampMax = "1.0", EditCondition = "bRemeshRender && RemeshMethodRender == EChaosClothAssetRemeshMethod::Remesh"))
 	double SmoothingRender = 0.25;
 
 	FChaosClothAssetRemeshNode(const Dataflow::FNodeParameters& InParam, FGuid InGuid = FGuid::NewGuid());
@@ -60,8 +72,6 @@ private:
 
 	void EmptyRenderSelections(const TSharedRef<FManagedArrayCollection>& ClothCollection) const;
 
-	void RemeshRenderPattern(const TSharedRef<const FManagedArrayCollection>& ClothCollection,
-		int32 PatternIndex,
+	void RemeshRenderMesh(const TSharedRef<const FManagedArrayCollection>& ClothCollection,
 		const TSharedRef<FManagedArrayCollection>& OutClothCollection) const;
-
 };
diff --git a/Engine/Plugins/Experimental/ChaosClothAssetEditor/Source/ChaosClothAssetTools/Private/ChaosClothAsset/ClothPatternToDynamicMesh.cpp b/Engine/Plugins/Experimental/ChaosClothAssetEditor/Source/ChaosClothAssetTools/Private/ChaosClothAsset/ClothPatternToDynamicMesh.cpp
--- a/Engine/Plugins/Experimental/ChaosClothAssetEditor/Source/ChaosClothAssetTools/Private/ChaosClothAsset/ClothPatternToDynamicMesh.cpp
+++ b/Engine/Plugins/Experimental/ChaosClothAssetEditor/Source/ChaosClothAssetTools/Private/ChaosClothAsset/ClothPatternToDynamicMesh.cpp
@@ -336,7 +336,7 @@ public:
 
 	bool HasBiTangents() const
 	{
-		return false;
+		return VertexDataType == EClothPatternVertexType::Render;
 	}
 
 	bool HasColors() const
