Refactoring Chaos Skin Weights Transfer node:
-Move creation of the transfer operator to where we actually doing the transfer instead of passing it around to methods.
-Create a helper struct for passing around transfer settings.
-Setup the UI toggles to allow in the future to transfer to both sim and render mesh, just the sim or just the render mesh.
-Setup the UI toggles to allow in the future to choose the type of the mesh to transfer from for the render mesh.

#rnx

diff --git a/Engine/Plugins/Experimental/ChaosClothAssetEditor/Source/ChaosClothAssetDataflowNodes/Private/ChaosClothAsset/TransferSkinWeightsNode.cpp b/Engine/Plugins/Experimental/ChaosClothAssetEditor/Source/ChaosClothAssetDataflowNodes/Private/ChaosClothAsset/TransferSkinWeightsNode.cpp
--- a/Engine/Plugins/Experimental/ChaosClothAssetEditor/Source/ChaosClothAssetDataflowNodes/Private/ChaosClothAsset/TransferSkinWeightsNode.cpp
+++ b/Engine/Plugins/Experimental/ChaosClothAssetEditor/Source/ChaosClothAssetDataflowNodes/Private/ChaosClothAsset/TransferSkinWeightsNode.cpp
@@ -27,7 +27,35 @@
 static const FName InpaintWeightMaskName(TEXT("_InpaintWeightMask"));
 
 namespace UE::Chaos::ClothAsset::Private
-{
+{	
+	/** Helper struct to pass the transfer settings around */
+	struct TransferBoneWeightsSettings
+	{
+		/** 
+		 * Settings for controlling which meshes to transfer to and from. 
+		 */
+		bool bTransferToSim = true; 	// if true, transfer to sim mesh, otherwise skip sim mesh
+		bool bTransferToRender = true;  // if true, transfer to render mesh, otherwise skip render mesh
+		bool bTransferToRenderFromSim = true; // if true, for render mesh only, transfer from the sim mesh, otherwise transfer from the source skeletal mesh
+
+		/** 
+		 * Shared Transfer Operator Settings 
+		 */
+		bool bUseParallel = false;
+		int32 MaxNumInfluences = 8;
+		UE::Geometry::FTransferBoneWeights::ETransferBoneWeightsMethod TransferMethod = UE::Geometry::FTransferBoneWeights::ETransferBoneWeightsMethod::ClosestPointOnSurface;
+
+		/** 
+		 * Settings for the ETransferBoneWeightsMethod::InpaintWeights transfer method
+		 */
+		double NormalThreshold = 0;
+		double RadiusPercentage = 0;
+		bool LayeredMeshSupport = false;
+		int32 NumSmoothingIterations = 0;
+		double SmoothingStrength = 0;
+		FString InpaintMaskWeightMapName;
+	};
+
 	/** Convert the USkeletalMesh to DynamicMesh. */
 	static void SkeletalMeshToDynamicMesh(USkeletalMesh* FromSkeletalMeshAsset, int32 LodIndex, FDynamicMesh3& ToDynamicMesh)
 	{
@@ -251,16 +279,10 @@ namespace UE::Chaos::ClothAsset::Private
 	/** Transfer skin weights to sim cloth. */
 	static bool TransferInpaintWeightsToSim(
 		const FReferenceSkeleton& TargetRefSkeleton,
-		const double NormalThreshold,
-		const double RadiusPercentage,
-		const bool LayeredMeshSupport,
-		const int32 NumSmoothingIterations,
-		const double SmoothingStrength,
-		bool bUseParallel,
+		const UE::Geometry::FDynamicMesh3& SourceDynamicMesh,
 		const TSharedRef<FManagedArrayCollection>& ClothCollection,
-		const FName& InpaintMaskWeightMapName,
-		UE::Geometry::FTransferBoneWeights& TransferBoneWeights,
-		FDynamicMesh3& WeldedSimMesh)
+		const TransferBoneWeightsSettings& TransferSettings,
+		UE::Geometry::FDynamicMesh3& WeldedSimMesh)
 	{
 		using namespace UE::Geometry;
 
@@ -276,12 +298,16 @@ namespace UE::Chaos::ClothAsset::Private
 		}
 
 		// Transfer the weights from the body to the welded sim mesh.
-		TransferBoneWeights.NormalThreshold = FMathd::DegToRad * NormalThreshold;
-		TransferBoneWeights.SearchRadius = RadiusPercentage * WeldedSimMesh.GetBounds().DiagonalLength();
-		TransferBoneWeights.NumSmoothingIterations = NumSmoothingIterations;
-		TransferBoneWeights.SmoothingStrength = SmoothingStrength;
-		TransferBoneWeights.LayeredMeshSupport = LayeredMeshSupport; // multilayerd clothing
-		TransferBoneWeights.ForceInpaintWeightMapName = InpaintMaskWeightMapName;
+		FTransferBoneWeights TransferBoneWeights(&SourceDynamicMesh, FSkeletalMeshAttributes::DefaultSkinWeightProfileName);
+		TransferBoneWeights.bUseParallel = TransferSettings.bUseParallel;
+		TransferBoneWeights.MaxNumInfluences = TransferSettings.MaxNumInfluences;
+		TransferBoneWeights.TransferMethod = TransferSettings.TransferMethod;
+		TransferBoneWeights.NormalThreshold = FMathd::DegToRad * TransferSettings.NormalThreshold;
+		TransferBoneWeights.SearchRadius = TransferSettings.RadiusPercentage * WeldedSimMesh.GetBounds().DiagonalLength();
+		TransferBoneWeights.NumSmoothingIterations = TransferSettings.NumSmoothingIterations;
+		TransferBoneWeights.SmoothingStrength = TransferSettings.SmoothingStrength;
+		TransferBoneWeights.LayeredMeshSupport = TransferSettings.LayeredMeshSupport; // multilayerd clothing
+		TransferBoneWeights.ForceInpaintWeightMapName = FName(TransferSettings.InpaintMaskWeightMapName);
 
 		if (!ensure(TransferBoneWeights.Validate() == EOperationValidationResult::Ok))
 		{
@@ -297,22 +323,20 @@ namespace UE::Chaos::ClothAsset::Private
 		
 		ClothFacade.AddWeightMap(InpaintWeightMaskName);
 
+		// Copy the new bone weight data and inpaint mask from the welded sim mesh back to the sim cloth patterns.
+		CopySkinWeightsFromDynamicMeshToCloth(WeldedSimMesh, true, ClothCollection, TransferSettings.MaxNumInfluences);
+		CopyInpaintMapFromDynamicMeshToCloth(WeldedSimMesh, InpaintWeightMaskName, true, TransferBoneWeights.MatchedVertices, ClothCollection);
+
 		return true;
 	}
 
 	/** Transfer skin weights to sim and render cloth. */
 	static bool TransferInpaintWeights(
 		const FReferenceSkeleton& TargetRefSkeleton,
-		const double NormalThreshold,
-		const double RadiusPercentage,
-		const bool LayeredMeshSupport,
-		const int32 NumSmoothingIterations,
-		const double SmoothingStrength,
-		bool bUseParallel,
+		const UE::Geometry::FDynamicMesh3& SourceDynamicMesh,
 		const TSharedRef<FManagedArrayCollection>& ClothCollection,
 		const TSharedRef<FManagedArrayCollection>& UserTransferCollection,
-		const FName InpaintMaskWeightMapName,
-		UE::Geometry::FTransferBoneWeights& TransferBoneWeights)
+		const TransferBoneWeightsSettings& TransferSettings)
 	{
 		using namespace UE::Geometry;
 
@@ -320,24 +344,18 @@ namespace UE::Chaos::ClothAsset::Private
 		// Convert cloth sim mesh LOD to the welded dynamic sim mesh and transfer weights.
 		//
 		FDynamicMesh3 WeldedSimMesh;
-		if (!TransferInpaintWeightsToSim(TargetRefSkeleton, NormalThreshold, RadiusPercentage, LayeredMeshSupport, NumSmoothingIterations, SmoothingStrength, bUseParallel, ClothCollection, InpaintMaskWeightMapName, TransferBoneWeights, WeldedSimMesh))
+		if (!TransferInpaintWeightsToSim(TargetRefSkeleton, SourceDynamicMesh, ClothCollection, TransferSettings, WeldedSimMesh))
 		{
 			return false;
 		}
 
-		//
-		// Copy the new bone weight data and inpaint mask from the welded sim mesh back to the sim cloth patterns.
-		//
-		CopySkinWeightsFromDynamicMeshToCloth(WeldedSimMesh, true, ClothCollection, TransferBoneWeights.MaxNumInfluences);
-		CopyInpaintMapFromDynamicMeshToCloth(WeldedSimMesh, InpaintWeightMaskName, true, TransferBoneWeights.MatchedVertices, ClothCollection);
-
 		//
 		// Check if the custom mesh collection for transferring weights is provided, in which case use it instead of the sim mesh
 		//
 		if (FCollectionClothFacade(UserTransferCollection).IsValid())
 		{
 			FDynamicMesh3 TransferWeightSimMesh;
-			if (TransferInpaintWeightsToSim(TargetRefSkeleton, NormalThreshold, RadiusPercentage, LayeredMeshSupport, NumSmoothingIterations, SmoothingStrength, bUseParallel, UserTransferCollection, InpaintMaskWeightMapName, TransferBoneWeights, TransferWeightSimMesh))
+			if (TransferInpaintWeightsToSim(TargetRefSkeleton, SourceDynamicMesh, UserTransferCollection, TransferSettings, TransferWeightSimMesh))
 			{
 				WeldedSimMesh = MoveTemp(TransferWeightSimMesh); 
 			}
@@ -355,9 +373,9 @@ namespace UE::Chaos::ClothAsset::Private
 		if (WeldedSimMesh.VertexCount() > 0 && WeldedSimMesh.TriangleCount() > 0)
 		{
 			FTransferBoneWeights SimToRenderMeshTransfer(&WeldedSimMesh, FSkeletalMeshAttributes::DefaultSkinWeightProfileName);
-			SimToRenderMeshTransfer.bUseParallel = bUseParallel;
+			SimToRenderMeshTransfer.bUseParallel = TransferSettings.bUseParallel;
 			SimToRenderMeshTransfer.TransferMethod = FTransferBoneWeights::ETransferBoneWeightsMethod::ClosestPointOnSurface;
-			SimToRenderMeshTransfer.MaxNumInfluences = TransferBoneWeights.MaxNumInfluences;
+			SimToRenderMeshTransfer.MaxNumInfluences = TransferSettings.MaxNumInfluences;
 			UE::Chaos::ClothAsset::FCollectionClothFacade ClothFacade(ClothCollection);
 			ParallelFor(ClothFacade.GetNumRenderVertices(), [&SimToRenderMeshTransfer, &ClothFacade](int32 VertexID)
 				{
@@ -365,7 +383,7 @@ namespace UE::Chaos::ClothAsset::Private
 					ClothFacade.GetRenderBoneWeights()[VertexID],
 					ClothFacade.GetRenderPosition()[VertexID]);
 
-				}, bUseParallel ? EParallelForFlags::None : EParallelForFlags::ForceSingleThread);
+				}, TransferSettings.bUseParallel ? EParallelForFlags::None : EParallelForFlags::ForceSingleThread);
 
 			return true;
 		}
@@ -377,9 +395,9 @@ namespace UE::Chaos::ClothAsset::Private
 
 	static bool TransferClosestPointOnSurface(
 		const FReferenceSkeleton& TargetRefSkeleton,
-		const bool bUseParallel,
+		const FDynamicMesh3& SourceDynamicMesh,
 		UE::Chaos::ClothAsset::FCollectionClothFacade& ClothFacade,
-		UE::Geometry::FTransferBoneWeights& TransferBoneWeights)
+		const UE::Chaos::ClothAsset::Private::TransferBoneWeightsSettings& TransferSettings)
 	{
 		using namespace UE::Geometry;
 
@@ -393,6 +411,11 @@ namespace UE::Chaos::ClothAsset::Private
 			TargetBoneToIndex.Add(TargetRefSkeleton.GetRawRefBoneInfo()[BoneIdx].Name, BoneIdx);
 		}
 
+		FTransferBoneWeights TransferBoneWeights(&SourceDynamicMesh, FSkeletalMeshAttributes::DefaultSkinWeightProfileName);
+		TransferBoneWeights.bUseParallel = TransferSettings.bUseParallel;
+		TransferBoneWeights.MaxNumInfluences = TransferSettings.MaxNumInfluences;
+		TransferBoneWeights.TransferMethod = TransferSettings.TransferMethod;
+
 		if (!ensure(TransferBoneWeights.Validate() == EOperationValidationResult::Ok))
 		{
 			UE_LOG(LogChaosClothAssetDataflowNodes, Warning, TEXT("TransferSkinWeightsNode: Transfer method parameters are invalid."));
@@ -409,7 +432,7 @@ namespace UE::Chaos::ClothAsset::Private
 				ClothFacade.GetSimPosition3D()[VertexID],
 				&TargetBoneToIndex);
 
-		}, bUseParallel ? EParallelForFlags::None : EParallelForFlags::ForceSingleThread);
+		}, TransferSettings.bUseParallel ? EParallelForFlags::None : EParallelForFlags::ForceSingleThread);
 
 		//
 		// Transfer weights to the render mesh.
@@ -421,7 +444,7 @@ namespace UE::Chaos::ClothAsset::Private
 				ClothFacade.GetRenderPosition()[VertexID],
 				&TargetBoneToIndex);
 
-		}, bUseParallel ? EParallelForFlags::None : EParallelForFlags::ForceSingleThread);
+		}, TransferSettings.bUseParallel ? EParallelForFlags::None : EParallelForFlags::ForceSingleThread);
 
 		return true;
 	}
@@ -475,14 +498,25 @@ void FChaosClothAssetTransferSkinWeightsNode::Evaluate(Dataflow::FContext& Conte
 			const FReferenceSkeleton& TargetRefSkeleton = SkeletalMesh->GetRefSkeleton();
 
 			//
-			// Setup the bone weight transfer operator for the source mesh.
+			// Setup the bone weight transfer settings.
 			//
-			constexpr bool bUseParallel = true;
-			FTransferBoneWeights TransferBoneWeights(&SourceDynamicMesh, FSkeletalMeshAttributes::DefaultSkinWeightProfileName);
-			TransferBoneWeights.MaxNumInfluences = static_cast<int32>(FChaosClothAssetTransferSkinWeightsNode::MaxNumInfluences);
-			TransferBoneWeights.bUseParallel = bUseParallel;
-			TransferBoneWeights.TransferMethod = static_cast<FTransferBoneWeights::ETransferBoneWeightsMethod>(TransferMethod);
-
+			TransferBoneWeightsSettings TransferSettings;
+
+			TransferSettings.bTransferToSim = TargetMeshType != EChaosClothAssetTransferTargetMeshType::Render;
+			TransferSettings.bTransferToRender = TargetMeshType != EChaosClothAssetTransferTargetMeshType::Simulation;
+			TransferSettings.bTransferToRenderFromSim = RenderMeshSourceType == EChaosClothAssetTransferRenderMeshSource::SimulationMesh;
+			
+			TransferSettings.bUseParallel = true;
+			TransferSettings.MaxNumInfluences = static_cast<int32>(FChaosClothAssetTransferSkinWeightsNode::MaxNumInfluences);
+			TransferSettings.TransferMethod = static_cast<FTransferBoneWeights::ETransferBoneWeightsMethod>(TransferMethod);
+
+			TransferSettings.NormalThreshold = NormalThreshold;
+			TransferSettings.RadiusPercentage = RadiusPercentage;
+			TransferSettings.LayeredMeshSupport = LayeredMeshSupport;
+			TransferSettings.NumSmoothingIterations = NumSmoothingIterations;
+			TransferSettings.SmoothingStrength = SmoothingStrength;
+			TransferSettings.InpaintMaskWeightMapName = GetValue<FString>(Context, &InpaintMask.WeightMap);
+			
 			//
 			// Transfer the bone weights from the source Skeletal mesh to the Cloth asset.
 			//
@@ -494,12 +528,11 @@ void FChaosClothAssetTransferSkinWeightsNode::Evaluate(Dataflow::FContext& Conte
 			{
 				FManagedArrayCollection TransferWeightsManagedCollection = GetValue<FManagedArrayCollection>(Context, &TransferWeightsCollection);
 				const TSharedRef<FManagedArrayCollection> ClothTransferWeightCollection = MakeShared<FManagedArrayCollection>(MoveTemp(TransferWeightsManagedCollection));
-				const FString& InpainMaskWeightMapName = GetValue<FString>(Context, &InpaintMask.WeightMap);
-				bTransferResult = TransferInpaintWeights(TargetRefSkeleton, NormalThreshold, RadiusPercentage, LayeredMeshSupport, NumSmoothingIterations, SmoothingStrength, bUseParallel, ClothCollection, ClothTransferWeightCollection, FName(InpainMaskWeightMapName), TransferBoneWeights);
+				bTransferResult = TransferInpaintWeights(TargetRefSkeleton, SourceDynamicMesh, ClothCollection, ClothTransferWeightCollection, TransferSettings);
 			}
 			else if (TransferMethod == EChaosClothAssetTransferSkinWeightsMethod::ClosestPointOnSurface)
 			{
-				bTransferResult = TransferClosestPointOnSurface(TargetRefSkeleton, bUseParallel, ClothFacade, TransferBoneWeights);
+				bTransferResult = TransferClosestPointOnSurface(TargetRefSkeleton, SourceDynamicMesh, ClothFacade, TransferSettings);
 			}
 			else
 			{
diff --git a/Engine/Plugins/Experimental/ChaosClothAssetEditor/Source/ChaosClothAssetDataflowNodes/Public/ChaosClothAsset/TransferSkinWeightsNode.h b/Engine/Plugins/Experimental/ChaosClothAssetEditor/Source/ChaosClothAssetDataflowNodes/Public/ChaosClothAsset/TransferSkinWeightsNode.h
--- a/Engine/Plugins/Experimental/ChaosClothAssetEditor/Source/ChaosClothAssetDataflowNodes/Public/ChaosClothAsset/TransferSkinWeightsNode.h
+++ b/Engine/Plugins/Experimental/ChaosClothAssetEditor/Source/ChaosClothAssetDataflowNodes/Public/ChaosClothAsset/TransferSkinWeightsNode.h
@@ -34,6 +34,29 @@ enum class EChaosClothAssetMaxNumInfluences : uint8
 	Twelve = 12			UMETA(DisplayName = "12")
 };
 
+UENUM(BlueprintType)
+enum class EChaosClothAssetTransferTargetMeshType : uint8
+{
+	/** Perform the skin weights transfer for both the simulation and render meshes. */
+	All,
+
+	/** Perform the skin weights transfer for the simulation mesh only. */
+	Simulation,
+	
+	/** Perform the skin weights transfer for the render mesh only. */
+	Render
+};
+
+UENUM(BlueprintType)
+enum class EChaosClothAssetTransferRenderMeshSource : uint8
+{
+	/** For render mesh, transfer weights from the source skeletal mesh. */
+	SkeletalMesh,
+
+	/** For render mesh, transfer weights from the simulation mesh. */
+	SimulationMesh
+};
+
 /** Transfer the skinning weights set on a skeletal mesh to the simulation and/or render mesh stored in the cloth collection. */
 USTRUCT(Meta = (DataflowCloth))
 struct FChaosClothAssetTransferSkinWeightsNode : public FDataflowNode
@@ -57,6 +80,14 @@ public:
 	UPROPERTY(EditAnywhere, Category = "Transfer Skin Weights|Source Mesh")
 	FTransform Transform;
 
+	/** The type of the mesh the transfer will be applied to. */
+	UPROPERTY(EditAnywhere, Category = "Transfer Skin Weights|Transfer Method", Meta = (DisplayName = "Target Mesh Type"))
+	EChaosClothAssetTransferTargetMeshType TargetMeshType = EChaosClothAssetTransferTargetMeshType::All;
+
+	/** For the render mesh, choose which source to use. */
+	UPROPERTY(EditAnywhere, Category = "Transfer Skin Weights|Transfer Method", Meta = (DisplayName = "Render Mesh Source Type", EditCondition="TargetMeshType!=EChaosClothAssetTransferTargetMeshType::Simulation"))
+	EChaosClothAssetTransferRenderMeshSource RenderMeshSourceType = EChaosClothAssetTransferRenderMeshSource::SimulationMesh;
+
 	/** Algorithm used for the transfer method. Use the simple ClosestPointOnSurface method or the more complex InpaintWeights method for better results. */
 	UPROPERTY(EditAnywhere, Category = "Transfer Skin Weights|Transfer Method", Meta = (DisplayName = "Algorithm"))
 	EChaosClothAssetTransferSkinWeightsMethod TransferMethod = EChaosClothAssetTransferSkinWeightsMethod::InpaintWeights;
