Adding a skin weight transfer method that operates on connected mesh components individually. This works well on disconnected Render meshes.

#rnx

diff --git a/Engine/Plugins/Experimental/ChaosClothAssetEditor/Source/ChaosClothAssetDataflowNodes/Private/ChaosClothAsset/TransferSkinWeightsNode.cpp b/Engine/Plugins/Experimental/ChaosClothAssetEditor/Source/ChaosClothAssetDataflowNodes/Private/ChaosClothAsset/TransferSkinWeightsNode.cpp
--- a/Engine/Plugins/Experimental/ChaosClothAssetEditor/Source/ChaosClothAssetDataflowNodes/Private/ChaosClothAsset/TransferSkinWeightsNode.cpp
+++ b/Engine/Plugins/Experimental/ChaosClothAssetEditor/Source/ChaosClothAssetDataflowNodes/Private/ChaosClothAsset/TransferSkinWeightsNode.cpp
@@ -19,6 +19,8 @@
 #include "Rendering/SkeletalMeshLODImporterData.h"
 #include "Rendering/SkeletalMeshModel.h"
 #include "SkeletalMeshAttributes.h"
+#include "Selections/MeshConnectedComponents.h"
+#include "BoneWeights.h"
 
 #include UE_INLINE_GENERATED_CPP_BY_NAME(TransferSkinWeightsNode)
 
@@ -56,6 +58,18 @@ namespace UE::Chaos::ClothAsset::Private
 		FString InpaintMaskWeightMapName;
 	};
 
+	UE::Geometry::FDynamicMeshVertexSkinWeightsAttribute* GetOrCreateSkinWeightsAttribute(FDynamicMesh3& InMesh, const FName& InProfileName="Default")
+	{
+		checkSlow(InMesh.HasAttributes());
+		UE::Geometry::FDynamicMeshVertexSkinWeightsAttribute* Attribute = InMesh.Attributes()->GetSkinWeightsAttribute(InProfileName);
+		if (Attribute == nullptr)
+		{
+			Attribute = new UE::Geometry::FDynamicMeshVertexSkinWeightsAttribute(&InMesh);
+			InMesh.Attributes()->AttachSkinWeightsAttribute(InProfileName, Attribute);
+		}
+		return Attribute;
+	}
+
 	/** Convert the USkeletalMesh to DynamicMesh. */
 	static void SkeletalMeshToDynamicMesh(USkeletalMesh* FromSkeletalMeshAsset, int32 LodIndex, FDynamicMesh3& ToDynamicMesh)
 	{
@@ -65,18 +79,19 @@ namespace UE::Chaos::ClothAsset::Private
 	}
 
 	/** Convert the ClothCollection to DynamicMesh. */
-	static bool SimClothToDynamicMesh(
+	static bool ClothToDynamicMesh(
 		const TSharedRef<FManagedArrayCollection>& ClothCollection,
 		const FReferenceSkeleton& TargetRefSkeleton, // the reference skeleton to add to the dynamic mesh
-		UE::Geometry::FDynamicMesh3& WeldedSimMesh)
+		const bool bIsSim, // if true, Mesh will contain the sim mesh, otherwise the render mesh
+		UE::Geometry::FDynamicMesh3& Mesh) // the resulting sim or render mesh
 	{
 		using namespace UE::Geometry;
 		using namespace UE::Chaos::ClothAsset;
 
 		// Check if ClothCollection is empty
 		FCollectionClothFacade ClothFacade(ClothCollection);
-		const int32 NumVertices = ClothFacade.GetNumSimVertices3D();
-		const int32 NumFaces = ClothFacade.GetNumSimFaces();
+		const int32 NumVertices = bIsSim ? ClothFacade.GetNumSimVertices3D() : ClothFacade.GetNumRenderVertices();
+		const int32 NumFaces = bIsSim ? ClothFacade.GetNumSimFaces() :  ClothFacade.GetNumRenderFaces();
 
 		if (NumVertices <= 0 || NumFaces <= 0)
 		{
@@ -85,25 +100,47 @@ namespace UE::Chaos::ClothAsset::Private
 		}
 
 		// Convert the sim mesh to DynamicMesh. 
-		// @todo: FTransferBoneWeights should accept raw data arrays (vertices/triangles) to avoid this conversion
 		FClothPatternToDynamicMesh PatternToDynamicMesh;
 		constexpr bool bDisableAttributes = false;
-		PatternToDynamicMesh.Convert(ClothCollection, INDEX_NONE, EClothPatternVertexType::Sim3D, WeldedSimMesh, bDisableAttributes);
+		const EClothPatternVertexType PatternType = bIsSim ? EClothPatternVertexType::Sim3D : EClothPatternVertexType::Render;
+		PatternToDynamicMesh.Convert(ClothCollection, INDEX_NONE, PatternType, Mesh, bDisableAttributes);
 
 		// Setup the skeleton.
 		// @note we can't simply copy the bone attributes from the source USkeletalMesh because the cloth  
 		// asset reference skeleton comes from the USkeleton not the USkeletalMesh
-		WeldedSimMesh.Attributes()->EnableBones(TargetRefSkeleton.GetRawBoneNum());
+		Mesh.Attributes()->EnableBones(TargetRefSkeleton.GetRawBoneNum());
 		for (int32 BoneIdx = 0; BoneIdx < TargetRefSkeleton.GetRawBoneNum(); ++BoneIdx)
 		{
-			WeldedSimMesh.Attributes()->GetBoneNames()->SetValue(BoneIdx, TargetRefSkeleton.GetRawRefBoneInfo()[BoneIdx].Name);
+			Mesh.Attributes()->GetBoneNames()->SetValue(BoneIdx, TargetRefSkeleton.GetRawRefBoneInfo()[BoneIdx].Name);
 		}
 
 		return true;
 	}
 
-	/** Copy the skin weights from DynamicMesh to Collection, handling split vertices. */
-	static void CopySkinWeightsFromDynamicMeshToCloth(const FDynamicMesh3& WeldedSimMesh, 
+	/** Copy the skin weights from DynamicMesh to Render Collection. */
+	static void CopySkinWeightsFromDynamicMeshToRenderCloth(const FDynamicMesh3& RenderMesh,
+															const bool bUseParallel,
+													  	 	const TSharedRef<FManagedArrayCollection>& ClothCollection)
+	{
+		using namespace UE::Geometry;
+
+		UE::Chaos::ClothAsset::FCollectionClothFacade ClothFacade(ClothCollection);
+
+		ParallelFor(RenderMesh.MaxVertexID(), [&ClothFacade, &RenderMesh](int32 VertexID)
+		{
+			const FDynamicMeshVertexSkinWeightsAttribute* OutAttribute = RenderMesh.Attributes()->GetSkinWeightsAttribute(FSkeletalMeshAttributes::DefaultSkinWeightProfileName);
+
+			checkSlow(OutAttribute);
+			checkSlow(RenderMesh.IsVertex(VertexID));
+			checkSlow(VertexID < ClothFacade.GetNumRenderVertices());
+			OutAttribute->GetValue(VertexID, ClothFacade.GetRenderBoneIndices()[VertexID],
+				ClothFacade.GetRenderBoneWeights()[VertexID]);
+
+		}, bUseParallel ? EParallelForFlags::None : EParallelForFlags::ForceSingleThread);
+	}
+
+	/** Copy the skin weights from DynamicMesh to Sim Collection, handling split vertices. */
+	static void CopySkinWeightsFromDynamicMeshToSimCloth(const FDynamicMesh3& WeldedSimMesh, 
 													  const bool bUseParallel,
 													  const TSharedRef<FManagedArrayCollection>& ClothCollection,
 													  const int32 MaxNumInfluences)
@@ -222,7 +259,7 @@ namespace UE::Chaos::ClothAsset::Private
 	}
 
 	/** Copy the inpaint mask array to Collection, handling split vertices. */
-	static void CopyInpaintMapFromDynamicMeshToCloth(const FDynamicMesh3& WeldedSimMesh, 
+	static void CopyInpaintMapFromDynamicMeshToSimCloth(const FDynamicMesh3& WeldedSimMesh, 
 													 const FName WeightMapName,
 													 const bool bUseParallel, 
 													 const TArray<bool>& MatchedVertices,
@@ -275,6 +312,134 @@ namespace UE::Chaos::ClothAsset::Private
 			}, bUseParallel ? EParallelForFlags::None : EParallelForFlags::ForceSingleThread);
 		}
 	}
+	
+
+	/** 
+	 * Transfer skin weights from the source to the target dynamic mesh. 
+	 * The target mesh is split into connected components and the transfer is run on each component separately.
+	 */
+	static bool TransferInpaintWeightsToMeshComponents(const FReferenceSkeleton& TargetRefSkeleton,
+													   const UE::Geometry::FDynamicMesh3& SourceDynamicMesh,
+													   UE::Geometry::FDynamicMesh3& TargetDynamicMesh,
+													   const TransferBoneWeightsSettings& TransferSettings)
+	{
+		using namespace UE::Geometry;
+
+		// Setup transfer weights operator
+		FTransferBoneWeights TransferBoneWeights(&SourceDynamicMesh, FSkeletalMeshAttributes::DefaultSkinWeightProfileName);
+		TransferBoneWeights.bUseParallel = TransferSettings.bUseParallel;
+		TransferBoneWeights.MaxNumInfluences = TransferSettings.MaxNumInfluences;
+		TransferBoneWeights.TransferMethod = TransferSettings.TransferMethod;
+		TransferBoneWeights.NormalThreshold = FMathd::DegToRad * TransferSettings.NormalThreshold;
+		TransferBoneWeights.SearchRadius = TransferSettings.RadiusPercentage * TargetDynamicMesh.GetBounds().DiagonalLength();
+		TransferBoneWeights.NumSmoothingIterations = TransferSettings.NumSmoothingIterations;
+		TransferBoneWeights.SmoothingStrength = TransferSettings.SmoothingStrength;
+		TransferBoneWeights.LayeredMeshSupport = TransferSettings.LayeredMeshSupport; // multilayerd clothing
+
+		// TODO: allow to specify an inpaint mask for the render mesh
+		// TransferBoneWeights.ForceInpaintWeightMapName = FName(TransferSettings.InpaintMaskWeightMapName);
+
+		if (!ensure(TransferBoneWeights.Validate() == EOperationValidationResult::Ok))
+		{
+			UE_LOG(LogChaosClothAssetDataflowNodes, Warning, TEXT("TransferSkinWeightsNode: Transfer method parameters are invalid."));
+			return false;
+		}
+
+		// Find connected-components
+		FMeshConnectedComponents ConnectedComponents(&TargetDynamicMesh);
+		ConnectedComponents.FindConnectedTriangles();
+
+		// Iterate over each component and perform per-component skin weight transfer
+		const int32 NumComponents = ConnectedComponents.Num();
+		for (int32 ComponentIdx = 0; ComponentIdx < NumComponents; ++ComponentIdx)
+		{
+			// Mesh containing the geometry of the current component
+			FDynamicMesh3 Submesh;
+			Submesh.EnableAttributes();
+			Submesh.Attributes()->EnableBones(TargetRefSkeleton.GetRawBoneNum());
+			for (int32 BoneIdx = 0; BoneIdx < TargetRefSkeleton.GetRawBoneNum(); ++BoneIdx)
+			{
+				Submesh.Attributes()->GetBoneNames()->SetValue(BoneIdx, TargetRefSkeleton.GetRawRefBoneInfo()[BoneIdx].Name);
+			}
+
+			// Keep track of the maps between Mesh and the component submesh vertex/triangle indices
+			TMap<int32, int32> BaseToSubmeshV;
+			TArray<int32> SubmeshToBaseV;
+			TArray<int32> SubmeshToBaseT;
+
+			const TArray<int32>& ComponentTris = ConnectedComponents[ComponentIdx].Indices;
+
+			for (const int32 TID : ComponentTris)
+			{
+				const FIndex3i Triangle = TargetDynamicMesh.GetTriangle(TID);
+				FIndex3i NewTriangle;
+				for (int32 TriCornerIdx = 0; TriCornerIdx < 3; ++TriCornerIdx)
+				{
+					const int32* FoundIdx = BaseToSubmeshV.Find(Triangle[TriCornerIdx]);
+					if (FoundIdx)
+					{
+						NewTriangle[TriCornerIdx] = *FoundIdx;
+					}
+					else
+					{
+						const FVector3d Position = TargetDynamicMesh.GetVertex(Triangle[TriCornerIdx]);
+						const int32 NewVtxID = Submesh.AppendVertex(Position);
+						check(NewVtxID == SubmeshToBaseV.Num());
+						SubmeshToBaseV.Add(Triangle[TriCornerIdx]);
+						BaseToSubmeshV.Add(Triangle[TriCornerIdx], NewVtxID);
+						NewTriangle[TriCornerIdx] = NewVtxID;
+					}
+				}
+
+				const int32 NewTriID = Submesh.AppendTriangle(NewTriangle);
+				check(NewTriID == SubmeshToBaseT.Num());
+				SubmeshToBaseT.Add(TID);
+			}
+
+			// Transfer weights to the current submesh only. If transfer using the inpaint method fails, fall back to  
+			// the closest point method which should always succeed. Common reason for the failure is if we didn't find  
+			// any matches on the source at all with the current transfer settings.
+			if (!TransferBoneWeights.TransferWeightsToMesh(Submesh, FSkeletalMeshAttributes::DefaultSkinWeightProfileName))
+			{	
+				UE_LOG(LogChaosClothAssetDataflowNodes, Warning, TEXT("TransferSkinWeightsNode: Failed to transfer skin weights to some of the vertices of the render mesh using inpaint method, falling back to closest point method."));
+
+				// We can reuse the same operator but change the method type
+				TransferBoneWeights.TransferMethod = FTransferBoneWeights::ETransferBoneWeightsMethod::ClosestPointOnSurface;
+				
+				// Ignore radius and normal settings, so all points on the source are considered
+				const double OldSearchRadius = TransferBoneWeights.SearchRadius;
+				const double OldNormalThreshold = TransferBoneWeights.NormalThreshold;
+				TransferBoneWeights.SearchRadius = -1;	  
+				TransferBoneWeights.NormalThreshold = -1; 
+
+				// This should always succeed
+				if (!ensure(TransferBoneWeights.TransferWeightsToMesh(Submesh, FSkeletalMeshAttributes::DefaultSkinWeightProfileName)))
+				{
+					UE_LOG(LogChaosClothAssetDataflowNodes, Warning, TEXT("TransferSkinWeightsNode:  Failed to transfer skin weights to some of the vertices of the render mesh."));
+				}
+
+				// Revert back the settings
+				TransferBoneWeights.TransferMethod = TransferSettings.TransferMethod;
+				TransferBoneWeights.SearchRadius = OldSearchRadius;
+				TransferBoneWeights.NormalThreshold = OldNormalThreshold;
+			}
+				
+			// Copy over the data from the submesh to the base mesh
+			FDynamicMeshVertexSkinWeightsAttribute* SubMeshSkinWeights = UE::Chaos::ClothAsset::Private::GetOrCreateSkinWeightsAttribute(Submesh);
+			FDynamicMeshVertexSkinWeightsAttribute* BaseMeshSkinWeights = UE::Chaos::ClothAsset::Private::GetOrCreateSkinWeightsAttribute(TargetDynamicMesh);
+
+			for (int32 SubMeshVID = 0; SubMeshVID < SubmeshToBaseV.Num(); ++SubMeshVID)
+			{
+				const int32 BaseMeshVID = SubmeshToBaseV[SubMeshVID];
+
+				UE::AnimationCore::FBoneWeights Weights;
+				SubMeshSkinWeights->GetValue(SubMeshVID, Weights);
+				BaseMeshSkinWeights->SetValue(BaseMeshVID, Weights);
+			}
+		}
+
+		return true;
+	}
 
 	/** Transfer skin weights to sim cloth. */
 	static bool TransferInpaintWeightsToSim(
@@ -291,13 +456,14 @@ namespace UE::Chaos::ClothAsset::Private
 		UE::Chaos::ClothAsset::FCollectionClothFacade ClothFacade(ClothCollection);
 
 		// Convert cloth sim mesh LOD to the welded dynamic sim mesh.
-		if (!SimClothToDynamicMesh(ClothCollection, TargetRefSkeleton, WeldedSimMesh))
+		if (!ClothToDynamicMesh(ClothCollection, TargetRefSkeleton, true, WeldedSimMesh))
 		{
 			UE_LOG(LogChaosClothAssetDataflowNodes, Warning, TEXT("TransferSkinWeightsNode: Failed to weld the simulation mesh for LOD."));
 			return false;
 		}
 
 		// Transfer the weights from the body to the welded sim mesh.
+		// TODO: run the transfer on components instead
 		FTransferBoneWeights TransferBoneWeights(&SourceDynamicMesh, FSkeletalMeshAttributes::DefaultSkinWeightProfileName);
 		TransferBoneWeights.bUseParallel = TransferSettings.bUseParallel;
 		TransferBoneWeights.MaxNumInfluences = TransferSettings.MaxNumInfluences;
@@ -324,8 +490,37 @@ namespace UE::Chaos::ClothAsset::Private
 		ClothFacade.AddWeightMap(InpaintWeightMaskName);
 
 		// Copy the new bone weight data and inpaint mask from the welded sim mesh back to the sim cloth patterns.
-		CopySkinWeightsFromDynamicMeshToCloth(WeldedSimMesh, true, ClothCollection, TransferSettings.MaxNumInfluences);
-		CopyInpaintMapFromDynamicMeshToCloth(WeldedSimMesh, InpaintWeightMaskName, true, TransferBoneWeights.MatchedVertices, ClothCollection);
+		CopySkinWeightsFromDynamicMeshToSimCloth(WeldedSimMesh, true, ClothCollection, TransferSettings.MaxNumInfluences);
+		CopyInpaintMapFromDynamicMeshToSimCloth(WeldedSimMesh, InpaintWeightMaskName, true, TransferBoneWeights.MatchedVertices, ClothCollection);
+
+		return true;
+	}
+
+	/** Transfer skin weights to render cloth. */
+	static bool TransferInpaintWeightsToRender(
+		const FReferenceSkeleton& TargetRefSkeleton,
+		const UE::Geometry::FDynamicMesh3& SourceDynamicMesh,
+		const TSharedRef<FManagedArrayCollection>& ClothCollection,
+		const TransferBoneWeightsSettings& TransferSettings)
+	{
+		using namespace UE::Geometry;
+
+		// Convert cloth render mesh LOD to the dynamic render mesh.
+		FDynamicMesh3 RenderDynamicMesh;
+		if (!ClothToDynamicMesh(ClothCollection, TargetRefSkeleton, false, RenderDynamicMesh))
+		{
+			UE_LOG(LogChaosClothAssetDataflowNodes, Warning, TEXT("TransferSkinWeightsNode: Failed to create the render dynamic mesh for LOD."));
+			return false;
+		}
+
+		// Transfer weights to render mesh
+		if (!TransferInpaintWeightsToMeshComponents(TargetRefSkeleton, SourceDynamicMesh, RenderDynamicMesh, TransferSettings))
+		{
+			UE_LOG(LogChaosClothAssetDataflowNodes, Warning, TEXT("TransferSkinWeightsNode: Failed to transfer skin weights to render mesh."));
+			return false;
+		}
+
+		CopySkinWeightsFromDynamicMeshToRenderCloth(RenderDynamicMesh, TransferSettings.bUseParallel, ClothCollection);
 
 		return true;
 	}
