Cloth collection merging of weight maps

diff --git a/Engine/Plugins/Experimental/ChaosClothAssetEditor/Source/ChaosClothAssetDataflowNodes/Private/ChaosClothAsset/MergeClothCollectionsNode.cpp b/Engine/Plugins/Experimental/ChaosClothAssetEditor/Source/ChaosClothAssetDataflowNodes/Private/ChaosClothAsset/MergeClothCollectionsNode.cpp
--- a/Engine/Plugins/Experimental/ChaosClothAssetEditor/Source/ChaosClothAssetDataflowNodes/Private/ChaosClothAsset/MergeClothCollectionsNode.cpp
+++ b/Engine/Plugins/Experimental/ChaosClothAssetEditor/Source/ChaosClothAssetDataflowNodes/Private/ChaosClothAsset/MergeClothCollectionsNode.cpp
@@ -3,6 +3,7 @@
 #include "ChaosClothAsset/MergeClothCollectionsNode.h"
 #include "ChaosClothAsset/CollectionClothFacade.h"
 #include "ChaosClothAsset/CollectionClothSelectionFacade.h"
+#include "ChaosClothAsset/ClothDataflowTools.h"
 #include "Chaos/CollectionPropertyFacade.h"
 #include "Dataflow/DataflowInputOutput.h"
 
@@ -10,6 +11,225 @@
 
 #define LOCTEXT_NAMESPACE "ChaosClothAssetMergeClothCollectionsNode"
 
+namespace UE::Chaos::ClothAsset::Private
+{
+	static void LogAndToastDifferentWeightMapNames(const FDataflowNode& DataflowNode, const FString& PropertyName, const FString& InWeightMapName, const FString& OutWeightMapName, const FString& WeightMapName)
+	{
+		using namespace UE::Chaos::ClothAsset;
+
+		static const FText Headline = LOCTEXT("DiffertentWeightMapNamesHeadline", "Different weight map names.");
+
+		const FText Details = FText::Format(
+			LOCTEXT(
+				"DiffertentWeightMapNamesDetails",
+				"Two identical Cloth Collection properties '{0}' are being merged but have different weight map names '{1}' and '{2}'. The weight map named '{3}' will be used in the resulting merge."),
+			FText::FromString(PropertyName),
+			FText::FromString(OutWeightMapName),
+			FText::FromString(InWeightMapName),
+			FText::FromString(WeightMapName));
+
+		FClothDataflowTools::LogAndToastWarning(DataflowNode, Headline, Details);
+	}
+
+	/** Build weight maps for each properties if necessary */
+	static FString BuildWeightMaps(const FDataflowNode& DataflowNode,
+		const FCollectionClothConstFacade& InClothFacade, FCollectionClothFacade& OutClothFacade,
+		const FVector2f& InPropertyBounds, const FVector2f& OutPropertyBounds,
+		const FVector2f& PropertyBounds, const FString& PropertyName,
+		const FString& InWeightMapName, const FString& OutWeightMapName)
+	{
+		FString WeightMapName = (OutWeightMapName.IsEmpty() && !InWeightMapName.IsEmpty()) ? InWeightMapName :
+								(OutWeightMapName.IsEmpty() && InWeightMapName.IsEmpty()) ? PropertyName : OutWeightMapName;
+
+		if(WeightMapName != OutWeightMapName)
+		{
+			FString IncrWeightMapName = WeightMapName;
+			int32 IncrCount = 0;
+			
+			while(OutClothFacade.GetWeightMap(FName(IncrWeightMapName)).Num() > 0)
+			{
+				IncrWeightMapName = WeightMapName;
+				IncrWeightMapName.AppendInt(++IncrCount);
+			}
+		}
+		
+		// If the low high values of the merged property are the same we don't need to build a weight map
+		if(PropertyBounds[0] != PropertyBounds[1])
+		{
+			// If names are different we must let the user know
+			if ((!InWeightMapName.IsEmpty() && InWeightMapName != WeightMapName) || (!OutWeightMapName.IsEmpty() && OutWeightMapName != WeightMapName))
+			{
+				Private::LogAndToastDifferentWeightMapNames(DataflowNode, PropertyName, InWeightMapName, OutWeightMapName, WeightMapName);	
+			}
+			// Create if necessary a new weight map
+			OutClothFacade.AddWeightMap(FName(WeightMapName));
+
+			auto FillWeightMap = [&PropertyName, &PropertyBounds](const TConstArrayView<float> InWeightMap, const FVector2f& InPropertyBounds,
+				const int32 InNumVertices, const int32 OutVertexOffset, TArrayView<float> OutWeightMap)
+			{
+				const bool bHasAlreadyValues = (InWeightMap.Num() > 0);
+				for(int32 VertexIndex = 0; VertexIndex < InNumVertices; ++VertexIndex)
+                {
+                	// If no values in the weight map we are using the low value
+					const float WeightMapValue = bHasAlreadyValues ? (InWeightMap[VertexIndex] * (InPropertyBounds[1] - InPropertyBounds[0]) +
+						InPropertyBounds[0]) : InPropertyBounds[0];
+					OutWeightMap[OutVertexOffset+VertexIndex] = (WeightMapValue - PropertyBounds[0]) / (PropertyBounds[1] - PropertyBounds[0]);
+                }
+			};
+			const int32 InNumVertices = InClothFacade.GetNumSimVertices3D();
+			const int32 OutNumVertices = OutClothFacade.GetNumSimVertices3D() - InNumVertices;
+			
+			TArrayView<float> OutWeightMap = OutClothFacade.GetWeightMap(FName(OutWeightMapName));
+			FillWeightMap(OutWeightMap, OutPropertyBounds, OutNumVertices, 0, OutWeightMap);
+			
+			const TConstArrayView<float> InWeightMap = InClothFacade.GetWeightMap(FName(InWeightMapName));
+			FillWeightMap(InWeightMap, InPropertyBounds, InNumVertices, OutNumVertices, OutWeightMap);
+		}
+		return WeightMapName;
+	}
+
+	// Merge the property bounds of 2 collections
+	static FVector2f MergePropertyBounds(const FVector2f& InPropertyBounds, const FVector2f& OutPropertyBounds)
+	{
+		FVector2f PropertyBounds(0.0f);
+		if(InPropertyBounds[0] <= InPropertyBounds[1])
+		{
+			if(OutPropertyBounds[0] <= OutPropertyBounds[1])
+			{
+				PropertyBounds[0] = FMath::Min(InPropertyBounds[0], OutPropertyBounds[0]);
+				PropertyBounds[1] = FMath::Max(InPropertyBounds[1], OutPropertyBounds[1]);
+			}
+			else
+			{
+				PropertyBounds[0] = FMath::Min(InPropertyBounds[0], OutPropertyBounds[1]);
+				PropertyBounds[1] = FMath::Max(InPropertyBounds[1], OutPropertyBounds[0]);
+			}
+		}
+		else
+		{
+			if(OutPropertyBounds[0] <= OutPropertyBounds[1])
+			{
+				PropertyBounds[0] = FMath::Min(InPropertyBounds[1], OutPropertyBounds[0]);
+				PropertyBounds[1] = FMath::Max(InPropertyBounds[0], OutPropertyBounds[1]);
+			}
+			else
+			{
+				PropertyBounds[0] = FMath::Min(InPropertyBounds[1], OutPropertyBounds[1]);
+				PropertyBounds[1] = FMath::Max(InPropertyBounds[0], OutPropertyBounds[0]);
+			}
+		}
+		if (FMath::IsNearlyEqual(PropertyBounds[0], PropertyBounds[1]))
+		{
+			PropertyBounds[1] = PropertyBounds[0];
+		}
+		return PropertyBounds;
+	}
+
+	static ::Chaos::Softs::ECollectionPropertyFlags MergePropertyFlags(
+		const ::Chaos::Softs::FCollectionPropertyConstFacade& InPropertyFacade,
+			  ::Chaos::Softs::FCollectionPropertyMutableFacade& OutPropertyFacade,
+			  const int32 InKeyIndex,
+			  const int32 OutKeyIndex,
+			  const ::Chaos::Softs::ECollectionPropertyFlags InPropertyFlags,
+			  const FString& PropertyName)
+	{
+		PRAGMA_DISABLE_DEPRECATION_WARNINGS
+		// TODO: GetFlags needs to return an ECollectionPropertyFlags, not an uint8, but the uint8 getter needs to be deprecated first
+		const ::Chaos::Softs::ECollectionPropertyFlags OutPropertyFlags = (::Chaos::Softs::ECollectionPropertyFlags)OutPropertyFacade.GetFlags(OutKeyIndex);
+		PRAGMA_ENABLE_DEPRECATION_WARNINGS
+		
+		::Chaos::Softs::ECollectionPropertyFlags PropertyFlags;
+		if(!OutPropertyFacade.IsEnabled(OutKeyIndex) && InPropertyFacade.IsEnabled(InKeyIndex))
+		{
+			PropertyFlags = InPropertyFlags;
+		}
+		else if(OutPropertyFacade.IsEnabled(OutKeyIndex) && !InPropertyFacade.IsEnabled(InKeyIndex))
+		{
+			PropertyFlags = OutPropertyFlags;
+		}
+		else
+		{
+			PropertyFlags = OutPropertyFlags;
+			if(OutPropertyFacade.IsAnimatable(OutKeyIndex) || InPropertyFacade.IsAnimatable(InKeyIndex))
+			{
+				EnumAddFlags(PropertyFlags, ::Chaos::Softs::ECollectionPropertyFlags::Animatable);  // Animatable
+			}
+			if(!ensure(OutPropertyFacade.IsIntrinsic(OutKeyIndex) == InPropertyFacade.IsIntrinsic(InKeyIndex)))
+			{
+				UE_LOG(LogChaosClothAssetDataflowNodes, Warning, TEXT("MergeClothCollectionsNode: Mismatch in intrinsic flag onto %s property"), *PropertyName);
+			}
+			if(!ensure(OutPropertyFacade.IsLegacy(OutKeyIndex) == InPropertyFacade.IsLegacy(InKeyIndex)))
+			{
+				UE_LOG(LogChaosClothAssetDataflowNodes, Warning, TEXT("MergeClothCollectionsNode: Mismatch in legacy flag onto %s property"), *PropertyName);
+			}
+			if(!ensure(OutPropertyFacade.IsInterpolable(OutKeyIndex) == InPropertyFacade.IsInterpolable(InKeyIndex)))
+			{
+				UE_LOG(LogChaosClothAssetDataflowNodes, Warning, TEXT("MergeClothCollectionsNode: Mismatch in interpolable flag onto %s property"), *PropertyName);
+			}
+		}
+		return PropertyFlags;
+	}
+
+	/** Append input properties to the output property facade and add potential weight maps */
+	static void AppendInputProperties(const FDataflowNode& DataflowNode,
+		const FCollectionClothConstFacade& InClothFacade,
+			  FCollectionClothFacade& OutClothFacade,
+		const ::Chaos::Softs::FCollectionPropertyConstFacade& InPropertyFacade,
+			  ::Chaos::Softs::FCollectionPropertyMutableFacade& OutPropertyFacade)
+	{
+		const int32 InNumInKeys = InPropertyFacade.Num();
+		for (int32 InKeyIndex = 0; InKeyIndex < InNumInKeys; ++InKeyIndex)
+		{
+			PRAGMA_DISABLE_DEPRECATION_WARNINGS
+			// TODO: GetFlags needs to return an ECollectionPropertyFlags, not an uint8, but the uint8 getter needs to be deprecated first
+			const ::Chaos::Softs::ECollectionPropertyFlags InPropertyFlags = (::Chaos::Softs::ECollectionPropertyFlags)InPropertyFacade.GetFlags(InKeyIndex);
+			PRAGMA_ENABLE_DEPRECATION_WARNINGS
+
+			// Get the matching output key for the given input one 
+			const FString& InPropertyKey = InPropertyFacade.GetKey(InKeyIndex);
+			int32 OutKeyIndex = OutPropertyFacade.GetKeyIndex(InPropertyKey);
+
+			// We first check if the output key exists into the output facade
+			bool bOverrideProperty = true;
+			if(OutKeyIndex != INDEX_NONE)
+			{
+				if(InPropertyFacade.IsInterpolable(InKeyIndex))
+				{
+					// If it exists we compute the min of the property low values and the max of the property high values
+					const FVector2f InPropertyBounds = InPropertyFacade.GetWeightedFloatValue(InKeyIndex);
+					const FVector2f OutPropertyBounds = OutPropertyFacade.GetWeightedFloatValue(OutKeyIndex);
+					const FVector2f PropertyBounds = MergePropertyBounds(InPropertyBounds, OutPropertyBounds);
+					
+					const ::Chaos::Softs::ECollectionPropertyFlags PropertyFlags = MergePropertyFlags(
+						InPropertyFacade, OutPropertyFacade, InKeyIndex, OutKeyIndex, InPropertyFlags, InPropertyKey);
+					
+                    OutPropertyFacade.SetFlags(OutKeyIndex, PropertyFlags);
+                    OutPropertyFacade.SetWeightedFloatValue(OutKeyIndex, PropertyBounds);
+    
+                    // We keep the string value to be the one in the output if defined
+                    const FString WeightMapName = BuildWeightMaps(DataflowNode, InClothFacade, OutClothFacade,
+                    	InPropertyBounds, OutPropertyBounds, PropertyBounds, InPropertyKey,
+							InPropertyFacade.GetStringValue(InKeyIndex), OutPropertyFacade.GetStringValue(InKeyIndex));
+
+					OutPropertyFacade.SetStringValue(OutKeyIndex, WeightMapName);
+					bOverrideProperty = false;
+				}
+			}
+			else
+			{
+				// If not we add a new property with the flags/bounds/string of the input one
+				OutKeyIndex = OutPropertyFacade.AddProperty(InPropertyKey, InPropertyFlags);
+			}
+			if(bOverrideProperty)
+			{
+				OutPropertyFacade.SetFlags(OutKeyIndex, InPropertyFlags);
+				OutPropertyFacade.SetWeightedValue(OutKeyIndex, InPropertyFacade.GetLowValue<FVector3f>(InKeyIndex), InPropertyFacade.GetHighValue<FVector3f>(InKeyIndex));
+				OutPropertyFacade.SetStringValue(OutKeyIndex, InPropertyFacade.GetStringValue(InKeyIndex));
+			}
+		}
+	}
+}
+
 FChaosClothAssetMergeClothCollectionsNode::FChaosClothAssetMergeClothCollectionsNode(const Dataflow::FNodeParameters& InParam, FGuid InGuid)
 	: FDataflowNode(InParam, InGuid)
 {
@@ -87,8 +307,17 @@ void FChaosClothAssetMergeClothCollectionsNode::Evaluate(Dataflow::FContext& Con
 			const FCollectionPropertyConstFacade OtherPropertyFacade(OtherClothCollection);
 			if (OtherPropertyFacade.IsValid())
 			{
-				constexpr bool bUpdateExistingProperties = true; // Want last one wins.
-				PropertyFacade.Append(OtherClothCollection.ToSharedPtr(), bUpdateExistingProperties);
+				// Change that boolean to come back to the old behavior
+				static constexpr bool bOverrideProperties = false;
+				if(bOverrideProperties)
+				{
+					constexpr bool bUpdateExistingProperties = true; // Want last one wins.
+					PropertyFacade.Append(OtherClothCollection.ToSharedPtr(), bUpdateExistingProperties);
+				}
+				else
+				{
+					Private::AppendInputProperties(*this, OtherClothFacade, ClothFacade, OtherPropertyFacade, PropertyFacade);
+				}
 				bAreAnyValid = true;
 			}
 		}
diff --git a/Engine/Plugins/Experimental/ChaosClothAssetEditor/Source/ChaosClothAssetDataflowNodes/Private/ChaosClothAsset/SimulationBaseConfigNode.cpp b/Engine/Plugins/Experimental/ChaosClothAssetEditor/Source/ChaosClothAssetDataflowNodes/Private/ChaosClothAsset/SimulationBaseConfigNode.cpp
--- a/Engine/Plugins/Experimental/ChaosClothAssetEditor/Source/ChaosClothAssetDataflowNodes/Private/ChaosClothAsset/SimulationBaseConfigNode.cpp
+++ b/Engine/Plugins/Experimental/ChaosClothAssetEditor/Source/ChaosClothAssetDataflowNodes/Private/ChaosClothAsset/SimulationBaseConfigNode.cpp
@@ -172,6 +172,8 @@ int32 FChaosClothAssetSimulationBaseConfigNode::FPropertyHelper::SetPropertyWeig
 	{
 		EnumRemoveFlags(PropertyFlags, ECollectionPropertyFlags::Animatable);  // Non-animatable
 	}
+	EnumAddFlags(PropertyFlags, ECollectionPropertyFlags::Interpolable);  // Interpolable
+	
 	const int32 PropertyKeyIndex = ConfigNode.AddPropertyHelper(Properties, PropertyName, SimilarPropertyNames, PropertyFlags);
 	Properties.SetWeightedValue(PropertyKeyIndex, PropertyLow, PropertyHigh);
 	Properties.SetStringValue(PropertyKeyIndex, ConfigNode.GetValue<FString>(Context, &WeightMap));
diff --git a/Engine/Source/Runtime/Experimental/Chaos/Private/Chaos/CollectionPropertyFacade.cpp b/Engine/Source/Runtime/Experimental/Chaos/Private/Chaos/CollectionPropertyFacade.cpp
--- a/Engine/Source/Runtime/Experimental/Chaos/Private/Chaos/CollectionPropertyFacade.cpp
+++ b/Engine/Source/Runtime/Experimental/Chaos/Private/Chaos/CollectionPropertyFacade.cpp
@@ -112,7 +112,7 @@ namespace Chaos::Softs
 		// Cannot set string dirty without also dirtying the property
 		Flags |= EnumHasAnyFlags(Flags, ECollectionPropertyFlags::StringDirty) ? ECollectionPropertyFlags::Dirty : ECollectionPropertyFlags::None;
 		// Cannot remove the Dirty, StringDirty, and Intrinsic flags
-		Flags |= GetFlagsArray()[KeyIndex] & (ECollectionPropertyFlags::Dirty | ECollectionPropertyFlags::StringDirty | ECollectionPropertyFlags::Intrinsic);
+		Flags |= GetFlagsArray()[KeyIndex] & (ECollectionPropertyFlags::Dirty | ECollectionPropertyFlags::StringDirty | ECollectionPropertyFlags::Intrinsic | ECollectionPropertyFlags::Interpolable);
 
 		SetValue(KeyIndex, GetFlagsArray(), Flags);
 	}
diff --git a/Engine/Source/Runtime/Experimental/Chaos/Public/Chaos/CollectionPropertyFacade.h b/Engine/Source/Runtime/Experimental/Chaos/Public/Chaos/CollectionPropertyFacade.h
--- a/Engine/Source/Runtime/Experimental/Chaos/Public/Chaos/CollectionPropertyFacade.h
+++ b/Engine/Source/Runtime/Experimental/Chaos/Public/Chaos/CollectionPropertyFacade.h
@@ -28,6 +28,7 @@ namespace Chaos::Softs
 		Enabled = 1 << 0,  /** Whether this property is enabled(so that it doesn't have to be removed from the collection when not needed). */
 		Animatable = 1 << 1,  /** Whether this property needs to be set at every frame. This flag is ignored when the Intrinsic flag is also set. */
 		Legacy = 1 << 2,  /** Whether this property has been set by a legacy system predating the property collection. Can be useful for overriding/upgrading some properties post conversion. */
+		Interpolable = 1 << 3,  /** Whether this property can be interpolated. Used to allow interpolation when merging float properties */
 		Intrinsic = 1 << 4,  /** Whether this property is intrinsically built into the simulated object model, rather than affecting the simulation itself (see Animatable in this case). Changing this property requires a re-construction of the simulated object model to be effective. Implies non Animatable. */
 		//~ Add new flags above this line
 		StringDirty = 1 << 6,  /** Whether this property's string has changed and needs to be updated at the next frame. */
@@ -100,6 +101,7 @@ namespace Chaos::Softs
 		bool IsIntrinsic(int32 KeyIndex) const { return HasAnyFlags(KeyIndex, ECollectionPropertyFlags::Intrinsic); }
 		bool IsStringDirty(int32 KeyIndex) const { return HasAnyFlags(KeyIndex, ECollectionPropertyFlags::StringDirty); }
 		bool IsDirty(int32 KeyIndex) const { return HasAnyFlags(KeyIndex, ECollectionPropertyFlags::Dirty); }
+		bool IsInterpolable(int32 KeyIndex) const { return HasAnyFlags(KeyIndex, ECollectionPropertyFlags::Interpolable); }
 
 		//~ Values access per key
 		template<typename T, TEMPLATE_REQUIRES(TIsWeightedType<T>::Value)>
@@ -178,6 +180,11 @@ PRAGMA_ENABLE_DEPRECATION_WARNINGS
 		{
 			return SafeGet(Key, [this](int32 KeyIndex)->bool { return IsDirty(KeyIndex); }, bDefault, OutKeyIndex);
 		}
+		
+		bool IsInterpolable(const FString& Key, bool bDefault = false, int32* OutKeyIndex = nullptr) const
+		{
+			return SafeGet(Key, [this](int32 KeyIndex)->bool { return IsInterpolable(KeyIndex); }, bDefault, OutKeyIndex);
+		}
 
 		friend ::uint32 GetTypeHash(const Chaos::Softs::FCollectionPropertyConstFacade& PropertyFacade)
 		{
@@ -281,7 +288,7 @@ PRAGMA_ENABLE_DEPRECATION_WARNINGS
 
 		void SetStringValue(int32 KeyIndex, const FString& Value) { if (GetStringValueArray()[KeyIndex] != Value) { GetStringValueArray()[KeyIndex] = Value; SetStringDirty(KeyIndex); } }
 
-		/** SetFlags cannot be used to remove Dirty, StringDirty, or Intrinsic flags. Use ClearDirtyFlags to remove dirty flags. */
+		/** SetFlags cannot be used to remove Dirty, StringDirty, Interpolable or Intrinsic flags. Use ClearDirtyFlags to remove dirty flags. */
 		CHAOS_API void SetFlags(int32 KeyIndex, ECollectionPropertyFlags Flags);
 		UE_DEPRECATED(5.3, "Use SetFlags(int32, ECollectionPropertyFlags) instead.")
 		void SetFlags(int32 KeyIndex, uint8 Flags) { return SetFlags(KeyIndex, (ECollectionPropertyFlags)Flags); }
@@ -296,6 +303,7 @@ PRAGMA_ENABLE_DEPRECATION_WARNINGS
 		UE_DEPRECATED(5.3, "SetDirty can only be set, to unset use ClearDirtyFlags instead.")
 		void SetDirty(int32 KeyIndex, bool bDirty) { EnableFlags(KeyIndex, ECollectionPropertyFlags::Dirty, bDirty); }
 		void SetStringDirty(int32 KeyIndex) { EnableFlags(KeyIndex, ECollectionPropertyFlags::StringDirty, true); }
+		void SetInterpolable(int32 KeyIndex) { EnableFlags(KeyIndex, ECollectionPropertyFlags::Interpolable, true); }
 
 		//~ Values set per key
 		template<typename T, TEMPLATE_REQUIRES(TIsWeightedType<T>::Value)>
@@ -377,6 +385,11 @@ PRAGMA_ENABLE_DEPRECATION_WARNINGS
 		{
 			return SafeSet(Key, [this](int32 KeyIndex) { SetStringDirty(KeyIndex); });
 		}
+		
+		int32 SetInterpolable(const FString& Key)
+		{
+			return SafeSet(Key, [this](int32 KeyIndex) { SetInterpolable(KeyIndex); });
+		}
 
 		CHAOS_API void ClearDirtyFlags();
 
