Allow the user to choose the target mesh (sim+render, just sim, just render) to  transfer to.
Just for render mesh, also allow to choose if we want to transfer from skeletal mesh or the sim mesh.
Removed ability to specify custom sim mesh for transfer.

#rnx

diff --git a/Engine/Plugins/Experimental/ChaosClothAssetEditor/Source/ChaosClothAssetDataflowNodes/Private/ChaosClothAsset/TransferSkinWeightsNode.cpp b/Engine/Plugins/Experimental/ChaosClothAssetEditor/Source/ChaosClothAssetDataflowNodes/Private/ChaosClothAsset/TransferSkinWeightsNode.cpp
--- a/Engine/Plugins/Experimental/ChaosClothAssetEditor/Source/ChaosClothAssetDataflowNodes/Private/ChaosClothAsset/TransferSkinWeightsNode.cpp
+++ b/Engine/Plugins/Experimental/ChaosClothAssetEditor/Source/ChaosClothAssetDataflowNodes/Private/ChaosClothAsset/TransferSkinWeightsNode.cpp
@@ -339,7 +339,7 @@ namespace UE::Chaos::ClothAsset::Private
 		// TODO: allow to specify an inpaint mask for the render mesh
 		// TransferBoneWeights.ForceInpaintWeightMapName = FName(TransferSettings.InpaintMaskWeightMapName);
 
-		if (!ensure(TransferBoneWeights.Validate() == EOperationValidationResult::Ok))
+		if (TransferBoneWeights.Validate() != EOperationValidationResult::Ok)
 		{
 			UE_LOG(LogChaosClothAssetDataflowNodes, Warning, TEXT("TransferSkinWeightsNode: Transfer method parameters are invalid."));
 			return false;
@@ -475,12 +475,12 @@ namespace UE::Chaos::ClothAsset::Private
 		TransferBoneWeights.LayeredMeshSupport = TransferSettings.LayeredMeshSupport; // multilayerd clothing
 		TransferBoneWeights.ForceInpaintWeightMapName = FName(TransferSettings.InpaintMaskWeightMapName);
 
-		if (!ensure(TransferBoneWeights.Validate() == EOperationValidationResult::Ok))
+		if (TransferBoneWeights.Validate() != EOperationValidationResult::Ok)
 		{
 			UE_LOG(LogChaosClothAssetDataflowNodes, Warning, TEXT("TransferSkinWeightsNode: Transfer method parameters are invalid."));
 			return false;
 		}
-		if (!ensure(TransferBoneWeights.TransferWeightsToMesh(WeldedSimMesh, FSkeletalMeshAttributes::DefaultSkinWeightProfileName)))
+		if (!TransferBoneWeights.TransferWeightsToMesh(WeldedSimMesh, FSkeletalMeshAttributes::DefaultSkinWeightProfileName))
 		{
 			UE_LOG(LogChaosClothAssetDataflowNodes, Warning, TEXT("TransferSkinWeightsNode: Transferring skin weights failed."));
 			return false;
@@ -530,71 +530,76 @@ namespace UE::Chaos::ClothAsset::Private
 		const FReferenceSkeleton& TargetRefSkeleton,
 		const UE::Geometry::FDynamicMesh3& SourceDynamicMesh,
 		const TSharedRef<FManagedArrayCollection>& ClothCollection,
-		const TSharedRef<FManagedArrayCollection>& UserTransferCollection,
 		const TransferBoneWeightsSettings& TransferSettings)
 	{
 		using namespace UE::Geometry;
 
-		//
-		// Convert cloth sim mesh LOD to the welded dynamic sim mesh and transfer weights.
-		//
 		FDynamicMesh3 WeldedSimMesh;
-		if (!TransferInpaintWeightsToSim(TargetRefSkeleton, SourceDynamicMesh, ClothCollection, TransferSettings, WeldedSimMesh))
-		{
-			return false;
-		}
 
-		//
-		// Check if the custom mesh collection for transferring weights is provided, in which case use it instead of the sim mesh
-		//
-		if (FCollectionClothFacade(UserTransferCollection).IsValid())
-		{
-			FDynamicMesh3 TransferWeightSimMesh;
-			if (TransferInpaintWeightsToSim(TargetRefSkeleton, SourceDynamicMesh, UserTransferCollection, TransferSettings, TransferWeightSimMesh))
-			{
-				WeldedSimMesh = MoveTemp(TransferWeightSimMesh); 
-			}
-			else
+		if (TransferSettings.bTransferToSim)
+		{	
+			//
+			// Convert cloth sim mesh LOD to the welded dynamic sim mesh and transfer weights.
+			//
+			if (!TransferInpaintWeightsToSim(TargetRefSkeleton, SourceDynamicMesh, ClothCollection, TransferSettings, WeldedSimMesh))
 			{
-				UE_LOG(LogChaosClothAssetDataflowNodes, Warning, TEXT("TransferSkinWeightsNode: Custom mesh collection for transferring weights is invalid, falling back to simulation mesh."));
+				return false;
 			}
 		}
 
-		//
-		// Compute the bone weights for the render mesh by transferring weights either from the sim mesh or custom user-specified mesh
-		// @todo If render mesh eventually supports welding, we should be transferring weights from the 
-		// body instead, same as we do for the sim mesh.
-		//
-		if (WeldedSimMesh.VertexCount() > 0 && WeldedSimMesh.TriangleCount() > 0)
+		if (TransferSettings.bTransferToRender)
 		{
-			FTransferBoneWeights SimToRenderMeshTransfer(&WeldedSimMesh, FSkeletalMeshAttributes::DefaultSkinWeightProfileName);
-			SimToRenderMeshTransfer.bUseParallel = TransferSettings.bUseParallel;
-			SimToRenderMeshTransfer.TransferMethod = FTransferBoneWeights::ETransferBoneWeightsMethod::ClosestPointOnSurface;
-			SimToRenderMeshTransfer.MaxNumInfluences = TransferSettings.MaxNumInfluences;
-			UE::Chaos::ClothAsset::FCollectionClothFacade ClothFacade(ClothCollection);
-			ParallelFor(ClothFacade.GetNumRenderVertices(), [&SimToRenderMeshTransfer, &ClothFacade](int32 VertexID)
+			//
+			// Compute the bone weights for the render mesh by transferring weights from the sim mesh
+			//
+			if (TransferSettings.bTransferToRenderFromSim)
+			{	
+				// If we previously transferred to sim mesh we can just reuse that dynamic mesh
+				if (!TransferSettings.bTransferToSim)
 				{
-					SimToRenderMeshTransfer.TransferWeightsToPoint(ClothFacade.GetRenderBoneIndices()[VertexID],
-					ClothFacade.GetRenderBoneWeights()[VertexID],
-					ClothFacade.GetRenderPosition()[VertexID]);
-
-				}, TransferSettings.bUseParallel ? EParallelForFlags::None : EParallelForFlags::ForceSingleThread);
+					// Otherwise get it from the collection
+					if (!ClothToDynamicMesh(ClothCollection, TargetRefSkeleton, true, WeldedSimMesh))
+					{
+						return false;
+					}
+				}
 
-			return true;
-		}
-		else
-		{
-			return false;
+				if (WeldedSimMesh.VertexCount() > 0 && WeldedSimMesh.TriangleCount() > 0)
+				{
+					// Transfer skin weights to render cloth from the sim mesh
+					if (!TransferInpaintWeightsToRender(TargetRefSkeleton, WeldedSimMesh, ClothCollection, TransferSettings))
+					{
+						return false;
+					}
+				}
+				else
+				{
+					UE_LOG(LogChaosClothAssetDataflowNodes, Warning, TEXT("TransferSkinWeightsNode: Could not transfer from simulation mesh because its empty."));
+					return false;
+				}
+			}
+			else 
+			{
+				// Transfer skin weights to render cloth from the skeletal asset mesh
+				if (!TransferInpaintWeightsToRender(TargetRefSkeleton, SourceDynamicMesh, ClothCollection, TransferSettings))
+				{
+					return false;
+				}
+			}
 		}
+
+		return true;
 	}
 
 	static bool TransferClosestPointOnSurface(
 		const FReferenceSkeleton& TargetRefSkeleton,
-		const FDynamicMesh3& SourceDynamicMesh,
-		UE::Chaos::ClothAsset::FCollectionClothFacade& ClothFacade,
+		const FDynamicMesh3& SkeletalDynamicMesh,
+		const TSharedRef<FManagedArrayCollection>& ClothCollection,
 		const UE::Chaos::ClothAsset::Private::TransferBoneWeightsSettings& TransferSettings)
 	{
 		using namespace UE::Geometry;
+		
+		FCollectionClothFacade ClothFacade(ClothCollection);
 
 		//
 		// Compute the bone index mappings. This allows the transfer operator to retarget weights to the correct skeleton.
@@ -606,41 +611,72 @@ namespace UE::Chaos::ClothAsset::Private
 			TargetBoneToIndex.Add(TargetRefSkeleton.GetRawRefBoneInfo()[BoneIdx].Name, BoneIdx);
 		}
 
-		FTransferBoneWeights TransferBoneWeights(&SourceDynamicMesh, FSkeletalMeshAttributes::DefaultSkinWeightProfileName);
-		TransferBoneWeights.bUseParallel = TransferSettings.bUseParallel;
-		TransferBoneWeights.MaxNumInfluences = TransferSettings.MaxNumInfluences;
-		TransferBoneWeights.TransferMethod = TransferSettings.TransferMethod;
-
-		if (!ensure(TransferBoneWeights.Validate() == EOperationValidationResult::Ok))
+		if (TransferSettings.bTransferToSim)
 		{
-			UE_LOG(LogChaosClothAssetDataflowNodes, Warning, TEXT("TransferSkinWeightsNode: Transfer method parameters are invalid."));
-			return false;
+			//
+			// Transfer weights to the sim mesh.
+			//
+
+			FTransferBoneWeights TransferBoneWeights(&SkeletalDynamicMesh, FSkeletalMeshAttributes::DefaultSkinWeightProfileName);
+			TransferBoneWeights.bUseParallel = TransferSettings.bUseParallel;
+			TransferBoneWeights.MaxNumInfluences = TransferSettings.MaxNumInfluences;
+			TransferBoneWeights.TransferMethod = TransferSettings.TransferMethod;
+
+			if (TransferBoneWeights.Validate() != EOperationValidationResult::Ok)
+			{
+				UE_LOG(LogChaosClothAssetDataflowNodes, Warning, TEXT("TransferSkinWeightsNode: Transfer method parameters are invalid."));
+				return false;
+			}
+		
+			ParallelFor(ClothFacade.GetNumSimVertices3D(), [&TransferBoneWeights, &TargetBoneToIndex, &ClothFacade](int32 VertexID)
+			{
+				TransferBoneWeights.TransferWeightsToPoint(ClothFacade.GetSimBoneIndices()[VertexID],
+					ClothFacade.GetSimBoneWeights()[VertexID],
+					ClothFacade.GetSimPosition3D()[VertexID],
+					&TargetBoneToIndex);
+
+			}, TransferSettings.bUseParallel ? EParallelForFlags::None : EParallelForFlags::ForceSingleThread);
 		}
 
-		//
-		// Transfer weights to the sim mesh.
-		//
-		ParallelFor(ClothFacade.GetNumSimVertices3D(), [&TransferBoneWeights, &TargetBoneToIndex, &ClothFacade](int32 VertexID)
+		if (TransferSettings.bTransferToRender)
 		{
-			TransferBoneWeights.TransferWeightsToPoint(ClothFacade.GetSimBoneIndices()[VertexID],
-				ClothFacade.GetSimBoneWeights()[VertexID],
-				ClothFacade.GetSimPosition3D()[VertexID],
-				&TargetBoneToIndex);
+			//
+			// Transfer weights to the render mesh.
+			//
 
-		}, TransferSettings.bUseParallel ? EParallelForFlags::None : EParallelForFlags::ForceSingleThread);
+			const FDynamicMesh3* SourceMeshToTransferFrom = &SkeletalDynamicMesh; // transfer from body
+			FDynamicMesh3 WeldedSimMesh;
+			if (TransferSettings.bTransferToRenderFromSim)
+			{
+				// Convert sim cloth to dynamic mesh
+				if (!ClothToDynamicMesh(ClothCollection, TargetRefSkeleton, true, WeldedSimMesh))
+				{
+					return false;
+				}
 
-		//
-		// Transfer weights to the render mesh.
-		//
-		ParallelFor(ClothFacade.GetNumRenderVertices(), [&TransferBoneWeights, &TargetBoneToIndex, &ClothFacade](int32 VertexID)
-		{
-			TransferBoneWeights.TransferWeightsToPoint(ClothFacade.GetRenderBoneIndices()[VertexID],
-				ClothFacade.GetRenderBoneWeights()[VertexID],
-				ClothFacade.GetRenderPosition()[VertexID],
-				&TargetBoneToIndex);
+				SourceMeshToTransferFrom = &WeldedSimMesh; // transfer from sim mesh instead
+			}
+
+			FTransferBoneWeights TransferBoneWeights(SourceMeshToTransferFrom, FSkeletalMeshAttributes::DefaultSkinWeightProfileName);
+			TransferBoneWeights.bUseParallel = TransferSettings.bUseParallel;
+			TransferBoneWeights.TransferMethod = TransferSettings.TransferMethod;
+			TransferBoneWeights.MaxNumInfluences = TransferSettings.MaxNumInfluences;
 
-		}, TransferSettings.bUseParallel ? EParallelForFlags::None : EParallelForFlags::ForceSingleThread);
+			if (TransferBoneWeights.Validate() != EOperationValidationResult::Ok)
+			{
+				UE_LOG(LogChaosClothAssetDataflowNodes, Warning, TEXT("TransferSkinWeightsNode: Transfer method parameters are invalid."));
+				return false;
+			}
 
+			ParallelFor(ClothFacade.GetNumRenderVertices(), [&TransferBoneWeights, &TargetBoneToIndex, &ClothFacade](int32 VertexID)
+			{
+				TransferBoneWeights.TransferWeightsToPoint(ClothFacade.GetRenderBoneIndices()[VertexID],
+					ClothFacade.GetRenderBoneWeights()[VertexID],
+					ClothFacade.GetRenderPosition()[VertexID],
+					&TargetBoneToIndex);
+
+			}, TransferSettings.bUseParallel ? EParallelForFlags::None : EParallelForFlags::ForceSingleThread);
+		}
 		return true;
 	}
 }
@@ -653,7 +689,6 @@ FChaosClothAssetTransferSkinWeightsNode::FChaosClothAssetTransferSkinWeightsNode
 	RegisterInputConnection(&Collection);
 	RegisterOutputConnection(&Collection, &Collection);
 	RegisterInputConnection(&InpaintMask.WeightMap, GET_MEMBER_NAME_CHECKED(FChaosClothAssetWeightedValueNonAnimatableNoLowHighRange, WeightMap));
-	RegisterInputConnection(&TransferWeightsCollection);
 }
 
 void FChaosClothAssetTransferSkinWeightsNode::Evaluate(Dataflow::FContext& Context, const FDataflowOutput* Out) const
@@ -721,13 +756,11 @@ void FChaosClothAssetTransferSkinWeightsNode::Evaluate(Dataflow::FContext& Conte
 			bool bTransferResult = false;
 			if (TransferMethod == EChaosClothAssetTransferSkinWeightsMethod::InpaintWeights)
 			{
-				FManagedArrayCollection TransferWeightsManagedCollection = GetValue<FManagedArrayCollection>(Context, &TransferWeightsCollection);
-				const TSharedRef<FManagedArrayCollection> ClothTransferWeightCollection = MakeShared<FManagedArrayCollection>(MoveTemp(TransferWeightsManagedCollection));
-				bTransferResult = TransferInpaintWeights(TargetRefSkeleton, SourceDynamicMesh, ClothCollection, ClothTransferWeightCollection, TransferSettings);
+				bTransferResult = TransferInpaintWeights(TargetRefSkeleton, SourceDynamicMesh, ClothCollection, TransferSettings);
 			}
 			else if (TransferMethod == EChaosClothAssetTransferSkinWeightsMethod::ClosestPointOnSurface)
 			{
-				bTransferResult = TransferClosestPointOnSurface(TargetRefSkeleton, SourceDynamicMesh, ClothFacade, TransferSettings);
+				bTransferResult = TransferClosestPointOnSurface(TargetRefSkeleton, SourceDynamicMesh, ClothCollection, TransferSettings);
 			}
 			else
 			{
diff --git a/Engine/Plugins/Experimental/ChaosClothAssetEditor/Source/ChaosClothAssetDataflowNodes/Public/ChaosClothAsset/TransferSkinWeightsNode.h b/Engine/Plugins/Experimental/ChaosClothAssetEditor/Source/ChaosClothAssetDataflowNodes/Public/ChaosClothAsset/TransferSkinWeightsNode.h
--- a/Engine/Plugins/Experimental/ChaosClothAssetEditor/Source/ChaosClothAssetDataflowNodes/Public/ChaosClothAsset/TransferSkinWeightsNode.h
+++ b/Engine/Plugins/Experimental/ChaosClothAssetEditor/Source/ChaosClothAssetDataflowNodes/Public/ChaosClothAsset/TransferSkinWeightsNode.h
@@ -131,10 +131,6 @@ public:
 	UPROPERTY(EditAnywhere, Category = "Transfer Skin Weights|Transfer Method", Meta = (DisplayName = "Inpaint Mask", EditCondition="TransferMethod==EChaosClothAssetTransferSkinWeightsMethod::InpaintWeights"))
 	FChaosClothAssetWeightedValueNonAnimatableNoLowHighRange InpaintMask = { TEXT("InpaintMask") };
 
-	/** Optional mesh that will be used for transferring weights to the render mesh instead of the simulation mesh. */
-	UPROPERTY(Meta = (Dataflowinput))
-	FManagedArrayCollection TransferWeightsCollection;
-
 	FChaosClothAssetTransferSkinWeightsNode(const Dataflow::FNodeParameters& InParam, FGuid InGuid = FGuid::NewGuid());
 
 private:
