Chaos Cloth Asset - Fixed the TransferSkinWeight node's weight map name override UI mechanism.

diff --git a/Engine/Plugins/Experimental/ChaosClothAssetEditor/Source/ChaosClothAssetDataflowNodes/Private/ChaosClothAsset/TransferSkinWeightsNode.cpp b/Engine/Plugins/Experimental/ChaosClothAssetEditor/Source/ChaosClothAssetDataflowNodes/Private/ChaosClothAsset/TransferSkinWeightsNode.cpp
--- a/Engine/Plugins/Experimental/ChaosClothAssetEditor/Source/ChaosClothAssetDataflowNodes/Private/ChaosClothAsset/TransferSkinWeightsNode.cpp
+++ b/Engine/Plugins/Experimental/ChaosClothAssetEditor/Source/ChaosClothAssetDataflowNodes/Private/ChaosClothAsset/TransferSkinWeightsNode.cpp
@@ -699,6 +699,9 @@ void FChaosClothAssetTransferSkinWeightsNode::Evaluate(Dataflow::FContext& Conte
 
 	if (Out->IsA<FManagedArrayCollection>(&Collection))
 	{
+		// Update the weight map override
+		InpaintMask.WeightMap_Override = GetValue<FString>(Context, &InpaintMask.WeightMap, UE::Chaos::ClothAsset::FWeightMapTools::NotOverridden);
+
 		// Evaluate inputs
 		FManagedArrayCollection InputCollection = GetValue<FManagedArrayCollection>(Context, &Collection);
 		const TSharedRef<FManagedArrayCollection> ClothCollection = MakeShared<FManagedArrayCollection>(MoveTemp(InputCollection));
@@ -717,8 +720,6 @@ void FChaosClothAssetTransferSkinWeightsNode::Evaluate(Dataflow::FContext& Conte
 				return;
 			}
 
-			InpaintMask.WeightMap_Override = GetValue<FString>(Context, &InpaintMask.WeightMap, FString());
-
 			//
 			// Convert source Skeletal Mesh to Dynamic Mesh.
 			//
diff --git a/Engine/Plugins/Experimental/ChaosClothAssetEditor/Source/ChaosClothAssetDataflowNodes/Public/ChaosClothAsset/WeightedValue.h b/Engine/Plugins/Experimental/ChaosClothAssetEditor/Source/ChaosClothAssetDataflowNodes/Public/ChaosClothAsset/WeightedValue.h
--- a/Engine/Plugins/Experimental/ChaosClothAssetEditor/Source/ChaosClothAssetDataflowNodes/Public/ChaosClothAsset/WeightedValue.h
+++ b/Engine/Plugins/Experimental/ChaosClothAssetEditor/Source/ChaosClothAssetDataflowNodes/Public/ChaosClothAsset/WeightedValue.h
@@ -120,5 +120,5 @@ struct FChaosClothAssetWeightedValueNonAnimatableNoLowHighRange
 
 	/** The weight map override value for when the WeightMap has a connection that replaces the provided weight map value. */
 	UPROPERTY(VisibleAnywhere, Transient, Category = "Weighted Value")
-	mutable FString WeightMap_Override;  // _Override has a special meaning to the property customization, mutable because this property is set while getting the original value
+	mutable FString WeightMap_Override = UE::Chaos::ClothAsset::FWeightMapTools::NotOverridden;  // _Override has a special meaning to the property customization, mutable because this property is set while getting the original value
 };
\ No newline at end of file
