Cloth Editor preview scene: add a toggle to enable/disable PostProcess Blueprints on the SkeletalMeshComponent

diff --git a/Engine/Plugins/Experimental/ChaosClothAssetEditor/Source/ChaosClothAssetEditor/Private/ChaosClothAsset/ClothEditorPreviewScene.cpp b/Engine/Plugins/Experimental/ChaosClothAssetEditor/Source/ChaosClothAssetEditor/Private/ChaosClothAsset/ClothEditorPreviewScene.cpp
--- a/Engine/Plugins/Experimental/ChaosClothAssetEditor/Source/ChaosClothAssetEditor/Private/ChaosClothAsset/ClothEditorPreviewScene.cpp
+++ b/Engine/Plugins/Experimental/ChaosClothAssetEditor/Source/ChaosClothAssetEditor/Private/ChaosClothAsset/ClothEditorPreviewScene.cpp
@@ -161,7 +161,7 @@ FChaosClothPreviewScene::FChaosClothPreviewScene(FPreviewScene::ConstructionValu
 
 	SkeletalMeshComponent = NewObject<USkeletalMeshComponent>(SceneActor);
 	SkeletalMeshComponent->SelectionOverrideDelegate = UPrimitiveComponent::FSelectionOverride::CreateRaw(this, &FChaosClothPreviewScene::IsComponentSelected);
-	SkeletalMeshComponent->SetDisablePostProcessBlueprint(true);
+	SkeletalMeshComponent->SetDisablePostProcessBlueprint(false);
 	SkeletalMeshComponent->RegisterComponentWithWorld(GetWorld());
 
 	ClothComponent = NewObject<UChaosClothComponent>(SceneActor);
@@ -292,6 +292,14 @@ void FChaosClothPreviewScene::SceneDescriptionPropertyChanged(const FName& Prope
 		}
 	}
 	
+	if (PropertyName == GET_MEMBER_NAME_CHECKED(UChaosClothPreviewSceneDescription, bPostProcessBlueprint))
+	{
+		if (SkeletalMeshComponent)
+		{
+			SkeletalMeshComponent->SetDisablePostProcessBlueprint(!PreviewSceneDescription->bPostProcessBlueprint);
+		}
+	}
+
 }
 
 UChaosClothComponent* FChaosClothPreviewScene::GetClothComponent()
diff --git a/Engine/Plugins/Experimental/ChaosClothAssetEditor/Source/ChaosClothAssetEditor/Private/ChaosClothAsset/ClothEditorPreviewScene.h b/Engine/Plugins/Experimental/ChaosClothAssetEditor/Source/ChaosClothAssetEditor/Private/ChaosClothAsset/ClothEditorPreviewScene.h
--- a/Engine/Plugins/Experimental/ChaosClothAssetEditor/Source/ChaosClothAssetEditor/Private/ChaosClothAsset/ClothEditorPreviewScene.h
+++ b/Engine/Plugins/Experimental/ChaosClothAssetEditor/Source/ChaosClothAssetEditor/Private/ChaosClothAsset/ClothEditorPreviewScene.h
@@ -46,6 +46,9 @@ public:
 	UPROPERTY(EditAnywhere, Transient, Category = "SkeletalMesh")
 	TObjectPtr<UAnimationAsset> AnimationAsset;
 
+	UPROPERTY(EditAnywhere, Transient, Category = "SkeletalMesh")
+	bool bPostProcessBlueprint;
+
 	UPROPERTY(EditAnywhere, Transient, Category = "Transform", Meta=(DisplayName="Location"))
 	FVector3d Translation = FVector3d::ZeroVector;
 
