Chaos Cloth Asset - Made the Fabric node base class abstract.

#rnx

diff --git a/Engine/Plugins/Experimental/ChaosClothAssetEditor/Source/ChaosClothAssetDataflowNodes/Public/ChaosClothAsset/SimulationFabricConfigNode.h b/Engine/Plugins/Experimental/ChaosClothAssetEditor/Source/ChaosClothAssetDataflowNodes/Public/ChaosClothAsset/SimulationFabricConfigNode.h
--- a/Engine/Plugins/Experimental/ChaosClothAssetEditor/Source/ChaosClothAssetDataflowNodes/Public/ChaosClothAsset/SimulationFabricConfigNode.h
+++ b/Engine/Plugins/Experimental/ChaosClothAssetEditor/Source/ChaosClothAssetDataflowNodes/Public/ChaosClothAsset/SimulationFabricConfigNode.h
@@ -7,7 +7,7 @@
 #include "SimulationFabricConfigNode.generated.h"
 
 /** Cloth configuration base node to import properties from fabrics . */
-USTRUCT(Meta = (DataflowCloth))
+USTRUCT(Meta = (Abstract))
 struct FChaosClothAssetSimulationFabricConfigNode : public FChaosClothAssetSimulationBaseConfigNode
 {
 	GENERATED_USTRUCT_BODY()
@@ -23,10 +23,7 @@ public:
 	FChaosClothAssetSimulationFabricConfigNode(const Dataflow::FNodeParameters& InParam, FGuid InGuid = FGuid::NewGuid());
 	
 protected :
-	
-	// FChaosClothAssetSimulationBaseConfigNode interface
-	virtual void AddProperties(struct FPropertyHelper& PropertyHelper) const override {}
-	
+
 	/** We can only use fabrics to set property values if the import settings is true and if the cloth facade has fabrics */
 	bool CanUseFabrics(const UE::Chaos::ClothAsset::FCollectionClothFacade& ClothFacade) const
 	{
