CIS - Fixed 'FChaosClothAssetSimulationFabricConfigNode': cannot instantiate abstract class.

#rnx

diff --git a/Engine/Plugins/Experimental/ChaosClothAssetEditor/Source/ChaosClothAssetDataflowNodes/Public/ChaosClothAsset/SimulationFabricConfigNode.h b/Engine/Plugins/Experimental/ChaosClothAssetEditor/Source/ChaosClothAssetDataflowNodes/Public/ChaosClothAsset/SimulationFabricConfigNode.h
--- a/Engine/Plugins/Experimental/ChaosClothAssetEditor/Source/ChaosClothAssetDataflowNodes/Public/ChaosClothAsset/SimulationFabricConfigNode.h
+++ b/Engine/Plugins/Experimental/ChaosClothAssetEditor/Source/ChaosClothAssetDataflowNodes/Public/ChaosClothAsset/SimulationFabricConfigNode.h
@@ -42,3 +42,11 @@ protected :
 		const TArray<FName>& SimilarPropertyNames, ECollectionPropertyFlags PropertyFlags = ECollectionPropertyFlags::None);
 };
 
+template<>
+struct TStructOpsTypeTraits<FChaosClothAssetSimulationFabricConfigNode> : public TStructOpsTypeTraitsBase2<FChaosClothAssetSimulationFabricConfigNode>
+{
+	enum
+	{
+		WithPureVirtual = true,
+	};
+};
