Chaos Cloth Asset: Dataflow nodes to edit sim geo

Update DeleteElementNode to accept a selection set to delete.
Add an "AddStitchNode" which currently just lets you merge vertices by creating a seam that connects all of the selected vertices.

diff --git a/Engine/Plugins/Experimental/ChaosClothAssetEditor/Source/ChaosClothAssetDataflowNodes/Private/ChaosClothAsset/AddStitchNode.cpp b/Engine/Plugins/Experimental/ChaosClothAssetEditor/Source/ChaosClothAssetDataflowNodes/Private/ChaosClothAsset/AddStitchNode.cpp
new file mode 100644
--- /dev/null
+++ b/Engine/Plugins/Experimental/ChaosClothAssetEditor/Source/ChaosClothAssetDataflowNodes/Private/ChaosClothAsset/AddStitchNode.cpp
@@ -0,0 +1,103 @@
+// Copyright Epic Games, Inc. All Rights Reserved.
+
+#include "ChaosClothAsset/AddStitchNode.h"
+#include "ChaosClothAsset/ClothCollectionGroup.h"
+#include "ChaosClothAsset/ClothDataflowTools.h"
+#include "ChaosClothAsset/ClothGeometryTools.h"
+#include "ChaosClothAsset/CollectionClothFacade.h"
+#include "ChaosClothAsset/CollectionClothSelectionFacade.h"
+#include "ChaosClothAsset/WeightedValue.h"
+#include "Dataflow/DataflowInputOutput.h"
+
+#include UE_INLINE_GENERATED_CPP_BY_NAME(AddStitchNode)
+
+#define LOCTEXT_NAMESPACE "ChaosClothAssetAddStitchNode"
+
+FChaosClothAssetAddStitchNode::FChaosClothAssetAddStitchNode(const Dataflow::FNodeParameters& InParam, FGuid InGuid)
+	: FDataflowNode(InParam, InGuid)
+{
+	RegisterInputConnection(&Collection);
+	RegisterOutputConnection(&Collection, &Collection);
+	RegisterInputConnection(&MergeToSingleVertexSelection.StringValue, GET_MEMBER_NAME_CHECKED(FChaosClothAssetConnectableIStringValue, StringValue));
+}
+
+void FChaosClothAssetAddStitchNode::Evaluate(Dataflow::FContext& Context, const FDataflowOutput* Out) const
+{
+	if (Out->IsA<FManagedArrayCollection>(&Collection))
+	{
+		using namespace UE::Chaos::ClothAsset;
+
+		// Evaluate in collection
+		FManagedArrayCollection InCollection = GetValue<FManagedArrayCollection>(Context, &Collection);
+		const TSharedRef<FManagedArrayCollection> ClothCollection = MakeShared<FManagedArrayCollection>(MoveTemp(InCollection));
+
+		FCollectionClothFacade Cloth(ClothCollection);
+		FCollectionClothSelectionConstFacade SelectionFacade(ClothCollection);
+		const FName InSelectionName(*GetValue<FString>(Context, &MergeToSingleVertexSelection.StringValue));
+		MergeToSingleVertexSelection.StringValue_Override = GetValue<FString>(Context, &MergeToSingleVertexSelection.StringValue, UE::Chaos::ClothAsset::FWeightMapTools::NotOverridden);
+
+		if (Cloth.IsValid() && SelectionFacade.IsValid() && InSelectionName != NAME_None)
+		{
+			const TSet<int32>* const SelectionSet = SelectionFacade.FindSelectionSet(InSelectionName);
+			if (SelectionSet && SelectionSet->Num() > 1)
+			{
+				const FName& SelectionGroup = SelectionFacade.GetSelectionGroup(InSelectionName);
+				TArray<FIntVector2> StitchPairs;
+				if (SelectionGroup == ClothCollectionGroup::SimVertices2D)
+				{
+					StitchPairs.Reserve(SelectionSet->Num() - 1);
+
+					int32 PrevIndex = INDEX_NONE;
+					for (const int32 Index : *SelectionSet)
+					{
+						if (PrevIndex != INDEX_NONE)
+						{
+							StitchPairs.Add({ PrevIndex, Index });
+						}
+						PrevIndex = Index;
+					}
+				}
+				else if (SelectionGroup == ClothCollectionGroup::SimVertices3D)
+				{
+					TConstArrayView<TArray<int32>> SimVertex2DLookup = Cloth.GetSimVertex2DLookup();
+					StitchPairs.Reserve(SelectionSet->Num() - 1);
+					int32 PrevIndex = INDEX_NONE;
+					for (const int32 Index : *SelectionSet)
+					{
+						if (SimVertex2DLookup.IsValidIndex(Index) && !SimVertex2DLookup[Index].IsEmpty())
+						{
+							// Just pick first 2D vertex that corresponds with this 3D vertex
+							const int32 Index2D = SimVertex2DLookup[Index][0];
+							if (PrevIndex != INDEX_NONE && Index2D != INDEX_NONE)
+							{
+								StitchPairs.Add({ PrevIndex, Index2D });
+							}
+							PrevIndex = Index2D;
+						}
+					}
+				}
+				else
+				{
+					FClothDataflowTools::LogAndToastWarning(*this,
+						LOCTEXT("SelectionTypeNotCorrectHeadline", "Invalid selection group."),
+						FText::Format(LOCTEXT("SelectionTypeNotCorrectDetails", "Selection \"{0}\" does not have its index dependency group set to \"{1}\" or \"{2}\"."),
+							FText::FromName(InSelectionName),
+							FText::FromName(ClothCollectionGroup::SimVertices2D),
+							FText::FromName(ClothCollectionGroup::SimVertices3D)));
+
+				}
+
+				if (!StitchPairs.IsEmpty())
+				{
+					FCollectionClothSeamFacade Seam = Cloth.AddGetSeam();
+					Seam.Initialize(TConstArrayView<FIntVector2>(StitchPairs));
+
+					FClothGeometryTools::CleanupAndCompactMesh(ClothCollection);
+				}
+			}
+		}
+		SetValue(Context, MoveTemp(*ClothCollection), &Collection);
+	}
+}
+
+#undef LOCTEXT_NAMESPACE
diff --git a/Engine/Plugins/Experimental/ChaosClothAssetEditor/Source/ChaosClothAssetDataflowNodes/Private/ChaosClothAsset/ChaosClothAssetDataflowNodesModule.cpp b/Engine/Plugins/Experimental/ChaosClothAssetEditor/Source/ChaosClothAssetDataflowNodes/Private/ChaosClothAsset/ChaosClothAssetDataflowNodesModule.cpp
--- a/Engine/Plugins/Experimental/ChaosClothAssetEditor/Source/ChaosClothAssetDataflowNodes/Private/ChaosClothAsset/ChaosClothAssetDataflowNodesModule.cpp
+++ b/Engine/Plugins/Experimental/ChaosClothAssetEditor/Source/ChaosClothAssetDataflowNodes/Private/ChaosClothAsset/ChaosClothAssetDataflowNodesModule.cpp
@@ -1,5 +1,6 @@
 // Copyright Epic Games, Inc. All Rights Reserved.
 
+#include "ChaosClothAsset/AddStitchNode.h"
 #include "ChaosClothAsset/AddWeightMapNode.h"
 #include "ChaosClothAsset/BindToRootBoneNode.h"
 #include "ChaosClothAsset/ColorScheme.h"
@@ -64,6 +65,7 @@ namespace UE::Chaos::ClothAsset
 	{
 		static void RegisterDataflowNodes()
 		{
+			DATAFLOW_NODE_REGISTER_CREATION_FACTORY(FChaosClothAssetAddStitchNode);
 			DATAFLOW_NODE_REGISTER_CREATION_FACTORY(FChaosClothAssetAddWeightMapNode);
 			DATAFLOW_NODE_REGISTER_CREATION_FACTORY(FChaosClothAssetBindToRootBoneNode);
 			DATAFLOW_NODE_REGISTER_CREATION_FACTORY(FChaosClothAssetCopySimulationToRenderMeshNode);
diff --git a/Engine/Plugins/Experimental/ChaosClothAssetEditor/Source/ChaosClothAssetDataflowNodes/Private/ChaosClothAsset/DeleteElementNode.cpp b/Engine/Plugins/Experimental/ChaosClothAssetEditor/Source/ChaosClothAssetDataflowNodes/Private/ChaosClothAsset/DeleteElementNode.cpp
--- a/Engine/Plugins/Experimental/ChaosClothAssetEditor/Source/ChaosClothAssetDataflowNodes/Private/ChaosClothAsset/DeleteElementNode.cpp
+++ b/Engine/Plugins/Experimental/ChaosClothAssetEditor/Source/ChaosClothAssetDataflowNodes/Private/ChaosClothAsset/DeleteElementNode.cpp
@@ -2,7 +2,11 @@
 
 #include "ChaosClothAsset/DeleteElementNode.h"
 #include "ChaosClothAsset/ClothGeometryTools.h"
+#include "ChaosClothAsset/ClothDataflowTools.h"
 #include "ChaosClothAsset/CollectionClothFacade.h"
+#include "ChaosClothAsset/ClothCollectionGroup.h"
+#include "ChaosClothAsset/CollectionClothSelectionFacade.h"
+#include "ChaosClothAsset/WeightedValue.h"
 #include "Dataflow/DataflowInputOutput.h"
 
 #include UE_INLINE_GENERATED_CPP_BY_NAME(DeleteElementNode)
@@ -14,6 +18,7 @@ FChaosClothAssetDeleteElementNode::FChaosClothAssetDeleteElementNode(const Dataf
 {
 	RegisterInputConnection(&Collection);
 	RegisterOutputConnection(&Collection, &Collection);
+	RegisterInputConnection(&SelectionName.StringValue, GET_MEMBER_NAME_CHECKED(FChaosClothAssetConnectableIStringValue, StringValue));
 }
 
 void FChaosClothAssetDeleteElementNode::Evaluate(Dataflow::FContext& Context, const FDataflowOutput* Out) const
@@ -29,247 +34,409 @@ void FChaosClothAssetDeleteElementNode::Evaluate(Dataflow::FContext& Context, co
 		FCollectionClothFacade Cloth(ClothCollection);
 		if (Cloth.IsValid())
 		{
-			switch (ElementType)
+			auto DeleteElementsByType = [&ClothCollection, &Cloth](const FName& GroupName, const TSet<int32>& Set)
 			{
-			case EChaosClothAssetElementType::None:
-			default:
-				break;
-			case EChaosClothAssetElementType::SimMesh:
-			{
-				FClothGeometryTools::DeleteSimMesh(ClothCollection);
-			} break;
-			case EChaosClothAssetElementType::RenderMesh:
-			{
-				FClothGeometryTools::DeleteRenderMesh(ClothCollection);
-			} break;
-			case EChaosClothAssetElementType::SimPattern:
-			{
-				if (Elements.IsEmpty())
-				{
-					// Remove all sim patterns
-					Cloth.SetNumSimPatterns(0);
-				}
-				else
+				if(GroupName == ClothCollectionGroup::SimPatterns)
 				{
-					TArray<int32> PatternsToRemove;
-					PatternsToRemove.Reserve(Elements.Num());
-					for (int32 ElementIndex = 0; ElementIndex < Cloth.GetNumSimPatterns(); ++ElementIndex)
+					if (Set.IsEmpty())
 					{
-						if (Elements.Find(ElementIndex) != INDEX_NONE)
-						{
-							PatternsToRemove.Add(ElementIndex);
-						}
+						// Remove all sim patterns
+						Cloth.SetNumSimPatterns(0);
 					}
-					if (PatternsToRemove.Num())
+					else
 					{
-						Cloth.RemoveSimPatterns(PatternsToRemove);
+						TArray<int32> PatternsToRemove;
+						PatternsToRemove.Reserve(Set.Num());
+						for (int32 ElementIndex = 0; ElementIndex < Cloth.GetNumSimPatterns(); ++ElementIndex)
+						{
+							if (Set.Contains(ElementIndex))
+							{
+								PatternsToRemove.Add(ElementIndex);
+							}
+						}
+						if (PatternsToRemove.Num())
+						{
+							Cloth.RemoveSimPatterns(PatternsToRemove);
+						}
 					}
 				}
-			} break;
-			case EChaosClothAssetElementType::RenderPattern:
-			{
-				if (Elements.IsEmpty())
+				else if(GroupName == ClothCollectionGroup::RenderPatterns)
 				{
-					// Remove all render patterns
-					Cloth.SetNumRenderPatterns(0);
+					if (Set.IsEmpty())
+					{
+						// Remove all render patterns
+						Cloth.SetNumRenderPatterns(0);
+					}
+					else
+					{
+						TArray<int32> PatternsToRemove;
+						PatternsToRemove.Reserve(Set.Num());
+						for (int32 ElementIndex = 0; ElementIndex < Cloth.GetNumRenderPatterns(); ++ElementIndex)
+						{
+							if (Set.Contains(ElementIndex))
+							{
+								PatternsToRemove.Add(ElementIndex);
+							}
+						}
+						if (PatternsToRemove.Num())
+						{
+							Cloth.RemoveRenderPatterns(PatternsToRemove);
+						}
+					}
 				}
-				else
+				else if(GroupName == ClothCollectionGroup::SimVertices2D)
 				{
-					TArray<int32> PatternsToRemove;
-					PatternsToRemove.Reserve(Elements.Num());
-					for (int32 ElementIndex = 0; ElementIndex < Cloth.GetNumRenderPatterns(); ++ElementIndex)
+					if (Set.IsEmpty())
 					{
-						if (Elements.Find(ElementIndex) != INDEX_NONE)
+						for (int32 ElementIndex = 0; ElementIndex < Cloth.GetNumSimPatterns(); ++ElementIndex)
 						{
-							PatternsToRemove.Add(ElementIndex);
+							FCollectionClothSimPatternFacade Pattern = Cloth.GetSimPattern(ElementIndex);
+							Pattern.RemoveAllSimVertices2D();
 						}
 					}
-					if (PatternsToRemove.Num())
+					else
 					{
-						Cloth.RemoveRenderPatterns(PatternsToRemove);
+						TArray<int32> SortedToDeleteList;
+						SortedToDeleteList.Reserve(Set.Num());
+						for (int32 PatternIndex = Cloth.GetNumSimPatterns() - 1; PatternIndex >= 0; --PatternIndex)
+						{
+							FCollectionClothSimPatternFacade Pattern = Cloth.GetSimPattern(PatternIndex);
+							const int32 VertexOffset = Pattern.GetSimVertices2DOffset();
+							SortedToDeleteList.Reset();
+							for (int32 VertexIndex = 0; VertexIndex < Pattern.GetNumSimVertices2D(); ++VertexIndex)
+							{
+								if (Set.Contains(VertexIndex + VertexOffset))
+								{
+									SortedToDeleteList.Add(VertexIndex);
+								}
+							}
+							if (SortedToDeleteList.Num())
+							{
+								Pattern.RemoveSimVertices2D(SortedToDeleteList);
+							}
+						}
 					}
 				}
-			} break;
-			case EChaosClothAssetElementType::SimVertex2D:
-			{
-				if (Elements.IsEmpty())
+				else if(GroupName == ClothCollectionGroup::SimVertices3D)
 				{
-					for (int32 ElementIndex = 0; ElementIndex < Cloth.GetNumSimPatterns(); ++ElementIndex)
+					if (Set.IsEmpty())
 					{
-						FCollectionClothSimPatternFacade Pattern = Cloth.GetSimPattern(ElementIndex);
-						Pattern.RemoveAllSimVertices2D();
+						Cloth.RemoveAllSimVertices3D();
 					}
-				}
-				else
-				{
-					TArray<int32> SortedToDeleteList;
-					SortedToDeleteList.Reserve(Elements.Num());
-					for (int32 PatternIndex = Cloth.GetNumSimPatterns()-1; PatternIndex >= 0; --PatternIndex)
+					else
 					{
-						FCollectionClothSimPatternFacade Pattern = Cloth.GetSimPattern(PatternIndex);
-						const int32 VertexOffset = Pattern.GetSimVertices2DOffset();
-						SortedToDeleteList.Reset();
-						for (int32 VertexIndex = 0; VertexIndex < Pattern.GetNumSimVertices2D(); ++VertexIndex)
+						TArray<int32> SortedToDeleteList;
+						SortedToDeleteList.Reserve(Set.Num());
+						for (int32 VertexIndex = 0; VertexIndex < Cloth.GetNumSimVertices3D(); ++VertexIndex)
 						{
-							if (Elements.Find(VertexIndex + VertexOffset) != INDEX_NONE)
+							if (Set.Contains(VertexIndex))
 							{
 								SortedToDeleteList.Add(VertexIndex);
 							}
 						}
-						if (SortedToDeleteList.Num())
-						{
-							Pattern.RemoveSimVertices2D(SortedToDeleteList);
-						}
+						Cloth.RemoveSimVertices3D(SortedToDeleteList);
 					}
 				}
-			} break;
-			case EChaosClothAssetElementType::SimVertex3D:
-			{
-				if (Elements.IsEmpty())
-				{
-					Cloth.RemoveAllSimVertices3D();
-				}
-				else
+				else if(GroupName == ClothCollectionGroup::RenderVertices)
 				{
-					TArray<int32> SortedToDeleteList;
-					SortedToDeleteList.Reserve(Elements.Num());
-					for (int32 VertexIndex = 0; VertexIndex < Cloth.GetNumSimVertices3D(); ++VertexIndex)
+					if (Set.IsEmpty())
 					{
-						if (Elements.Find(VertexIndex) != INDEX_NONE)
+						for (int32 ElementIndex = 0; ElementIndex < Cloth.GetNumRenderPatterns(); ++ElementIndex)
 						{
-							SortedToDeleteList.Add(VertexIndex);
+							FCollectionClothRenderPatternFacade Pattern = Cloth.GetRenderPattern(ElementIndex);
+							Pattern.SetNumRenderVertices(0);
 						}
 					}
-					Cloth.RemoveSimVertices3D(SortedToDeleteList);
-				}
-			} break;
-			case EChaosClothAssetElementType::RenderVertex:
-			{
-				if (Elements.IsEmpty())
-				{
-					for (int32 ElementIndex = 0; ElementIndex < Cloth.GetNumRenderPatterns(); ++ElementIndex)
+					else
 					{
-						FCollectionClothRenderPatternFacade Pattern = Cloth.GetRenderPattern(ElementIndex);
-						Pattern.SetNumRenderVertices(0);
+						TArray<int32> SortedToDeleteList;
+						SortedToDeleteList.Reserve(Set.Num());
+						for (int32 PatternIndex = Cloth.GetNumRenderPatterns() - 1; PatternIndex >= 0; --PatternIndex)
+						{
+							FCollectionClothRenderPatternFacade Pattern = Cloth.GetRenderPattern(PatternIndex);
+							const int32 VertexOffset = Pattern.GetRenderVerticesOffset();
+							SortedToDeleteList.Reset();
+							for (int32 VertexIndex = 0; VertexIndex < Pattern.GetNumRenderVertices(); ++VertexIndex)
+							{
+								if (Set.Contains(VertexIndex + VertexOffset))
+								{
+									SortedToDeleteList.Add(VertexIndex);
+								}
+							}
+							if (SortedToDeleteList.Num())
+							{
+								Pattern.RemoveRenderVertices(SortedToDeleteList);
+							}
+						}
 					}
 				}
-				else
+				else if(GroupName == ClothCollectionGroup::SimFaces)
 				{
-					TArray<int32> SortedToDeleteList;
-					SortedToDeleteList.Reserve(Elements.Num());
-					for (int32 PatternIndex = Cloth.GetNumRenderPatterns()-1; PatternIndex >= 0; --PatternIndex)
+					if (Set.IsEmpty())
 					{
-						FCollectionClothRenderPatternFacade Pattern = Cloth.GetRenderPattern(PatternIndex);
-						const int32 VertexOffset = Pattern.GetRenderVerticesOffset();
-						SortedToDeleteList.Reset();
-						for (int32 VertexIndex = 0; VertexIndex < Pattern.GetNumRenderVertices(); ++VertexIndex)
+						for (int32 ElementIndex = 0; ElementIndex < Cloth.GetNumSimPatterns(); ++ElementIndex)
 						{
-							if (Elements.Find(VertexIndex + VertexOffset) != INDEX_NONE)
-							{
-								SortedToDeleteList.Add(VertexIndex);
-							}
+							FCollectionClothSimPatternFacade Pattern = Cloth.GetSimPattern(ElementIndex);
+							Pattern.SetNumSimFaces(0);
 						}
-						if (SortedToDeleteList.Num())
+					}
+					else
+					{
+						TArray<int32> SortedToDeleteList;
+						SortedToDeleteList.Reserve(Set.Num());
+						for (int32 PatternIndex = Cloth.GetNumSimPatterns() - 1; PatternIndex >= 0; --PatternIndex)
 						{
-							Pattern.RemoveRenderVertices(SortedToDeleteList);
+							FCollectionClothSimPatternFacade Pattern = Cloth.GetSimPattern(PatternIndex);
+							const int32 FaceOffset = Pattern.GetSimFacesOffset();
+							SortedToDeleteList.Reset();
+							for (int32 FaceIndex = 0; FaceIndex < Pattern.GetNumSimFaces(); ++FaceIndex)
+							{
+								if (Set.Contains(FaceIndex + FaceOffset))
+								{
+									SortedToDeleteList.Add(FaceIndex);
+								}
+							}
+							if (SortedToDeleteList.Num())
+							{
+								Pattern.RemoveSimFaces(SortedToDeleteList);
+							}
 						}
 					}
 				}
-			} break;
-			case EChaosClothAssetElementType::SimFace:
-			{
-				if (Elements.IsEmpty())
+				else if(GroupName == ClothCollectionGroup::RenderFaces)
 				{
-					for (int32 ElementIndex = 0; ElementIndex < Cloth.GetNumSimPatterns(); ++ElementIndex)
+					if (Set.IsEmpty())
 					{
-						FCollectionClothSimPatternFacade Pattern = Cloth.GetSimPattern(ElementIndex);
-						Pattern.SetNumSimFaces(0);
+						for (int32 PatternIndex = 0; PatternIndex < Cloth.GetNumRenderPatterns(); ++PatternIndex)
+						{
+							FCollectionClothRenderPatternFacade Pattern = Cloth.GetRenderPattern(PatternIndex);
+							Pattern.SetNumRenderFaces(0);
+						}
+					}
+					else
+					{
+						TArray<int32> SortedToDeleteList;
+						SortedToDeleteList.Reserve(Set.Num());
+						for (int32 PatternIndex = Cloth.GetNumRenderPatterns() - 1; PatternIndex >= 0; --PatternIndex)
+						{
+							FCollectionClothRenderPatternFacade Pattern = Cloth.GetRenderPattern(PatternIndex);
+							const int32 FaceOffset = Pattern.GetRenderFacesOffset();
+							SortedToDeleteList.Reset();
+							for (int32 FaceIndex = 0; FaceIndex < Pattern.GetNumRenderFaces(); ++FaceIndex)
+							{
+								if (Set.Contains(FaceIndex + FaceOffset))
+								{
+									SortedToDeleteList.Add(FaceIndex);
+								}
+							}
+							if (SortedToDeleteList.Num())
+							{
+								Pattern.RemoveRenderFaces(SortedToDeleteList);
+							}
+						}
 					}
 				}
-				else
+				else if(GroupName == ClothCollectionGroup::Seams)
 				{
-					TArray<int32> SortedToDeleteList;
-					SortedToDeleteList.Reserve(Elements.Num());
-					for (int32 PatternIndex = Cloth.GetNumSimPatterns()-1; PatternIndex >= 0; --PatternIndex)
+					if (Set.IsEmpty())
+					{
+						// Remove all seams
+						Cloth.SetNumSeams(0);
+					}
+					else
 					{
-						FCollectionClothSimPatternFacade Pattern = Cloth.GetSimPattern(PatternIndex);
-						const int32 FaceOffset = Pattern.GetSimFacesOffset();
-						SortedToDeleteList.Reset();
-						for (int32 FaceIndex = 0; FaceIndex < Pattern.GetNumSimFaces(); ++FaceIndex)
+						TArray<int32> SeamsToRemove;
+						SeamsToRemove.Reserve(Set.Num());
+						for (int32 SeamIndex = 0; SeamIndex < Cloth.GetNumSeams(); ++SeamIndex)
 						{
-							if (Elements.Find(FaceIndex + FaceOffset) != INDEX_NONE)
+							if (Set.Contains(SeamIndex))
 							{
-								SortedToDeleteList.Add(FaceIndex);
+								SeamsToRemove.Add(SeamIndex);
 							}
 						}
-						if (SortedToDeleteList.Num())
+						if (SeamsToRemove.Num())
 						{
-							Pattern.RemoveSimFaces(SortedToDeleteList);
+							Cloth.RemoveSeams(SeamsToRemove);
 						}
 					}
 				}
-			} break;
-			case EChaosClothAssetElementType::RenderFace:
-			{
-				if (Elements.IsEmpty())
+				else if(GroupName == ClothCollectionGroup::Fabrics)
 				{
-					for (int32 PatternIndex = 0; PatternIndex < Cloth.GetNumRenderPatterns(); ++PatternIndex)
+					if (Set.IsEmpty())
 					{
-						FCollectionClothRenderPatternFacade Pattern = Cloth.GetRenderPattern(PatternIndex);
-						Pattern.SetNumRenderFaces(0);
+						// Remove all fabrics
+						Cloth.SetNumFabrics(0);
 					}
-				}
-				else
-				{
-					TArray<int32> SortedToDeleteList;
-					SortedToDeleteList.Reserve(Elements.Num());
-					for (int32 PatternIndex = Cloth.GetNumRenderPatterns()-1; PatternIndex >= 0; --PatternIndex)
+					else
 					{
-						FCollectionClothRenderPatternFacade Pattern = Cloth.GetRenderPattern(PatternIndex);
-						const int32 FaceOffset = Pattern.GetRenderFacesOffset();
-						SortedToDeleteList.Reset();
-						for (int32 FaceIndex = 0; FaceIndex < Pattern.GetNumRenderFaces(); ++FaceIndex)
+						TArray<int32> FabricsToRemove;
+						FabricsToRemove.Reserve(Set.Num());
+						for (int32 FabricIndex = 0; FabricIndex < Cloth.GetNumFabrics(); ++FabricIndex)
 						{
-							if (Elements.Find(FaceIndex + FaceOffset) != INDEX_NONE)
+							if (Set.Contains(FabricIndex))
 							{
-								SortedToDeleteList.Add(FaceIndex);
+								FabricsToRemove.Add(FabricIndex);
 							}
 						}
-						if (SortedToDeleteList.Num())
+						if (FabricsToRemove.Num())
 						{
-							Pattern.RemoveRenderFaces(SortedToDeleteList);
+							Cloth.RemoveFabrics(FabricsToRemove);
 						}
 					}
 				}
-			} break;
-			case EChaosClothAssetElementType::Seam:
+			};
+
+
+			// Selection set
+			FCollectionClothSelectionFacade SelectionFacade(ClothCollection);
+			const FName InSelectionName(*GetValue<FString>(Context, &SelectionName.StringValue));
+			SelectionName.StringValue_Override = GetValue<FString>(Context, &SelectionName.StringValue, UE::Chaos::ClothAsset::FWeightMapTools::NotOverridden);
+
+			FName SelectionGroup;
+			TSet<int32> SelectedElements;
+			bool bValidSelectionGroup = false;
+
+			if (SelectionFacade.IsValid() && InSelectionName != NAME_None)
 			{
-				if (Elements.IsEmpty())
+				SelectionGroup = SelectionFacade.GetSelectionGroup(InSelectionName);
+				if (const TSet<int32>* const SelectionSet = SelectionFacade.FindSelectionSet(InSelectionName))
 				{
-					// Remove all seams
-					Cloth.SetNumSeams(0);
-				}
-				else
-				{
-					TArray<int32> SeamsToRemove;
-					SeamsToRemove.Reserve(Elements.Num());
-					for (int32 SeamIndex = 0; SeamIndex < Cloth.GetNumSeams(); ++SeamIndex)
+					if (SelectionGroup == ClothCollectionGroup::Seams ||
+						SelectionGroup == ClothCollectionGroup::SimPatterns ||
+						SelectionGroup == ClothCollectionGroup::RenderPatterns ||
+						SelectionGroup == ClothCollectionGroup::SimFaces ||
+						SelectionGroup == ClothCollectionGroup::SimVertices2D ||
+						SelectionGroup == ClothCollectionGroup::SimVertices3D ||
+						SelectionGroup == ClothCollectionGroup::RenderFaces ||
+						SelectionGroup == ClothCollectionGroup::RenderVertices ||
+						SelectionGroup == ClothCollectionGroup::Fabrics)
 					{
-						if (Elements.Find(SeamIndex) != INDEX_NONE)
-						{
-							SeamsToRemove.Add(SeamIndex);
-						}
+						SelectedElements = *SelectionSet;
+						bValidSelectionGroup = true;
 					}
-					if (SeamsToRemove.Num())
+					else
 					{
-						Cloth.RemoveSeams(SeamsToRemove);
+						FClothDataflowTools::LogAndToastWarning(*this,
+							LOCTEXT("SelectionTypeNotCorrectHeadline", "Invalid selection group."),
+							FText::Format(LOCTEXT("SelectionTypeNotCorrectDetails", "Selection \"{0}\" does not have its index dependency group set to a type known to the DeleteElement node."),
+								FText::FromName(InSelectionName)));
+
+					}
+
+					if (bValidSelectionGroup && !SelectedElements.IsEmpty())
+					{
+						// Remove the selection since all of its elements will be deleted.
+						SelectionFacade.RemoveSelectionSet(InSelectionName);
 					}
 				}
-			}break;
+			}
+
+			const FName GroupName(*Group.Name);
+			if (bValidSelectionGroup && SelectionGroup == GroupName)
+			{
+				// Merge selections
+				if (Elements.IsEmpty())
+				{
+					// Empty elements means delete everything.
+					SelectedElements.Reset();
+				}
+				else
+				{
+					SelectedElements.Append(Elements);
+				}
+
+				DeleteElementsByType(SelectionGroup, SelectedElements);
+			}
+			else
+			{
+				if (!SelectedElements.IsEmpty())
+				{
+					DeleteElementsByType(SelectionGroup, SelectedElements);
+				}
+
+				DeleteElementsByType(GroupName, TSet<int32>(Elements));
+			}
+
+			if (bDeleteSimMesh)
+			{
+				FClothGeometryTools::DeleteSimMesh(ClothCollection);
+			}
+
+			if (bDeleteRenderMesh)
+			{
+				FClothGeometryTools::DeleteRenderMesh(ClothCollection);
 			}
 		}
 		SetValue(Context, MoveTemp(*ClothCollection), &Collection);
 	}
 }
 
+void FChaosClothAssetDeleteElementNode::OnSelected(Dataflow::FContext& Context)
+{
+	// Re-evaluate the input collection
+	const FManagedArrayCollection& SelectionCollection = GetValue<FManagedArrayCollection>(Context, &Collection);
+
+	// Update the list of used group for the UI customization
+	CachedCollectionGroupNames = SelectionCollection.GroupNames();
+}
+
+void FChaosClothAssetDeleteElementNode::OnDeselected()
+{
+	// Clean up, to avoid another toolkit picking up the wrong context evaluation
+	CachedCollectionGroupNames.Reset();
+}
+
+void FChaosClothAssetDeleteElementNode::Serialize(FArchive& Ar)
+{
+	using namespace UE::Chaos::ClothAsset;
+
+	// This is just for convenience and can be removed post 5.4 once the plugin loses its experimental status
+	PRAGMA_DISABLE_DEPRECATION_WARNINGS
+	if (Ar.IsLoading() && ElementType_DEPRECATED != EChaosClothAssetElementType::Deprecated)
+	{
+		switch (ElementType_DEPRECATED)
+		{
+		case EChaosClothAssetElementType::None:
+		default:
+			break;
+		case EChaosClothAssetElementType::SimMesh:
+			bDeleteSimMesh = true;
+			break;
+		case EChaosClothAssetElementType::RenderMesh:
+			bDeleteRenderMesh = true;
+			break;
+		case EChaosClothAssetElementType::SimPattern:
+			Group.Name = ClothCollectionGroup::SimPatterns.ToString();
+			break;
+		case EChaosClothAssetElementType::RenderPattern:
+			Group.Name = ClothCollectionGroup::RenderPatterns.ToString();
+			break;
+		case EChaosClothAssetElementType::SimVertex2D:
+			Group.Name = ClothCollectionGroup::SimVertices2D.ToString();
+			break;
+		case EChaosClothAssetElementType::SimVertex3D:
+			Group.Name = ClothCollectionGroup::SimVertices3D.ToString();
+			break;
+		case EChaosClothAssetElementType::RenderVertex:
+			Group.Name = ClothCollectionGroup::RenderVertices.ToString();
+			break;
+		case EChaosClothAssetElementType::SimFace:
+			Group.Name = ClothCollectionGroup::SimFaces.ToString();
+			break;
+		case EChaosClothAssetElementType::RenderFace:
+			Group.Name = ClothCollectionGroup::RenderFaces.ToString();
+			break;
+		case EChaosClothAssetElementType::Seam:
+			Group.Name = ClothCollectionGroup::Seams.ToString();
+			break;
+		}
+		ElementType_DEPRECATED = EChaosClothAssetElementType::Deprecated;  // This is only for clarity since the Type property won't be saved from now on
+
+		FClothDataflowTools::LogAndToastWarning(*this,
+			LOCTEXT("DeprecatedDeleteElementType", "Outdated Dataflow asset."),
+				LOCTEXT("DeprecatedDeleteElementType", "This node is out of data and contain deprecated data. The asset needs to be re-saved before it stops working at the next version update."));
+	}
+	PRAGMA_ENABLE_DEPRECATION_WARNINGS
+}
+
 #undef LOCTEXT_NAMESPACE
diff --git a/Engine/Plugins/Experimental/ChaosClothAssetEditor/Source/ChaosClothAssetDataflowNodes/Private/ChaosClothAsset/SelectionGroupCustomization.cpp b/Engine/Plugins/Experimental/ChaosClothAssetEditor/Source/ChaosClothAssetDataflowNodes/Private/ChaosClothAsset/SelectionGroupCustomization.cpp
--- a/Engine/Plugins/Experimental/ChaosClothAssetEditor/Source/ChaosClothAssetDataflowNodes/Private/ChaosClothAsset/SelectionGroupCustomization.cpp
+++ b/Engine/Plugins/Experimental/ChaosClothAssetEditor/Source/ChaosClothAssetDataflowNodes/Private/ChaosClothAsset/SelectionGroupCustomization.cpp
@@ -2,6 +2,7 @@
 
 #include "ChaosClothAsset/SelectionGroupCustomization.h"
 #include "ChaosClothAsset/SelectionNode.h"
+#include "ChaosClothAsset/DeleteElementNode.h"
 #include "ChaosClothAsset/ClothDataflowTools.h"  // For MakeCollectionName
 #include "Widgets/Input/SEditableTextBox.h"
 #include "Widgets/Input/SComboButton.h"
@@ -107,14 +108,21 @@ namespace UE::Chaos::ClothAsset
 		// Find all group names in the parent selection node's collection
 		GroupNames.Reset();
 
+		TArray<FName> CollectionGroupNames;
 		if (const FChaosClothAssetSelectionNode* const SelectionNode = GetOwnerStruct<FChaosClothAssetSelectionNode>())
 		{
-			const TArray<FName> CollectionGroupNames = SelectionNode->GetCachedCollectionGroupNames();
+			CollectionGroupNames = SelectionNode->GetCachedCollectionGroupNames();
 
-			for (const FName& CollectionGroupName : CollectionGroupNames)
-			{
-				GroupNames.Add(MakeShareable(new FText(FText::FromName(CollectionGroupName))));
-			}
+		}
+		else if (const FChaosClothAssetDeleteElementNode* const DeleteNode =
+			GetOwnerStruct<FChaosClothAssetDeleteElementNode>())
+		{
+			CollectionGroupNames = DeleteNode->GetCachedCollectionGroupNames();
+		}
+
+		for (const FName& CollectionGroupName : CollectionGroupNames)
+		{
+			GroupNames.Add(MakeShareable(new FText(FText::FromName(CollectionGroupName))));
 		}
 
 		return SNew(SVerticalBox)
diff --git a/Engine/Plugins/Experimental/ChaosClothAssetEditor/Source/ChaosClothAssetDataflowNodes/Public/ChaosClothAsset/AddStitchNode.h b/Engine/Plugins/Experimental/ChaosClothAssetEditor/Source/ChaosClothAssetDataflowNodes/Public/ChaosClothAsset/AddStitchNode.h
new file mode 100644
--- /dev/null
+++ b/Engine/Plugins/Experimental/ChaosClothAssetEditor/Source/ChaosClothAssetDataflowNodes/Public/ChaosClothAsset/AddStitchNode.h
@@ -0,0 +1,28 @@
+// Copyright Epic Games, Inc. All Rights Reserved.
+
+#pragma once 
+
+#include "Dataflow/DataflowNode.h"
+#include "GeometryCollection/ManagedArrayCollection.h"
+#include "ChaosClothAsset/ConnectableValue.h"
+#include "AddStitchNode.generated.h"
+
+USTRUCT(Meta = (DataflowCloth))
+struct FChaosClothAssetAddStitchNode : public FDataflowNode
+{
+	GENERATED_USTRUCT_BODY()
+	DATAFLOW_NODE_DEFINE_INTERNAL(FChaosClothAssetAddStitchNode, "AddStitch", "Cloth", "Cloth Simulation Add Stitch")
+
+public:
+
+	UPROPERTY(Meta = (Dataflowinput, DataflowOutput, DataflowPassthrough = "Collection"))
+	FManagedArrayCollection Collection;
+
+	/** Set of vertices to stitch together. Can be 2D or 3D vertices. A seam will be created by making a chain of stitches (all vertices will merge to a single 3D vertex).*/
+	UPROPERTY(EditAnywhere, Category = "Add Stitch")
+	FChaosClothAssetConnectableIStringValue MergeToSingleVertexSelection;
+
+	FChaosClothAssetAddStitchNode(const Dataflow::FNodeParameters& InParam, FGuid InGuid = FGuid::NewGuid());
+
+	virtual void Evaluate(Dataflow::FContext& Context, const FDataflowOutput* Out) const override;
+};
diff --git a/Engine/Plugins/Experimental/ChaosClothAssetEditor/Source/ChaosClothAssetDataflowNodes/Public/ChaosClothAsset/DeleteElementNode.h b/Engine/Plugins/Experimental/ChaosClothAssetEditor/Source/ChaosClothAssetDataflowNodes/Public/ChaosClothAsset/DeleteElementNode.h
--- a/Engine/Plugins/Experimental/ChaosClothAssetEditor/Source/ChaosClothAssetDataflowNodes/Public/ChaosClothAsset/DeleteElementNode.h
+++ b/Engine/Plugins/Experimental/ChaosClothAssetEditor/Source/ChaosClothAssetDataflowNodes/Public/ChaosClothAsset/DeleteElementNode.h
@@ -4,10 +4,12 @@
 
 #include "Dataflow/DataflowNode.h"
 #include "GeometryCollection/ManagedArrayCollection.h"
+#include "ChaosClothAsset/SelectionNode.h"
+#include "ChaosClothAsset/ConnectableValue.h"
 #include "DeleteElementNode.generated.h"
 
 UENUM()
-enum class EChaosClothAssetElementType : uint8
+enum class UE_DEPRECATED(5.4, "Use FChaosClothAssetNodeSelectionGroup instead") EChaosClothAssetElementType : uint8
 {
 	None,
 	SimMesh,
@@ -19,9 +21,12 @@ enum class EChaosClothAssetElementType : uint8
 	RenderVertex,
 	SimFace,
 	RenderFace,
-	Seam
+	Seam,
+	/** Deprecated marker */
+	Deprecated UMETA(Hidden)
 };
 
+PRAGMA_DISABLE_DEPRECATION_WARNINGS  // For EChaosClothAssetElementType
 USTRUCT(Meta = (DataflowCloth))
 struct FChaosClothAssetDeleteElementNode : public FDataflowNode
 {
@@ -34,14 +39,41 @@ public:
 	FManagedArrayCollection Collection;
 
 	/** Element type to delete.*/
+	UE_DEPRECATED(5.4, "Use Group instead")
+	UPROPERTY()
+	EChaosClothAssetElementType ElementType_DEPRECATED = EChaosClothAssetElementType::Deprecated;
+
+	/** Delete the sim mesh. */
+	UPROPERTY(EditAnywhere, Category = "Delete Element")
+	bool bDeleteSimMesh = false;
+
+	/** Delete the render mesh. */
+	UPROPERTY(EditAnywhere, Category = "Delete Element")
+	bool bDeleteRenderMesh = false;
+
+	/** Delete specific elements.*/
 	UPROPERTY(EditAnywhere, Category = "Delete Element")
-	EChaosClothAssetElementType ElementType = EChaosClothAssetElementType::None;
+	FChaosClothAssetNodeSelectionGroup Group;
 
-	/** List of Elements to apply the operation on. All Elements will be used if left empty. */
-	UPROPERTY(EditAnywhere, Category = "Delete Element", Meta = (EditCondition = "ElementType != EChaosClothAssetElementType::SimMesh && ElementType != EChaosClothAssetElementType::RenderMesh"))
+	/** List of Elements to delete from Group. All Elements will be used if left empty. */
+	UPROPERTY(EditAnywhere, Category = "Delete Element")
 	TArray<int32> Elements;
 
+	/** Set of Elements to delete. This selection set will be deleted from the downstream collection since it will now be empty.*/
+	UPROPERTY(EditAnywhere, Category = "Delete Element")
+	FChaosClothAssetConnectableIStringValue SelectionName = {""};
+
 	FChaosClothAssetDeleteElementNode(const Dataflow::FNodeParameters& InParam, FGuid InGuid = FGuid::NewGuid());
 
+	/** Return a cached array of all the groups used by the input collection during at the time of the latest evaluation. */
+	const TArray<FName>& GetCachedCollectionGroupNames() const { return CachedCollectionGroupNames; }
+
+private:
 	virtual void Evaluate(Dataflow::FContext& Context, const FDataflowOutput* Out) const override;
+	virtual void OnSelected(Dataflow::FContext& Context) override;
+	virtual void OnDeselected() override;
+	virtual void Serialize(FArchive& Ar);
+
+	TArray<FName> CachedCollectionGroupNames;
 };
+PRAGMA_ENABLE_DEPRECATION_WARNINGS  // For EChaosClothAssetElementType
