Chaos Cloth: Add desired timestep parameter.

diff --git a/Engine/Plugins/Experimental/ChaosCloth/Source/Chaos/Private/ChaosCloth/ChaosClothingSimulationSolver.cpp b/Engine/Plugins/Experimental/ChaosCloth/Source/Chaos/Private/ChaosCloth/ChaosClothingSimulationSolver.cpp
--- a/Engine/Plugins/Experimental/ChaosCloth/Source/Chaos/Private/ChaosCloth/ChaosClothingSimulationSolver.cpp
+++ b/Engine/Plugins/Experimental/ChaosCloth/Source/Chaos/Private/ChaosCloth/ChaosClothingSimulationSolver.cpp
@@ -84,6 +84,7 @@ namespace ClothingSimulationSolverDefault
 	static const int32 MaxNumIterations = 10;
 	static const int32 NumSubsteps = 1;
 	static const int32 MinNumSubsteps = 1;
+	static const FRealSingle DynamicSubstepDeltaTime = 0.f;
 	static const bool bEnableNumSelfCollisionSubsteps = false;
 	static const int32 NumSelfCollisionSubsteps = 1;
 	static const FRealSingle SolverFrequency = 60.f;
@@ -1947,11 +1948,16 @@ void FClothingSimulationSolver::Update(Softs::FSolverReal InDeltaTime)
 				Properties.GetValue<int32>(TEXT("NumSubsteps"), ClothingSimulationSolverDefault::NumSubsteps),
 				ClothingSimulationSolverDefault::MinNumSubsteps);
 
+			const Softs::FSolverReal ConfigDynamicSubstepDeltaTime = Properties.GetValue<float>(TEXT("DynamicSubstepDeltaTime"), ClothingSimulationSolverDefault::DynamicSubstepDeltaTime); //  This is in ms
+
+			// Calculate NumSubsteps for dynamic substepping.
+			NumUsedSubsteps = ConfigDynamicSubstepDeltaTime > (Softs::FSolverReal)0.f ? FMath::Clamp(FMath::RoundToInt32(DeltaTime * 1000.f / ConfigDynamicSubstepDeltaTime), 1, ConfigNumSubsteps) : ConfigNumSubsteps;
+
 			const bool bConfigEnableNumSelfCollisionSubsteps = Properties.GetValue<bool>(TEXT("EnableNumSelfCollisionSubsteps"), ClothingSimulationSolverDefault::bEnableNumSelfCollisionSubsteps);
 
-			const int32 ConfigNumSelfCollisionSubsteps = bConfigEnableNumSelfCollisionSubsteps ? FMath::Clamp(Properties.GetValue<int32>(TEXT("NumSelfCollisionSubsteps"), ClothingSimulationSolverDefault::NumSelfCollisionSubsteps), ClothingSimulationSolverDefault::MinNumSubsteps, ConfigNumSubsteps) : ConfigNumSubsteps;
+			const int32 ConfigNumSelfCollisionSubsteps = bConfigEnableNumSelfCollisionSubsteps ? FMath::Clamp(Properties.GetValue<int32>(TEXT("NumSelfCollisionSubsteps"), ClothingSimulationSolverDefault::NumSelfCollisionSubsteps), ClothingSimulationSolverDefault::MinNumSubsteps, NumUsedSubsteps) : NumUsedSubsteps;
 
-			const int32 NumSubstepsPerSelfCollisionSubstep = (ConfigNumSubsteps - 1) / ConfigNumSelfCollisionSubsteps + 1;
+			const int32 NumSubstepsPerSelfCollisionSubstep = (NumUsedSubsteps - 1) / ConfigNumSelfCollisionSubsteps + 1;
 
 			if (Evolution)
 			{
@@ -1980,14 +1986,14 @@ void FClothingSimulationSolver::Update(Softs::FSolverReal InDeltaTime)
 			}
 
 			// Advance substeps
-			const Softs::FSolverReal SubstepDeltaTime = DeltaTime / (Softs::FSolverReal)ConfigNumSubsteps;
+			const Softs::FSolverReal SubstepDeltaTime = DeltaTime / (Softs::FSolverReal)NumUsedSubsteps;
 	
-			for (int32 i = 0; i < ConfigNumSubsteps; ++i)
+			for (int32 i = 0; i < NumUsedSubsteps; ++i)
 			{
-				PreSubstep(FMath::Clamp((Softs::FSolverReal)(i + 1) / (Softs::FSolverReal)ConfigNumSubsteps, (Softs::FSolverReal)0., (Softs::FSolverReal)1.), i% NumSubstepsPerSelfCollisionSubstep == 0);
+				PreSubstep(FMath::Clamp((Softs::FSolverReal)(i + 1) / (Softs::FSolverReal)NumUsedSubsteps, (Softs::FSolverReal)0., (Softs::FSolverReal)1.), i% NumSubstepsPerSelfCollisionSubstep == 0);
 				if (Evolution)
 				{
-					Evolution->AdvanceOneTimeStep(SubstepDeltaTime, (Softs::FSolverReal)ConfigNumSubsteps);
+					Evolution->AdvanceOneTimeStep(SubstepDeltaTime, (Softs::FSolverReal)NumUsedSubsteps);
 				}
 				else
 				{
diff --git a/Engine/Plugins/Experimental/ChaosCloth/Source/Chaos/Public/ChaosCloth/ChaosClothingSimulationSolver.h b/Engine/Plugins/Experimental/ChaosCloth/Source/Chaos/Public/ChaosCloth/ChaosClothingSimulationSolver.h
--- a/Engine/Plugins/Experimental/ChaosCloth/Source/Chaos/Public/ChaosCloth/ChaosClothingSimulationSolver.h
+++ b/Engine/Plugins/Experimental/ChaosCloth/Source/Chaos/Public/ChaosCloth/ChaosClothingSimulationSolver.h
@@ -108,6 +108,7 @@ namespace Chaos
 
 		// Return the actual of number of iterations used by the Evolution solver after the update (different from the number of iterations, depends on frame rate)
 		CHAOSCLOTH_API int32 GetNumUsedIterations() const;
+		int32 GetNumUsedSubsteps() const { return NumUsedSubsteps; }
 		// Return the actual number of linear solver iterations used by the Evolution solver after the update (force-based solver only)
 		CHAOSCLOTH_API int32 GetNumLinearSolverIterations(int32 ParticleRangeId) const;
 		// Return the final error of the linear solver after the update (force-based solver only)
@@ -433,6 +434,7 @@ namespace Chaos
 		// Time stepping
 		FSolverReal Time;
 		FSolverReal DeltaTime;
+		int32 NumUsedSubsteps;  // may change depending on dynamic substepping.
 
 		// FEvolution-only
 		TArray<FSolverVec3> CollisionContacts;
diff --git a/Engine/Plugins/Experimental/ChaosClothAsset/Source/ChaosClothAssetEngine/Private/ChaosClothAsset/ClothSimulationProxy.cpp b/Engine/Plugins/Experimental/ChaosClothAsset/Source/ChaosClothAssetEngine/Private/ChaosClothAsset/ClothSimulationProxy.cpp
--- a/Engine/Plugins/Experimental/ChaosClothAsset/Source/ChaosClothAssetEngine/Private/ChaosClothAsset/ClothSimulationProxy.cpp
+++ b/Engine/Plugins/Experimental/ChaosClothAsset/Source/ChaosClothAssetEngine/Private/ChaosClothAsset/ClothSimulationProxy.cpp
@@ -255,7 +255,7 @@ PRAGMA_ENABLE_DEPRECATION_WARNINGS
 
 		// Keep the actual used number of iterations for the stats
 		NumIterations = Solver->GetNumUsedIterations();
-		NumSubsteps = Solver->GetNumSubsteps();
+		NumSubsteps = Solver->GetNumUsedSubsteps();
 
 		// Update simulation time in ms (and provide an instant average instead of the value in real-time)
 		const float CurrSimulationTime = (float)((FPlatformTime::Seconds() - StartTime) * 1000.);
diff --git a/Engine/Plugins/Experimental/ChaosClothAssetEditor/Source/ChaosClothAssetDataflowNodes/Private/ChaosClothAsset/SimulationSolverConfigNode.cpp b/Engine/Plugins/Experimental/ChaosClothAssetEditor/Source/ChaosClothAssetDataflowNodes/Private/ChaosClothAsset/SimulationSolverConfigNode.cpp
--- a/Engine/Plugins/Experimental/ChaosClothAssetEditor/Source/ChaosClothAssetDataflowNodes/Private/ChaosClothAsset/SimulationSolverConfigNode.cpp
+++ b/Engine/Plugins/Experimental/ChaosClothAssetEditor/Source/ChaosClothAssetDataflowNodes/Private/ChaosClothAsset/SimulationSolverConfigNode.cpp
@@ -16,6 +16,8 @@ void FChaosClothAssetSimulationSolverConfigNode::AddProperties(FPropertyHelper&
 	PropertyHelper.SetProperty(this, &NumIterations);
 	PropertyHelper.SetProperty(this, &MaxNumIterations);
 	PropertyHelper.SetProperty(this, &NumSubsteps);
+	const float DynamicSubstepDeltaTimeValue = bEnableDynamicSubstepping ? DynamicSubstepDeltaTime : 0.f;
+	PropertyHelper.SetProperty(TEXT("DynamicSubstepDeltaTime"), DynamicSubstepDeltaTimeValue);
 	PropertyHelper.SetPropertyBool(this, &bEnableNumSelfCollisionSubsteps);
 	PropertyHelper.SetProperty(this, &NumSelfCollisionSubsteps);
 	PropertyHelper.SetPropertyBool(this, &bEnableForceBasedSolver, {}, ECollectionPropertyFlags::Intrinsic);
diff --git a/Engine/Plugins/Experimental/ChaosClothAssetEditor/Source/ChaosClothAssetDataflowNodes/Public/ChaosClothAsset/SimulationSolverConfigNode.h b/Engine/Plugins/Experimental/ChaosClothAssetEditor/Source/ChaosClothAssetDataflowNodes/Public/ChaosClothAsset/SimulationSolverConfigNode.h
--- a/Engine/Plugins/Experimental/ChaosClothAssetEditor/Source/ChaosClothAssetDataflowNodes/Public/ChaosClothAsset/SimulationSolverConfigNode.h
+++ b/Engine/Plugins/Experimental/ChaosClothAssetEditor/Source/ChaosClothAssetDataflowNodes/Public/ChaosClothAsset/SimulationSolverConfigNode.h
@@ -36,6 +36,17 @@ public:
 	UPROPERTY(EditAnywhere, Category = Simulation, meta = (UIMin = "1", UIMax = "10", ClampMin = "1", ClampMax = "100"))
 	int32 NumSubsteps = 1;
 
+	/**
+	 * Enable dynamic substepping.
+	 */
+	UPROPERTY(EditAnywhere, Category = Simulation, meta = (InlineEditConditionToggle))
+	bool bEnableDynamicSubstepping = false;
+
+	/**
+	 * Choose the number of substeps based on a target substep delta time in milliseconds. Substeps are clamped to [1, NumSubsteps]. 
+	 */
+	UPROPERTY(EditAnywhere, Category = Simulation, meta = (UIMin = "1", UIMax = "30", ClampMin = "0", ClampMax = "1000", EditCondition = "bEnableDynamicSubstepping"))
+	float DynamicSubstepDeltaTime = 16.67f;
 
 	/**
 	* Enable setting separate SelfCollisionSubsteps. Otherwise, self collisions will be detected every substep.
diff --git a/Engine/Plugins/Experimental/ChaosClothAssetEditor/Source/ChaosClothAssetEditor/Private/ChaosClothAsset/ClothEditorSimulationVisualization.cpp b/Engine/Plugins/Experimental/ChaosClothAssetEditor/Source/ChaosClothAssetEditor/Private/ChaosClothAsset/ClothEditorSimulationVisualization.cpp
--- a/Engine/Plugins/Experimental/ChaosClothAssetEditor/Source/ChaosClothAssetEditor/Private/ChaosClothAsset/ClothEditorSimulationVisualization.cpp
+++ b/Engine/Plugins/Experimental/ChaosClothAssetEditor/Source/ChaosClothAssetEditor/Private/ChaosClothAsset/ClothEditorSimulationVisualization.cpp
@@ -46,6 +46,10 @@ static FText GetSimulationStatisticsString(const FClothSimulationProxy* SimProxy
 		{
 			TextValue = ConcatenateLine(TextValue, FText::Format(LOCTEXT("NumIterations", "Iterations: {0}"), NumIterations));
 		}
+		if (const int32 NumSubSteps = SimProxy->GetNumSubsteps())
+		{
+			TextValue = ConcatenateLine(TextValue, FText::Format(LOCTEXT("NumSubsteps", "Substeps: {0}"), NumSubSteps));
+		}
 		if (const int32 NumLinearSolveIterations = SimProxy->GetNumLinearSolveIterations())
 		{
 			TextValue = ConcatenateLine(TextValue, FText::Format(LOCTEXT("NumCGIterations", "CGIterations: {0}"), NumLinearSolveIterations));
@@ -144,6 +148,7 @@ const FVisualizationOption FVisualizationOption::OptionData[] =
 	FVisualizationOption(&::Chaos::FClothVisualization::DrawPhysMeshShaded      , TEXT("p.ChaosClothAssetEditor.DebugDrawPhysMeshShaded"), LOCTEXT("ChaosVisName_PhysMesh"            , "Physical Mesh (Flat Shaded)"), LOCTEXT("ChaosVisName_PhysMeshShaded_ToolTip"      , "Draws the current physical result as a doubled sided flat shaded mesh"), /*bDisablesSimulation =*/false, /*bHidesClothSections=*/true),
 	FVisualizationOption(&::Chaos::FClothVisualization::DrawPhysMeshWired       , TEXT("p.ChaosClothAssetEditor.DebugDrawPhysMeshWired"), LOCTEXT("ChaosVisName_PhysMeshWire"        , "Physical Mesh (Wireframe)"), LOCTEXT("ChaosVisName_PhysMeshWired_ToolTip"       , "Draws the current physical mesh result in wireframe")),
 	FVisualizationOption(&::Chaos::FClothVisualization::DrawAnimMeshWired       , TEXT("p.ChaosClothAssetEditor.DebugDrawAnimMeshWired"), LOCTEXT("ChaosVisName_AnimMeshWire"        , "Animated Mesh (Wireframe)"), LOCTEXT("ChaosVisName_AnimMeshWired_ToolTip"       , "Draws the current animated mesh input in wireframe")),
+	FVisualizationOption((FLocalDebugDisplayStringFunction)nullptr                                                , TEXT("p.ChaosClothAssetEditor.DebugDrawHideRenderMesh"), LOCTEXT("ChaosVisName_HideRenderMesh"        , "Hide Render Mesh"), LOCTEXT("ChaosVisName_HideRenderMesh_ToolTip"       , "Hide the render mesh."), /*bDisablesSimulation =*/false, /*bHidesClothSections=*/true),
 	FVisualizationOption(&::Chaos::FClothVisualization::DrawParticleIndices     , TEXT("p.ChaosClothAssetEditor.DebugDrawParticleIndices"), LOCTEXT("ChaosVisName_ParticleIndices"     , "Particle Indices"), LOCTEXT("ChaosVisName_ParticleIndices_ToolTip"     , "Draws the particle indices as instantiated by the solver")),
 	FVisualizationOption(&::Chaos::FClothVisualization::DrawElementIndices      , TEXT("p.ChaosClothAssetEditor.DebugDrawElementIndices"), LOCTEXT("ChaosVisName_ElementIndices"      , "Element Indices"), LOCTEXT("ChaosVisName_ElementIndices_ToolTip"      , "Draws the element's (triangle or other) indices as instantiated by the solver")),
 	FVisualizationOption(&::Chaos::FClothVisualization::DrawPointNormals        , TEXT("p.ChaosClothAssetEditor.DebugDrawPointNormals"), LOCTEXT("ChaosVisName_PointNormals"        , "Physical Mesh Normals"), LOCTEXT("ChaosVisName_PointNormals_ToolTip"        , "Draws the current point normals for the simulation mesh")),
