Chaos Cloth: debug draw kinematic cloth colliders

diff --git a/Engine/Plugins/Experimental/ChaosCloth/Source/Chaos/Private/ChaosCloth/ChaosClothVisualization.cpp b/Engine/Plugins/Experimental/ChaosCloth/Source/Chaos/Private/ChaosCloth/ChaosClothVisualization.cpp
--- a/Engine/Plugins/Experimental/ChaosCloth/Source/Chaos/Private/ChaosCloth/ChaosClothVisualization.cpp
+++ b/Engine/Plugins/Experimental/ChaosCloth/Source/Chaos/Private/ChaosCloth/ChaosClothVisualization.cpp
@@ -2329,6 +2329,100 @@ static FAutoConsoleVariableRef CVarClothVizWeightMapName(TEXT("p.ChaosClothVisua
 			}
 		}
 	}
+
+	void FClothVisualization::DrawKinematicColliderWired(FPrimitiveDrawInterface* PDI) const
+	{
+		if (!Solver)
+		{
+			return;
+		}
+
+		static const FLinearColor WireframeColor = FColor::Silver;
+		const FVec3& LocalSpaceLocation = Solver->GetLocalSpaceLocation();
+
+		for (const FClothingSimulationCloth* const Cloth : Solver->GetCloths())
+		{
+			const int32 ParticleRangeId = Cloth->GetParticleRangeId(Solver);
+			if (ParticleRangeId == INDEX_NONE)
+			{
+				continue;
+			}
+
+			const FClothConstraints& ClothConstraints = Solver->GetClothConstraints(ParticleRangeId);
+			if (const Softs::FPBDTriangleMeshCollisions* const SelfCollisionInit = ClothConstraints.GetSelfCollisionInit().Get())
+			{
+				const FTriangleMesh& KinematicColliderMesh = SelfCollisionInit->GetCollidableSubMesh().GetKinematicColliderSubMesh();
+				// Elements are local indexed for force based solver
+				const int32 Offset = Solver->IsForceBasedSolver() ? 0 : ParticleRangeId;
+				const TConstArrayView<TVec3<int32>> Elements = KinematicColliderMesh.GetElements();
+				const TConstArrayView<Softs::FSolverVec3> Positions = Cloth->GetParticlePositions(Solver);
+				for (int32 ElementIndex = 0; ElementIndex < Elements.Num(); ++ElementIndex)
+				{
+					const TVec3<int32>& Element = Elements[ElementIndex];
+
+					const FVector Pos0 = LocalSpaceLocation + FVector(Positions[Element.X - Offset]); // TODO: Triangle Mesh shouldn't really be solver dependent (ie not use an offset)
+					const FVector Pos1 = LocalSpaceLocation + FVector(Positions[Element.Y - Offset]);
+					const FVector Pos2 = LocalSpaceLocation + FVector(Positions[Element.Z - Offset]);
+
+					DrawLine(PDI, Pos0, Pos1, WireframeColor);
+					DrawLine(PDI, Pos1, Pos2, WireframeColor);
+					DrawLine(PDI, Pos2, Pos0, WireframeColor);
+				}
+			}
+		}
+	}
+
+#if WITH_EDITOR
+	void FClothVisualization::DrawKinematicColliderShaded(FPrimitiveDrawInterface* PDI) const
+	{
+		if (!Solver || !CollisionMaterial)
+		{
+			return;
+		}
+
+		FDynamicMeshBuilder MeshBuilder(PDI->View->GetFeatureLevel());
+		int32 VertexIndex = 0;
+
+		for (const FClothingSimulationCloth* const Cloth : Solver->GetCloths())
+		{
+			const int32 ParticleRangeId = Cloth->GetParticleRangeId(Solver);
+			if (ParticleRangeId == INDEX_NONE)
+			{
+				continue;
+			}
+
+			const FClothConstraints& ClothConstraints = Solver->GetClothConstraints(ParticleRangeId);
+			if (const Softs::FPBDTriangleMeshCollisions* const SelfCollisionInit = ClothConstraints.GetSelfCollisionInit().Get())
+			{
+				const FTriangleMesh& KinematicColliderMesh = SelfCollisionInit->GetCollidableSubMesh().GetKinematicColliderSubMesh();
+				// Elements are local indexed for force based solver
+				const int32 Offset = Solver->IsForceBasedSolver() ? 0 : ParticleRangeId;
+				const TConstArrayView<TVec3<int32>> Elements = KinematicColliderMesh.GetElements();
+				const TConstArrayView<Softs::FSolverVec3> Positions = Cloth->GetParticlePositions(Solver);
+				for (int32 ElementIndex = 0; ElementIndex < Elements.Num(); ++ElementIndex, VertexIndex += 3)
+				{
+					const TVec3<int32>& Element = Elements[ElementIndex];
+					const FVector3f Pos0(Positions[Element.X - Offset]); // TODO: Triangle Mesh shouldn't really be solver dependent (ie not use an offset)
+					const FVector3f Pos1(Positions[Element.Y - Offset]);
+					const FVector3f Pos2(Positions[Element.Z - Offset]);
+
+					const FVector3f Normal = FVector3f::CrossProduct(Pos2 - Pos0, Pos1 - Pos0).GetSafeNormal();
+					const FVector3f Tangent = ((Pos1 + Pos2) * 0.5f - Pos0).GetSafeNormal();
+
+					MeshBuilder.AddVertex(FDynamicMeshVertex(Pos0, Tangent, Normal, FVector2f(0.f, 0.f), FColor::White));
+					MeshBuilder.AddVertex(FDynamicMeshVertex(Pos1, Tangent, Normal, FVector2f(0.f, 1.f), FColor::White));
+					MeshBuilder.AddVertex(FDynamicMeshVertex(Pos2, Tangent, Normal, FVector2f(1.f, 1.f), FColor::White));
+					MeshBuilder.AddTriangle(VertexIndex, VertexIndex + 1, VertexIndex + 2);
+				}
+			}
+		}
+
+		const FVec3& LocalSpaceLocation = Solver->GetLocalSpaceLocation();
+		FMatrix LocalSimSpaceToWorld(FMatrix::Identity);
+		LocalSimSpaceToWorld.SetOrigin(Solver->GetLocalSpaceLocation());
+		MeshBuilder.Draw(PDI, LocalSimSpaceToWorld, CollisionMaterial->GetRenderProxy(), SDPG_World, false, false);
+	}
+#endif // #if WITH_EDITOR
 }  // End namespace Chaos
 #else  // #if CHAOS_DEBUG_DRAW
 namespace Chaos
diff --git a/Engine/Plugins/Experimental/ChaosCloth/Source/Chaos/Public/ChaosCloth/ChaosClothVisualization.h b/Engine/Plugins/Experimental/ChaosCloth/Source/Chaos/Public/ChaosCloth/ChaosClothVisualization.h
--- a/Engine/Plugins/Experimental/ChaosCloth/Source/Chaos/Public/ChaosCloth/ChaosClothVisualization.h
+++ b/Engine/Plugins/Experimental/ChaosCloth/Source/Chaos/Public/ChaosCloth/ChaosClothVisualization.h
@@ -49,6 +49,7 @@ namespace Chaos
 		CHAOSCLOTH_API void DrawLocalSpace(FPrimitiveDrawInterface* PDI = nullptr) const;
 		CHAOSCLOTH_API void DrawSelfCollision(FPrimitiveDrawInterface* PDI = nullptr) const;
 		CHAOSCLOTH_API void DrawSelfIntersection(FPrimitiveDrawInterface* PDI = nullptr) const;
+		CHAOSCLOTH_API void DrawKinematicColliderWired(FPrimitiveDrawInterface* PDI = nullptr) const;
 		CHAOSCLOTH_API void DrawBounds(FPrimitiveDrawInterface* PDI = nullptr) const;
 		CHAOSCLOTH_API void DrawGravity(FPrimitiveDrawInterface* PDI = nullptr) const;
 		CHAOSCLOTH_API void DrawFictitiousAngularForces(FPrimitiveDrawInterface* PDI = nullptr) const;
@@ -74,6 +75,7 @@ namespace Chaos
 		void DrawLocalSpace(FPrimitiveDrawInterface* /*PDI*/ = nullptr) const {}
 		void DrawSelfCollision(FPrimitiveDrawInterface* /*PDI*/ = nullptr) const {}
 		void DrawSelfIntersection(FPrimitiveDrawInterface* /*PDI*/ = nullptr) const {}
+		CHAOSCLOTH_API void DrawKinematicColliderWired(FPrimitiveDrawInterface* PDI = nullptr) const {}
 		void DrawBounds(FPrimitiveDrawInterface* /*PDI*/ = nullptr) const {}
 		void DrawGravity(FPrimitiveDrawInterface* /*PDI*/ = nullptr) const {}
 		void DrawFictitiousAngularForces(FPrimitiveDrawInterface* /*PDI*/ = nullptr) const {}
@@ -88,6 +90,7 @@ namespace Chaos
 		CHAOSCLOTH_API void DrawMaxDistanceValues(FCanvas* Canvas, const FSceneView* SceneView) const;
 		CHAOSCLOTH_API void DrawWeightMap(FPrimitiveDrawInterface* PDI = nullptr) const;
 		CHAOSCLOTH_API void DrawInpaintWeightsMatched(FPrimitiveDrawInterface* PDI = nullptr) const;
+		CHAOSCLOTH_API void DrawKinematicColliderShaded(FPrimitiveDrawInterface* PDI) const;
 	protected:
 		// FGCObject interface
 		CHAOSCLOTH_API virtual void AddReferencedObjects(FReferenceCollector& Collector) override;
@@ -99,6 +102,7 @@ namespace Chaos
 		void DrawParticleIndices(FCanvas* /*Canvas*/, const FSceneView* /*SceneView*/) const {}
 		void DrawElementIndices(FCanvas* /*Canvas*/, const FSceneView* /*SceneView*/) const {}
 		void DrawMaxDistanceValues(FCanvas* /*Canvas*/, const FSceneView* /*SceneView*/) const {}
+		CHAOSCLOTH_API void DrawKinematicColliderShaded(FPrimitiveDrawInterface* PDI) const {}
 #endif  // #if WITH_EDITOR && CHAOS_DEBUG_DRAW
 
 #if CHAOS_DEBUG_DRAW
diff --git a/Engine/Plugins/Experimental/ChaosClothAssetEditor/Source/ChaosClothAssetEditor/Private/ChaosClothAsset/ClothEditorSimulationVisualization.cpp b/Engine/Plugins/Experimental/ChaosClothAssetEditor/Source/ChaosClothAssetEditor/Private/ChaosClothAsset/ClothEditorSimulationVisualization.cpp
--- a/Engine/Plugins/Experimental/ChaosClothAssetEditor/Source/ChaosClothAssetEditor/Private/ChaosClothAsset/ClothEditorSimulationVisualization.cpp
+++ b/Engine/Plugins/Experimental/ChaosClothAssetEditor/Source/ChaosClothAssetEditor/Private/ChaosClothAsset/ClothEditorSimulationVisualization.cpp
@@ -145,33 +145,35 @@ struct FVisualizationOption
 
 const FVisualizationOption FVisualizationOption::OptionData[] =
 {
-	FVisualizationOption(&::Chaos::FClothVisualization::DrawPhysMeshShaded      , TEXT("p.ChaosClothAssetEditor.DebugDrawPhysMeshShaded"), LOCTEXT("ChaosVisName_PhysMesh"            , "Physical Mesh (Flat Shaded)"), LOCTEXT("ChaosVisName_PhysMeshShaded_ToolTip"      , "Draws the current physical result as a doubled sided flat shaded mesh"), /*bDisablesSimulation =*/false, /*bHidesClothSections=*/true),
-	FVisualizationOption(&::Chaos::FClothVisualization::DrawPhysMeshWired       , TEXT("p.ChaosClothAssetEditor.DebugDrawPhysMeshWired"), LOCTEXT("ChaosVisName_PhysMeshWire"        , "Physical Mesh (Wireframe)"), LOCTEXT("ChaosVisName_PhysMeshWired_ToolTip"       , "Draws the current physical mesh result in wireframe")),
-	FVisualizationOption(&::Chaos::FClothVisualization::DrawAnimMeshWired       , TEXT("p.ChaosClothAssetEditor.DebugDrawAnimMeshWired"), LOCTEXT("ChaosVisName_AnimMeshWire"        , "Animated Mesh (Wireframe)"), LOCTEXT("ChaosVisName_AnimMeshWired_ToolTip"       , "Draws the current animated mesh input in wireframe")),
-	FVisualizationOption((FLocalDebugDisplayStringFunction)nullptr                                                , TEXT("p.ChaosClothAssetEditor.DebugDrawHideRenderMesh"), LOCTEXT("ChaosVisName_HideRenderMesh"        , "Hide Render Mesh"), LOCTEXT("ChaosVisName_HideRenderMesh_ToolTip"       , "Hide the render mesh."), /*bDisablesSimulation =*/false, /*bHidesClothSections=*/true),
-	FVisualizationOption(&::Chaos::FClothVisualization::DrawParticleIndices     , TEXT("p.ChaosClothAssetEditor.DebugDrawParticleIndices"), LOCTEXT("ChaosVisName_ParticleIndices"     , "Particle Indices"), LOCTEXT("ChaosVisName_ParticleIndices_ToolTip"     , "Draws the particle indices as instantiated by the solver")),
-	FVisualizationOption(&::Chaos::FClothVisualization::DrawElementIndices      , TEXT("p.ChaosClothAssetEditor.DebugDrawElementIndices"), LOCTEXT("ChaosVisName_ElementIndices"      , "Element Indices"), LOCTEXT("ChaosVisName_ElementIndices_ToolTip"      , "Draws the element's (triangle or other) indices as instantiated by the solver")),
-	FVisualizationOption(&::Chaos::FClothVisualization::DrawPointNormals        , TEXT("p.ChaosClothAssetEditor.DebugDrawPointNormals"), LOCTEXT("ChaosVisName_PointNormals"        , "Physical Mesh Normals"), LOCTEXT("ChaosVisName_PointNormals_ToolTip"        , "Draws the current point normals for the simulation mesh")),
-	FVisualizationOption(&::Chaos::FClothVisualization::DrawPointVelocities     , TEXT("p.ChaosClothAssetEditor.DebugDrawPointVelocities"), LOCTEXT("ChaosVisName_PointVelocities"     , "Point Velocities"), LOCTEXT("ChaosVisName_PointVelocities_ToolTip"     , "Draws the current point velocities for the simulation mesh")),
-	FVisualizationOption(&::Chaos::FClothVisualization::DrawAnimNormals         , TEXT("p.ChaosClothAssetEditor.DebugDrawAnimNormals"), LOCTEXT("ChaosVisName_AnimNormals"         , "Animated Mesh Normals"), LOCTEXT("ChaosVisName_AnimNormals_ToolTip"         , "Draws the current point normals for the animated mesh")),
-	FVisualizationOption(&::Chaos::FClothVisualization::DrawCollision           , TEXT("p.ChaosClothAssetEditor.DebugDrawCollision"), LOCTEXT("ChaosVisName_Collision"           , "Collisions"), LOCTEXT("ChaosVisName_Collision_ToolTip"           , "Draws the collision bodies the simulation is currently using")),
-	FVisualizationOption(&::Chaos::FClothVisualization::DrawBackstops           , TEXT("p.ChaosClothAssetEditor.DebugDrawBackstops"), LOCTEXT("ChaosVisName_Backstop"            , "Backstops"), LOCTEXT("ChaosVisName_Backstop_ToolTip"            , "Draws the backstop radius and position for each simulation particle")),
-	FVisualizationOption(&::Chaos::FClothVisualization::DrawBackstopDistances   , TEXT("p.ChaosClothAssetEditor.DebugDrawBackstopDistances"), LOCTEXT("ChaosVisName_BackstopDistance"    , "Backstop Distances"), LOCTEXT("ChaosVisName_BackstopDistance_ToolTip"    , "Draws the backstop distance offset for each simulation particle"), /*bDisablesSimulation =*/true),
-	FVisualizationOption(&::Chaos::FClothVisualization::DrawMaxDistances        , TEXT("p.ChaosClothAssetEditor.DebugDrawMaxDistances"), LOCTEXT("ChaosVisName_MaxDistance"         , "Max Distances"), LOCTEXT("ChaosVisName_MaxDistance_ToolTip"         , "Draws the current max distances for the sim particles as a line along its normal"), /*bDisablesSimulation =*/true),
-	FVisualizationOption(&::Chaos::FClothVisualization::DrawMaxDistanceValues   , TEXT("p.ChaosClothAssetEditor.DebugDrawMaxDistanceValues"), LOCTEXT("ChaosVisName_MaxDistanceValue"    , "Max Distances As Numbers"), LOCTEXT("ChaosVisName_MaxDistanceValue_ToolTip"    , "Draws the current max distances as numbers")),
-	FVisualizationOption(&::Chaos::FClothVisualization::DrawAnimDrive           , TEXT("p.ChaosClothAssetEditor.DebugDrawAnimDrive"), LOCTEXT("ChaosVisName_AnimDrive"           , "Anim Drive"), LOCTEXT("ChaosVisName_AnimDrive_Tooltip"           , "Draws the current skinned reference mesh for the simulation which anim drive will attempt to reach if enabled")),
-	FVisualizationOption(&::Chaos::FClothVisualization::DrawEdgeConstraint      , TEXT("p.ChaosClothAssetEditor.DebugDrawEdgeConstraint"), LOCTEXT("ChaosVisName_EdgeConstraint"      , "Edge Constraint"), LOCTEXT("ChaosVisName_EdgeConstraint_Tooltip"      , "Draws the edge spring constraints")),
-	FVisualizationOption(&::Chaos::FClothVisualization::DrawBendingConstraint   , TEXT("p.ChaosClothAssetEditor.DebugDrawBendingConstraint"), LOCTEXT("ChaosVisName_BendingConstraint"   , "Bending Constraint"), LOCTEXT("ChaosVisName_BendingConstraint_Tooltip"   , "Draws the bending spring constraints")),
-	FVisualizationOption(&::Chaos::FClothVisualization::DrawLongRangeConstraint , TEXT("p.ChaosClothAssetEditor.DebugDrawLongRangeConstraint"), LOCTEXT("ChaosVisName_LongRangeConstraint" , "Long Range Constraint"), LOCTEXT("ChaosVisName_LongRangeConstraint_Tooltip" , "Draws the long range attachment constraint distances")),
+	FVisualizationOption(&::Chaos::FClothVisualization::DrawPhysMeshShaded, TEXT("p.ChaosClothAssetEditor.DebugDrawPhysMeshShaded"), LOCTEXT("ChaosVisName_PhysMesh", "Physical Mesh (Flat Shaded)"), LOCTEXT("ChaosVisName_PhysMeshShaded_ToolTip", "Draws the current physical result as a doubled sided flat shaded mesh"), /*bDisablesSimulation =*/false, /*bHidesClothSections=*/true),
+	FVisualizationOption(&::Chaos::FClothVisualization::DrawPhysMeshWired, TEXT("p.ChaosClothAssetEditor.DebugDrawPhysMeshWired"), LOCTEXT("ChaosVisName_PhysMeshWire", "Physical Mesh (Wireframe)"), LOCTEXT("ChaosVisName_PhysMeshWired_ToolTip", "Draws the current physical mesh result in wireframe")),
+	FVisualizationOption(&::Chaos::FClothVisualization::DrawAnimMeshWired, TEXT("p.ChaosClothAssetEditor.DebugDrawAnimMeshWired"), LOCTEXT("ChaosVisName_AnimMeshWire", "Animated Mesh (Wireframe)"), LOCTEXT("ChaosVisName_AnimMeshWired_ToolTip", "Draws the current animated mesh input in wireframe")),
+	FVisualizationOption((FLocalDebugDisplayStringFunction)nullptr, TEXT("p.ChaosClothAssetEditor.DebugDrawHideRenderMesh"), LOCTEXT("ChaosVisName_HideRenderMesh", "Hide Render Mesh"), LOCTEXT("ChaosVisName_HideRenderMesh_ToolTip", "Hide the render mesh."), /*bDisablesSimulation =*/false, /*bHidesClothSections=*/true),
+	FVisualizationOption(&::Chaos::FClothVisualization::DrawParticleIndices, TEXT("p.ChaosClothAssetEditor.DebugDrawParticleIndices"), LOCTEXT("ChaosVisName_ParticleIndices", "Particle Indices"), LOCTEXT("ChaosVisName_ParticleIndices_ToolTip", "Draws the particle indices as instantiated by the solver")),
+	FVisualizationOption(&::Chaos::FClothVisualization::DrawElementIndices, TEXT("p.ChaosClothAssetEditor.DebugDrawElementIndices"), LOCTEXT("ChaosVisName_ElementIndices", "Element Indices"), LOCTEXT("ChaosVisName_ElementIndices_ToolTip", "Draws the element's (triangle or other) indices as instantiated by the solver")),
+	FVisualizationOption(&::Chaos::FClothVisualization::DrawPointNormals , TEXT("p.ChaosClothAssetEditor.DebugDrawPointNormals"), LOCTEXT("ChaosVisName_PointNormals", "Physical Mesh Normals"), LOCTEXT("ChaosVisName_PointNormals_ToolTip", "Draws the current point normals for the simulation mesh")),
+	FVisualizationOption(&::Chaos::FClothVisualization::DrawPointVelocities, TEXT("p.ChaosClothAssetEditor.DebugDrawPointVelocities"), LOCTEXT("ChaosVisName_PointVelocities", "Point Velocities"), LOCTEXT("ChaosVisName_PointVelocities_ToolTip", "Draws the current point velocities for the simulation mesh")),
+	FVisualizationOption(&::Chaos::FClothVisualization::DrawAnimNormals, TEXT("p.ChaosClothAssetEditor.DebugDrawAnimNormals"), LOCTEXT("ChaosVisName_AnimNormals", "Animated Mesh Normals"), LOCTEXT("ChaosVisName_AnimNormals_ToolTip", "Draws the current point normals for the animated mesh")),
+	FVisualizationOption(&::Chaos::FClothVisualization::DrawCollision, TEXT("p.ChaosClothAssetEditor.DebugDrawCollision"), LOCTEXT("ChaosVisName_Collision", "Collisions"), LOCTEXT("ChaosVisName_Collision_ToolTip", "Draws the collision bodies the simulation is currently using")),
+	FVisualizationOption(&::Chaos::FClothVisualization::DrawBackstops , TEXT("p.ChaosClothAssetEditor.DebugDrawBackstops"), LOCTEXT("ChaosVisName_Backstop", "Backstops"), LOCTEXT("ChaosVisName_Backstop_ToolTip", "Draws the backstop radius and position for each simulation particle")),
+	FVisualizationOption(&::Chaos::FClothVisualization::DrawBackstopDistances, TEXT("p.ChaosClothAssetEditor.DebugDrawBackstopDistances"), LOCTEXT("ChaosVisName_BackstopDistance", "Backstop Distances"), LOCTEXT("ChaosVisName_BackstopDistance_ToolTip", "Draws the backstop distance offset for each simulation particle"), /*bDisablesSimulation =*/true),
+	FVisualizationOption(&::Chaos::FClothVisualization::DrawMaxDistances, TEXT("p.ChaosClothAssetEditor.DebugDrawMaxDistances"), LOCTEXT("ChaosVisName_MaxDistance", "Max Distances"), LOCTEXT("ChaosVisName_MaxDistance_ToolTip", "Draws the current max distances for the sim particles as a line along its normal"), /*bDisablesSimulation =*/true),
+	FVisualizationOption(&::Chaos::FClothVisualization::DrawMaxDistanceValues, TEXT("p.ChaosClothAssetEditor.DebugDrawMaxDistanceValues"), LOCTEXT("ChaosVisName_MaxDistanceValue", "Max Distances As Numbers"), LOCTEXT("ChaosVisName_MaxDistanceValue_ToolTip", "Draws the current max distances as numbers")),
+	FVisualizationOption(&::Chaos::FClothVisualization::DrawAnimDrive, TEXT("p.ChaosClothAssetEditor.DebugDrawAnimDrive"), LOCTEXT("ChaosVisName_AnimDrive", "Anim Drive"), LOCTEXT("ChaosVisName_AnimDrive_Tooltip", "Draws the current skinned reference mesh for the simulation which anim drive will attempt to reach if enabled")),
+	FVisualizationOption(&::Chaos::FClothVisualization::DrawEdgeConstraint, TEXT("p.ChaosClothAssetEditor.DebugDrawEdgeConstraint"), LOCTEXT("ChaosVisName_EdgeConstraint", "Edge Constraint"), LOCTEXT("ChaosVisName_EdgeConstraint_Tooltip", "Draws the edge spring constraints")),
+	FVisualizationOption(&::Chaos::FClothVisualization::DrawBendingConstraint, TEXT("p.ChaosClothAssetEditor.DebugDrawBendingConstraint"), LOCTEXT("ChaosVisName_BendingConstraint", "Bending Constraint"), LOCTEXT("ChaosVisName_BendingConstraint_Tooltip", "Draws the bending spring constraints")),
+	FVisualizationOption(&::Chaos::FClothVisualization::DrawLongRangeConstraint, TEXT("p.ChaosClothAssetEditor.DebugDrawLongRangeConstraint"), LOCTEXT("ChaosVisName_LongRangeConstraint" , "Long Range Constraint"), LOCTEXT("ChaosVisName_LongRangeConstraint_Tooltip", "Draws the long range attachment constraint distances")),
 	FVisualizationOption(&::Chaos::FClothVisualization::DrawWindAndPressureForces, TEXT("p.ChaosClothAssetEditor.DebugDrawWindAndPressureForces"), LOCTEXT("ChaosVisName_WindAndPressureForces", "Wind Aerodynamic And Pressure Forces"), LOCTEXT("ChaosVisName_WindAndPressure_Tooltip", "Draws the Wind drag and lift and pressure forces")),
-	FVisualizationOption(&::Chaos::FClothVisualization::DrawLocalSpace          , TEXT("p.ChaosClothAssetEditor.DebugDrawLocalSpace"), LOCTEXT("ChaosVisName_LocalSpace"          , "Local Space Reference Bone"), LOCTEXT("ChaosVisName_LocalSpace_Tooltip"          , "Draws the local space reference bone")),
-	FVisualizationOption(&::Chaos::FClothVisualization::DrawGravity             , TEXT("p.ChaosClothAssetEditor.DebugDrawGravity"), LOCTEXT("ChaosVisName_Gravity", "Gravity"), LOCTEXT("ChaosVisName_Gravity_Tooltip", "Draws gravity")),
-	FVisualizationOption(&::Chaos::FClothVisualization::DrawFictitiousAngularForces             , TEXT("p.ChaosClothAssetEditor.DebugDrawFictitiousAngularForces"), LOCTEXT("ChaosVisName_FictitiousAngularForces", "Fictitious Angular Forces"), LOCTEXT("ChaosVisName_Gravity_FictitiousAngularForces", "Draws fictitious angular forces (force based solver only)")),
-	FVisualizationOption(&::Chaos::FClothVisualization::DrawSelfCollision       , TEXT("p.ChaosClothAssetEditor.DebugDrawSelfCollision"), LOCTEXT("ChaosVisName_SelfCollision"       , "Self Collision"), LOCTEXT("ChaosVisName_SelfCollision_Tooltip"       , "Draws the self collision thickness/debugging information")),
-	FVisualizationOption(&::Chaos::FClothVisualization::DrawSelfIntersection    , TEXT("p.ChaosClothAssetEditor.DebugDrawSelfIntersection"), LOCTEXT("ChaosVisName_SelfIntersection"    , "Self Intersection"), LOCTEXT("ChaosVisName_SelfIntersection_Tooltip"    , "Draws the self intersection contour/region information")),
-	FVisualizationOption(&GetSimulationStatisticsString    , TEXT("p.ChaosClothAssetEditor.DebugDrawSimulationStatistics"), LOCTEXT("ChaosVisName_SimulationStatistics"    , "Simulation Statistics"), LOCTEXT("ChaosVisName_SimulationStatistics_Tooltip"    , "Displays simulation statistics")),
-	FVisualizationOption(&::Chaos::FClothVisualization::DrawWeightMap			, TEXT("p.ChaosClothAssetEditor.DrawWeightMap"), LOCTEXT("ChaosVisName_DrawWeightMap"        , "Weight Map"), LOCTEXT("ChaosVisName_DrawWeightMap_ToolTip"       , "Draw the weight map for the simulation mesh. You can control the name of the map to be visualized by setting the p.ChaosClothVisualization.WeightMapName console variable."), /*bDisablesSimulation =*/false, /*bHidesClothSections=*/true),
-	FVisualizationOption(&::Chaos::FClothVisualization::DrawInpaintWeightsMatched, TEXT("p.ChaosClothAssetEditor.DrawInpaintWeightsMatched"), LOCTEXT("ChaosVisName_DrawInpaintWeightsMatched"        , "Transfer Skin Weights Node: Matched Vertices"), LOCTEXT("ChaosVisName_DrawInpaintWeightsMatched_ToolTip"       , "When transferring weights using the InpaintWeights method, will highlight the vertices for which we copied the weights directly from the source mesh. For all other vertices, the weights were computed automatically."), /*bDisablesSimulation =*/false, /*bHidesClothSections=*/true),
+	FVisualizationOption(&::Chaos::FClothVisualization::DrawLocalSpace, TEXT("p.ChaosClothAssetEditor.DebugDrawLocalSpace"), LOCTEXT("ChaosVisName_LocalSpace", "Local Space Reference Bone"), LOCTEXT("ChaosVisName_LocalSpace_Tooltip", "Draws the local space reference bone")),
+	FVisualizationOption(&::Chaos::FClothVisualization::DrawGravity, TEXT("p.ChaosClothAssetEditor.DebugDrawGravity"), LOCTEXT("ChaosVisName_Gravity", "Gravity"), LOCTEXT("ChaosVisName_Gravity_Tooltip", "Draws gravity")),
+	FVisualizationOption(&::Chaos::FClothVisualization::DrawFictitiousAngularForces, TEXT("p.ChaosClothAssetEditor.DebugDrawFictitiousAngularForces"), LOCTEXT("ChaosVisName_FictitiousAngularForces", "Fictitious Angular Forces"), LOCTEXT("ChaosVisName_Gravity_FictitiousAngularForces", "Draws fictitious angular forces (force based solver only)")),
+	FVisualizationOption(&::Chaos::FClothVisualization::DrawSelfCollision, TEXT("p.ChaosClothAssetEditor.DebugDrawSelfCollision"), LOCTEXT("ChaosVisName_SelfCollision"       , "Self Collision"), LOCTEXT("ChaosVisName_SelfCollision_Tooltip", "Draws the self collision thickness/debugging information")),
+	FVisualizationOption(&::Chaos::FClothVisualization::DrawSelfIntersection, TEXT("p.ChaosClothAssetEditor.DebugDrawSelfIntersection"), LOCTEXT("ChaosVisName_SelfIntersection"    , "Self Intersection"), LOCTEXT("ChaosVisName_SelfIntersection_Tooltip", "Draws the self intersection contour/region information")),
+	FVisualizationOption(&::Chaos::FClothVisualization::DrawKinematicColliderShaded, TEXT("p.ChaosClothAssetEditor.DebugDrawKinematicColliderShaded"), LOCTEXT("ChaosVisName_DrawKinematicColliderShaded", "Draw Kinematic Colliders (Shaded)"), LOCTEXT("ChaosVisName_DrawKinematicColliderShaded_Tooltip", "Draw kinematic cloth colliders with flat shading.")),
+	FVisualizationOption(&::Chaos::FClothVisualization::DrawKinematicColliderWired, TEXT("p.ChaosClothAssetEditor.DebugDrawKinematicColliderWired"), LOCTEXT("ChaosVisName_DrawKinematicColliderWired", "Draw Kinematic Colliders (Wireframe)"), LOCTEXT("ChaosVisName_DrawKinematicColliderShaded_Tooltip", "Draw kinematic cloth colliders in wireframe.")),
+	FVisualizationOption(&GetSimulationStatisticsString, TEXT("p.ChaosClothAssetEditor.DebugDrawSimulationStatistics"), LOCTEXT("ChaosVisName_SimulationStatistics"    , "Simulation Statistics"), LOCTEXT("ChaosVisName_SimulationStatistics_Tooltip", "Displays simulation statistics")),
+	FVisualizationOption(&::Chaos::FClothVisualization::DrawWeightMap, TEXT("p.ChaosClothAssetEditor.DrawWeightMap"), LOCTEXT("ChaosVisName_DrawWeightMap", "Weight Map"), LOCTEXT("ChaosVisName_DrawWeightMap_ToolTip", "Draw the weight map for the simulation mesh. You can control the name of the map to be visualized by setting the p.ChaosClothVisualization.WeightMapName console variable."), /*bDisablesSimulation =*/false, /*bHidesClothSections=*/true),
+	FVisualizationOption(&::Chaos::FClothVisualization::DrawInpaintWeightsMatched, TEXT("p.ChaosClothAssetEditor.DrawInpaintWeightsMatched"), LOCTEXT("ChaosVisName_DrawInpaintWeightsMatched"        , "Transfer Skin Weights Node: Matched Vertices"), LOCTEXT("ChaosVisName_DrawInpaintWeightsMatched_ToolTip", "When transferring weights using the InpaintWeights method, will highlight the vertices for which we copied the weights directly from the source mesh. For all other vertices, the weights were computed automatically."), /*bDisablesSimulation =*/false, /*bHidesClothSections=*/true),
 };
 const uint32 FVisualizationOption::Count = sizeof(OptionData) / sizeof(FVisualizationOption);
 
