Cloth Selection Node: check that selection indices are valid for the associated group before adding to the collection

diff --git a/Engine/Plugins/Experimental/ChaosClothAssetEditor/Source/ChaosClothAssetDataflowNodes/Private/ChaosClothAsset/SelectionNode.cpp b/Engine/Plugins/Experimental/ChaosClothAssetEditor/Source/ChaosClothAssetDataflowNodes/Private/ChaosClothAsset/SelectionNode.cpp
--- a/Engine/Plugins/Experimental/ChaosClothAssetEditor/Source/ChaosClothAssetDataflowNodes/Private/ChaosClothAsset/SelectionNode.cpp
+++ b/Engine/Plugins/Experimental/ChaosClothAssetEditor/Source/ChaosClothAssetDataflowNodes/Private/ChaosClothAsset/SelectionNode.cpp
@@ -39,9 +39,39 @@ void FChaosClothAssetSelectionNode::Evaluate(Dataflow::FContext& Context, const
 
 			FCollectionClothSelectionFacade SelectionFacade(SelectionCollection);
 			SelectionFacade.DefineSchema();
+			check(SelectionFacade.IsValid());
 
 			TSet<int32>& SelectionSet = SelectionFacade.FindOrAddSelectionSet(SelectionName, SelectionGroupName);
-			SelectionSet = Indices;
+
+			const int32 NumElementsInGroup = SelectionCollection->NumElements(SelectionGroupName);
+			bool bFoundAnyInvalidIndex = false;
+
+			for (const int32 Index : Indices)
+			{
+				if (Index < 0 || Index >= NumElementsInGroup)
+				{
+					const FText LogErrorMessage = FText::Format(LOCTEXT("SelectionIndexOutOfBoundsDetails", "Selection index {0} not valid for group \"{1}\" with {2} elements"),
+						Index,
+						FText::FromName(SelectionGroupName),
+						NumElementsInGroup);
+						
+					// Log all indices, but toast once
+					UE_LOG(LogChaosClothAssetDataflowNodes, Warning, TEXT("%s"), *LogErrorMessage.ToString());
+					bFoundAnyInvalidIndex = true;
+				}
+				else
+				{
+					SelectionSet.Add(Index);
+				}
+			}
+
+			if (bFoundAnyInvalidIndex)
+			{
+				// Toast once
+				const FText ToastErrorMessage = FText::Format(LOCTEXT("AnySelectionIndexOutOfBoundsDetails", "Found invalid selection indices for group \"{0}.\" See log for details"),
+					FText::FromName(SelectionGroupName));
+				FClothDataflowTools::LogAndToastWarning(*this, LOCTEXT("AnySelectionIndexOutOfBoundsHeadline", "Invalid selection"), ToastErrorMessage);
+			}
 
 			SetValue(Context, MoveTemp(*SelectionCollection), &Collection);
 		}
