LOCTEXT CIS fixes

diff --git a/Engine/Plugins/Experimental/ChaosClothAssetEditor/Source/ChaosClothAssetDataflowNodes/Private/ChaosClothAsset/DeleteElementNode.cpp b/Engine/Plugins/Experimental/ChaosClothAssetEditor/Source/ChaosClothAssetDataflowNodes/Private/ChaosClothAsset/DeleteElementNode.cpp
--- a/Engine/Plugins/Experimental/ChaosClothAssetEditor/Source/ChaosClothAssetDataflowNodes/Private/ChaosClothAsset/DeleteElementNode.cpp
+++ b/Engine/Plugins/Experimental/ChaosClothAssetEditor/Source/ChaosClothAssetDataflowNodes/Private/ChaosClothAsset/DeleteElementNode.cpp
@@ -433,8 +433,8 @@ void FChaosClothAssetDeleteElementNode::Serialize(FArchive& Ar)
 		ElementType_DEPRECATED = EChaosClothAssetElementType::Deprecated;  // This is only for clarity since the Type property won't be saved from now on
 
 		FClothDataflowTools::LogAndToastWarning(*this,
-			LOCTEXT("DeprecatedDeleteElementType", "Outdated Dataflow asset."),
-				LOCTEXT("DeprecatedDeleteElementType", "This node is out of data and contain deprecated data. The asset needs to be re-saved before it stops working at the next version update."));
+			LOCTEXT("DeprecatedDeleteElementTypeHeadline", "Outdated Dataflow asset."),
+				LOCTEXT("DeprecatedDeleteElementTypeDetails", "This node is out of data and contain deprecated data. The asset needs to be re-saved before it stops working at the next version update."));
 	}
 	PRAGMA_ENABLE_DEPRECATION_WARNINGS
 }
diff --git a/Engine/Plugins/Experimental/ChaosClothAssetEditor/Source/ChaosClothAssetEditor/Private/ChaosClothAsset/ClothEditorSimulationVisualization.cpp b/Engine/Plugins/Experimental/ChaosClothAssetEditor/Source/ChaosClothAssetEditor/Private/ChaosClothAsset/ClothEditorSimulationVisualization.cpp
--- a/Engine/Plugins/Experimental/ChaosClothAssetEditor/Source/ChaosClothAssetEditor/Private/ChaosClothAsset/ClothEditorSimulationVisualization.cpp
+++ b/Engine/Plugins/Experimental/ChaosClothAssetEditor/Source/ChaosClothAssetEditor/Private/ChaosClothAsset/ClothEditorSimulationVisualization.cpp
@@ -168,9 +168,9 @@ const FVisualizationOption FVisualizationOption::OptionData[] =
 	FVisualizationOption(&::Chaos::FClothVisualization::DrawGravity, TEXT("p.ChaosClothAssetEditor.DebugDrawGravity"), LOCTEXT("ChaosVisName_Gravity", "Gravity"), LOCTEXT("ChaosVisName_Gravity_Tooltip", "Draws gravity")),
 	FVisualizationOption(&::Chaos::FClothVisualization::DrawFictitiousAngularForces, TEXT("p.ChaosClothAssetEditor.DebugDrawFictitiousAngularForces"), LOCTEXT("ChaosVisName_FictitiousAngularForces", "Fictitious Angular Forces"), LOCTEXT("ChaosVisName_Gravity_FictitiousAngularForces", "Draws fictitious angular forces (force based solver only)")),
 	FVisualizationOption(&::Chaos::FClothVisualization::DrawSelfCollision, TEXT("p.ChaosClothAssetEditor.DebugDrawSelfCollision"), LOCTEXT("ChaosVisName_SelfCollision"       , "Self Collision"), LOCTEXT("ChaosVisName_SelfCollision_Tooltip", "Draws the self collision thickness/debugging information")),
-	FVisualizationOption(&::Chaos::FClothVisualization::DrawSelfIntersection, TEXT("p.ChaosClothAssetEditor.DebugDrawSelfIntersection"), LOCTEXT("ChaosVisName_SelfIntersection"    , "Self Intersection"), LOCTEXT("ChaosVisName_SelfIntersection_Tooltip", "Draws the self intersection contour/region information")),
+	FVisualizationOption(&::Chaos::FClothVisualization::DrawSelfIntersection, TEXT("p.ChaosClothAssetEditor.DebugDrawSelfIntersection"), LOCTEXT("ChaosVisName_SelfIntersection", "Self Intersection"), LOCTEXT("ChaosVisName_SelfIntersection_Tooltip", "Draws the self intersection contour/region information")),
 	FVisualizationOption(&::Chaos::FClothVisualization::DrawKinematicColliderShaded, TEXT("p.ChaosClothAssetEditor.DebugDrawKinematicColliderShaded"), LOCTEXT("ChaosVisName_DrawKinematicColliderShaded", "Draw Kinematic Colliders (Shaded)"), LOCTEXT("ChaosVisName_DrawKinematicColliderShaded_Tooltip", "Draw kinematic cloth colliders with flat shading.")),
-	FVisualizationOption(&::Chaos::FClothVisualization::DrawKinematicColliderWired, TEXT("p.ChaosClothAssetEditor.DebugDrawKinematicColliderWired"), LOCTEXT("ChaosVisName_DrawKinematicColliderWired", "Draw Kinematic Colliders (Wireframe)"), LOCTEXT("ChaosVisName_DrawKinematicColliderShaded_Tooltip", "Draw kinematic cloth colliders in wireframe.")),
+	FVisualizationOption(&::Chaos::FClothVisualization::DrawKinematicColliderWired, TEXT("p.ChaosClothAssetEditor.DebugDrawKinematicColliderWired"), LOCTEXT("ChaosVisName_DrawKinematicColliderWired", "Draw Kinematic Colliders (Wireframe)"), LOCTEXT("ChaosVisName_DrawKinematicColliderWired_Tooltip", "Draw kinematic cloth colliders in wireframe.")),
 	FVisualizationOption(&GetSimulationStatisticsString, TEXT("p.ChaosClothAssetEditor.DebugDrawSimulationStatistics"), LOCTEXT("ChaosVisName_SimulationStatistics"    , "Simulation Statistics"), LOCTEXT("ChaosVisName_SimulationStatistics_Tooltip", "Displays simulation statistics")),
 	FVisualizationOption(&::Chaos::FClothVisualization::DrawWeightMap, TEXT("p.ChaosClothAssetEditor.DrawWeightMap"), LOCTEXT("ChaosVisName_DrawWeightMap", "Weight Map"), LOCTEXT("ChaosVisName_DrawWeightMap_ToolTip", "Draw the weight map for the simulation mesh. You can control the name of the map to be visualized by setting the p.ChaosClothVisualization.WeightMapName console variable."), /*bDisablesSimulation =*/false, /*bHidesClothSections=*/true),
 	FVisualizationOption(&::Chaos::FClothVisualization::DrawInpaintWeightsMatched, TEXT("p.ChaosClothAssetEditor.DrawInpaintWeightsMatched"), LOCTEXT("ChaosVisName_DrawInpaintWeightsMatched"        , "Transfer Skin Weights Node: Matched Vertices"), LOCTEXT("ChaosVisName_DrawInpaintWeightsMatched_ToolTip", "When transferring weights using the InpaintWeights method, will highlight the vertices for which we copied the weights directly from the source mesh. For all other vertices, the weights were computed automatically."), /*bDisablesSimulation =*/false, /*bHidesClothSections=*/true),
