Add a "secondary selection set" to the Selection Node, Tool, and Facade. This is an optional second set of mesh elements for a given Selection.
The intended use case is for operations where the user wants to associate one selection set with another -- for example pairing up Sim Faces and Render Verts for wrap deformer binding.

diff --git a/Engine/Plugins/Experimental/ChaosClothAsset/Source/ChaosClothAsset/Private/ChaosClothAsset/CollectionClothSelectionFacade.cpp b/Engine/Plugins/Experimental/ChaosClothAsset/Source/ChaosClothAsset/Private/ChaosClothAsset/CollectionClothSelectionFacade.cpp
--- a/Engine/Plugins/Experimental/ChaosClothAsset/Source/ChaosClothAsset/Private/ChaosClothAsset/CollectionClothSelectionFacade.cpp
+++ b/Engine/Plugins/Experimental/ChaosClothAsset/Source/ChaosClothAsset/Private/ChaosClothAsset/CollectionClothSelectionFacade.cpp
@@ -11,6 +11,7 @@ namespace UE::Chaos::ClothAsset
 	namespace Private
 	{
 		static const FName SelectionGroup(TEXT("Selection"));
+		static const FName SelectionSecondaryGroup(TEXT("SelectionSecondary"));
 	}
 
 	// --------------- FCollectionClothSelectionConstFacade -------------------------------
@@ -63,6 +64,47 @@ namespace UE::Chaos::ClothAsset
 		return Selection ? Selection->GetData() : nullptr;
 	}
 
+	bool FCollectionClothSelectionConstFacade::HasSelectionSecondarySet(const FName& Name) const
+	{
+		if (!IsValid())
+		{
+			return false;
+		}
+		const bool bPrimaryAttributeExists = ManagedArrayCollection->HasAttribute(Name, Private::SelectionGroup);
+		const bool bSecondaryAttributeExists = ManagedArrayCollection->HasAttribute(Name, Private::SelectionSecondaryGroup);
+		checkf(bPrimaryAttributeExists || !bSecondaryAttributeExists, TEXT("FCollectionClothSelectionConstFacade: Secondary Selection found with no matching Primary attribute"));
+		return bPrimaryAttributeExists && bSecondaryAttributeExists;
+	}
+
+	FName FCollectionClothSelectionConstFacade::GetSelectionSecondaryGroup(const FName& Name) const
+	{
+		check(IsValid());
+		check(HasSelectionSecondarySet(Name));	// Also checks the primary set exists
+		return ManagedArrayCollection->GetDependency(Name, Private::SelectionSecondaryGroup);
+	}
+
+	const TSet<int32>& FCollectionClothSelectionConstFacade::GetSelectionSecondarySet(const FName& Name) const
+	{
+		check(IsValid());
+		check(ManagedArrayCollection->HasAttribute(Name, Private::SelectionGroup));
+
+		const TManagedArray<TSet<int32>>* const Selection = ManagedArrayCollection->FindAttributeTyped<TSet<int32>>(Name, Private::SelectionSecondaryGroup);
+		check(Selection);
+		return *Selection->GetData();
+	}
+
+	const TSet<int32>* FCollectionClothSelectionConstFacade::FindSelectionSecondarySet(const FName& Name) const
+	{
+		check(IsValid());
+		const bool bPrimaryAttributeExists = ManagedArrayCollection->HasAttribute(Name, Private::SelectionGroup);
+		const TManagedArray<TSet<int32>>* const SecondarySelection = ManagedArrayCollection->FindAttributeTyped<TSet<int32>>(Name, Private::SelectionSecondaryGroup);
+
+		checkf(bPrimaryAttributeExists || !SecondarySelection, TEXT("FCollectionClothSelectionConstFacade: Secondary Selection found with no matching Primary attribute"));
+
+		return SecondarySelection ? SecondarySelection->GetData() : nullptr;
+	}
+
+
 	// --------------- FCollectionClothSelectionFacade -------------------------------
 
 	FCollectionClothSelectionFacade::FCollectionClothSelectionFacade(const TSharedRef<const FManagedArrayCollection>& ManagedArrayCollection) :
@@ -71,19 +113,25 @@ namespace UE::Chaos::ClothAsset
 
 	void FCollectionClothSelectionFacade::DefineSchema()
 	{
-		if (!ManagedArrayCollection->HasGroup(Private::SelectionGroup))
-		{
-			ManagedArrayCollection->AddGroup(Private::SelectionGroup);
-		}
-		const int32 NumElements = ManagedArrayCollection->NumElements(Private::SelectionGroup);
-		if (NumElements > 1)
-		{
-			ManagedArrayCollection->RemoveElements(Private::SelectionGroup, NumElements - 1, 1);
-		}
-		else if (NumElements == 0)
+		auto InitSelectionGroup = [this](const FName& SelectionGroupName)
 		{
-			ManagedArrayCollection->AddElements(1, Private::SelectionGroup);
-		}
+			if (!ManagedArrayCollection->HasGroup(SelectionGroupName))
+			{
+				ManagedArrayCollection->AddGroup(SelectionGroupName);
+			}
+			const int32 NumElements = ManagedArrayCollection->NumElements(SelectionGroupName);
+			if (NumElements > 1)
+			{
+				ManagedArrayCollection->RemoveElements(SelectionGroupName, NumElements - 1, 1);
+			}
+			else if (NumElements == 0)
+			{
+				ManagedArrayCollection->AddElements(1, SelectionGroupName);
+			}
+		};
+
+		InitSelectionGroup(Private::SelectionGroup);
+		InitSelectionGroup(Private::SelectionSecondaryGroup);
 	}
 
 	void FCollectionClothSelectionFacade::AppendWithOffsets(const FCollectionClothSelectionConstFacade& Other, bool bOverwriteExistingIfMismatched, const TMap<FName, int32>& GroupOffsets)
@@ -159,21 +207,51 @@ namespace UE::Chaos::ClothAsset
 		ManagedArrayCollection->RemoveAttribute(Name, Private::SelectionGroup);
 	}
 
-	TSet<int32>& FCollectionClothSelectionFacade::FindOrAddSelectionSet(const FName& Name, const FName& GroupName)
+
+	TSet<int32>& FCollectionClothSelectionFacade::GetSelectionSecondarySet(const FName& Name)
+	{
+		check(IsValid());
+		check(ManagedArrayCollection->FindAttributeTyped<TSet<int32>>(Name, Private::SelectionGroup) != nullptr);
+
+		TManagedArray<TSet<int32>>* const Selection = ManagedArrayCollection->FindAttributeTyped<TSet<int32>>(Name, Private::SelectionSecondaryGroup);
+		check(Selection);
+		return *Selection->GetData();
+	}
+
+	TSet<int32>* FCollectionClothSelectionFacade::FindSelectionSecondarySet(const FName& Name)
+	{
+		check(IsValid());
+		const bool bPrimaryAttributeExists = ManagedArrayCollection->HasAttribute(Name, Private::SelectionGroup);
+
+		TManagedArray<TSet<int32>>* const SecondarySelection = ManagedArrayCollection->FindAttributeTyped<TSet<int32>>(Name, Private::SelectionSecondaryGroup);	
+		
+		checkf(bPrimaryAttributeExists || !SecondarySelection, TEXT("FCollectionClothSelectionFacade: Secondary Selection found with no Primary attribute"));
+
+		return SecondarySelection ? SecondarySelection->GetData() : nullptr;
+	}
+
+	void FCollectionClothSelectionFacade::RemoveSelectionSecondarySet(const FName& Name)
+	{
+		check(IsValid());
+		ManagedArrayCollection->RemoveAttribute(Name, Private::SelectionSecondaryGroup);
+	}
+
+
+	TSet<int32>& FCollectionClothSelectionFacade::FindOrAddSelectionSetInternal(const FName& Name, const FName& GroupName, const FName& SelectionGroupName)
 	{
 		check(IsValid());
 		ensure(GroupName != NAME_None && Name != NAME_None);
 		constexpr bool bAllowCircularDependency = false;
 
-		TManagedArray<TSet<int32>>* Selection = ManagedArrayCollection->FindAttributeTyped<TSet<int32>>(Name, Private::SelectionGroup);
+		TManagedArray<TSet<int32>>* Selection = ManagedArrayCollection->FindAttributeTyped<TSet<int32>>(Name, SelectionGroupName);
 		if (Selection)
 		{
 			check(Selection->Num() == 1);  // This should always be the case if the facade is valid
 
 			// Recycle the existing selection, with the new group if needed
-			if (ManagedArrayCollection->GetDependency(Name, Private::SelectionGroup) != GroupName)
+			if (ManagedArrayCollection->GetDependency(Name, SelectionGroupName) != GroupName)
 			{
-				ManagedArrayCollection->SetDependency(Name, Private::SelectionGroup, GroupName, bAllowCircularDependency);
+				ManagedArrayCollection->SetDependency(Name, SelectionGroupName, GroupName, bAllowCircularDependency);
 				(*Selection)[0].Reset();  // No point in keeping unrelated selection indices since the group has changed better clear everything
 			}
 		}
@@ -182,13 +260,26 @@ namespace UE::Chaos::ClothAsset
 			// Create a new selection
 			constexpr bool bSaved = true;
 			const FManagedArrayCollection::FConstructionParameters GroupDependency(GroupName, bSaved, bAllowCircularDependency);
-			Selection = &ManagedArrayCollection->AddAttribute<TSet<int32>>(Name, Private::SelectionGroup, GroupDependency);
+			Selection = &ManagedArrayCollection->AddAttribute<TSet<int32>>(Name, SelectionGroupName, GroupDependency);
 
 			check(Selection->Num() == 1);  // This should always be the case if the facade is valid
 		}
 
 		return (*Selection)[0];
 	}
+
+
+	TSet<int32>& FCollectionClothSelectionFacade::FindOrAddSelectionSet(const FName& Name, const FName& GroupName)
+	{
+		return FindOrAddSelectionSetInternal(Name, GroupName, Private::SelectionGroup);
+	}
+
+	TSet<int32>& FCollectionClothSelectionFacade::FindOrAddSelectionSecondarySet(const FName& Name, const FName& GroupName)
+	{
+		checkf(HasSelection(Name), TEXT("FCollectionClothSelectionFacade: Can't add a Secondary Selection with no matching Primary attribute"));
+		return FindOrAddSelectionSetInternal(Name, GroupName, Private::SelectionSecondaryGroup);
+	}
+
 }	// namespace  UE::Chaos::ClothAsset
 
 #undef LOCTEXT_NAMESPACE
diff --git a/Engine/Plugins/Experimental/ChaosClothAsset/Source/ChaosClothAsset/Public/ChaosClothAsset/CollectionClothSelectionFacade.h b/Engine/Plugins/Experimental/ChaosClothAsset/Source/ChaosClothAsset/Public/ChaosClothAsset/CollectionClothSelectionFacade.h
--- a/Engine/Plugins/Experimental/ChaosClothAsset/Source/ChaosClothAsset/Public/ChaosClothAsset/CollectionClothSelectionFacade.h
+++ b/Engine/Plugins/Experimental/ChaosClothAsset/Source/ChaosClothAsset/Public/ChaosClothAsset/CollectionClothSelectionFacade.h
@@ -34,7 +34,7 @@ namespace UE::Chaos::ClothAsset
 		/** Return an array of all the selections' name in the collection. */
 		CHAOSCLOTHASSET_API TArray<FName> GetNames() const;
 
-		/** Return whether a selection with the given name cyrrently exists in the collection. */
+		/** Return whether a selection with the given name currently exists in the collection. */
 		CHAOSCLOTHASSET_API bool HasSelection(const FName& Name) const;
 
 		/** Return the selection group dependency. The selection must exist to call this function. */
@@ -46,6 +46,19 @@ namespace UE::Chaos::ClothAsset
 		/** Find a selection with the given name, or nullptr if no such selection exists. */
 		CHAOSCLOTHASSET_API const TSet<int32>* FindSelectionSet(const FName& Name) const;
 
+
+		/** Return whether a selection with the given name currently exists in the collection and has a secondary set associated with it. */
+		CHAOSCLOTHASSET_API bool HasSelectionSecondarySet(const FName& Name) const;
+
+		/** Return the secondary selection set group dependency. The selection must exist and have a secondary set to call this function. */
+		CHAOSCLOTHASSET_API FName GetSelectionSecondaryGroup(const FName& Name) const;
+
+		/** Get the secondary selection set for the given selection name. The selection must exist and have a secondary set to call this function. */
+		CHAOSCLOTHASSET_API const TSet<int32>& GetSelectionSecondarySet(const FName& Name) const;
+
+		/** Find the secondary selection set for the given selection name. Returns nullptr if no such selection exists or if it does not have an associated secondary set. */
+		CHAOSCLOTHASSET_API const TSet<int32>* FindSelectionSecondarySet(const FName& Name) const;
+
 	protected:
 		TSharedRef<FManagedArrayCollection> ManagedArrayCollection;
 	};
@@ -78,6 +91,15 @@ namespace UE::Chaos::ClothAsset
 		/** Remove the selection with the given name if it exists. */
 		CHAOSCLOTHASSET_API void RemoveSelectionSet(const FName& Name);
 
+		/** Get the secondary selection set for the given selection name. The selection must exist and have a secondary set to call this function. */
+		CHAOSCLOTHASSET_API TSet<int32>& GetSelectionSecondarySet(const FName& Name);
+
+		/** Find the secondary selection set for the given selection name. Returns nullptr if no such selection exists or if it does not have an associated secondary set. */
+		CHAOSCLOTHASSET_API TSet<int32>* FindSelectionSecondarySet(const FName& Name);
+
+		/** Remove the secondary selection set from the selection with the given name if it exists. */
+		CHAOSCLOTHASSET_API void RemoveSelectionSecondarySet(const FName& Name);
+
 		/** 
 		 * Append all sets from an existing collection to this collection.
 		 * Matching sets (i.e., same name and type) will be unioned.
@@ -107,5 +129,19 @@ namespace UE::Chaos::ClothAsset
 		 * If the selection already exists, but depends on a different group, the old selection will be deleted and a new one recreated with the new group dependency.
 		 */
 		CHAOSCLOTHASSET_API TSet<int32>& FindOrAddSelectionSet(const FName& Name, const FName& GroupName);
+
+		/**
+		 * Find, or add if it doesn't already exist, a secondary selection set for the specified group with the given name.
+		 * NOTE: A selection with the given name must already exist. This function just adds a secondary set to the existing selection.
+		 * If the group doesn't already exists, this function will create it.
+		 * If the secondary selection set already exists, but depends on a different group, the old secondary selection set will be deleted and a new one recreated with the new group dependency.
+		 */
+		CHAOSCLOTHASSET_API TSet<int32>& FindOrAddSelectionSecondarySet(const FName& Name, const FName& GroupName);
+
+	private:
+
+		// Helper for FindOrAddSelectionSet and FindOrAddSelectionSecondarySet
+		TSet<int32>& FindOrAddSelectionSetInternal(const FName& Name, const FName& GroupName, const FName& SelectionGroupName);
+
 	};
 }
diff --git a/Engine/Plugins/Experimental/ChaosClothAssetEditor/Source/ChaosClothAssetDataflowNodes/Private/ChaosClothAsset/SelectionNode.cpp b/Engine/Plugins/Experimental/ChaosClothAssetEditor/Source/ChaosClothAssetDataflowNodes/Private/ChaosClothAsset/SelectionNode.cpp
--- a/Engine/Plugins/Experimental/ChaosClothAssetEditor/Source/ChaosClothAssetDataflowNodes/Private/ChaosClothAsset/SelectionNode.cpp
+++ b/Engine/Plugins/Experimental/ChaosClothAssetEditor/Source/ChaosClothAssetDataflowNodes/Private/ChaosClothAsset/SelectionNode.cpp
@@ -22,59 +22,75 @@ void FChaosClothAssetSelectionNode::Evaluate(Dataflow::FContext& Context, const
 {
 	using namespace UE::Chaos::ClothAsset;
 
+	auto CopyIntoSelection = [this](const TSharedRef<FManagedArrayCollection> SelectionCollection, 
+		const FName& SelectionGroupName, 
+		const TSet<int32>& SourceIndices,
+		TSet<int32>& DestSelectionSet)
+	{
+		const int32 NumElementsInGroup = SelectionCollection->NumElements(SelectionGroupName);
+		bool bFoundAnyInvalidIndex = false;
+
+		for (const int32 Index : SourceIndices)
+		{
+			if (Index < 0 || Index >= NumElementsInGroup)
+			{
+				const FText LogErrorMessage = FText::Format(LOCTEXT("SelectionIndexOutOfBoundsDetails", "Selection index {0} not valid for group \"{1}\" with {2} elements"),
+					Index,
+					FText::FromName(SelectionGroupName),
+					NumElementsInGroup);
+
+				// Log all indices, but toast once
+				UE_LOG(LogChaosClothAssetDataflowNodes, Warning, TEXT("%s"), *LogErrorMessage.ToString());
+				bFoundAnyInvalidIndex = true;
+			}
+			else
+			{
+				DestSelectionSet.Add(Index);
+			}
+		}
+
+		if (bFoundAnyInvalidIndex)
+		{
+			// Toast once
+			const FText ToastErrorMessage = FText::Format(LOCTEXT("AnySelectionIndexOutOfBoundsDetails", "Found invalid selection indices for group \"{0}.\" See log for details"),
+				FText::FromName(SelectionGroupName));
+			FClothDataflowTools::LogAndToastWarning(*this, LOCTEXT("AnySelectionIndexOutOfBoundsHeadline", "Invalid selection"), ToastErrorMessage);
+		}
+	};
+
+
 	if (Out->IsA<FManagedArrayCollection>(&Collection))
 	{
 		if (Name.IsEmpty() || Group.Name.IsEmpty())
 		{
 			const FManagedArrayCollection& SelectionCollection = GetValue<FManagedArrayCollection>(Context, &Collection);
 			SetValue(Context, SelectionCollection, &Collection);
+			return;
 		}
-		else
-		{
-			const FName SelectionName(*Name);
-			const FName SelectionGroupName(*Group.Name);
 
-			FManagedArrayCollection InSelectionCollection = GetValue<FManagedArrayCollection>(Context, &Collection);
-			const TSharedRef<FManagedArrayCollection> SelectionCollection = MakeShared<FManagedArrayCollection>(MoveTemp(InSelectionCollection));
+		const FName SelectionName(*Name);
+		const FName SelectionGroupName(*Group.Name);
 
-			FCollectionClothSelectionFacade SelectionFacade(SelectionCollection);
-			SelectionFacade.DefineSchema();
-			check(SelectionFacade.IsValid());
+		FManagedArrayCollection InSelectionCollection = GetValue<FManagedArrayCollection>(Context, &Collection);
+		const TSharedRef<FManagedArrayCollection> SelectionCollection = MakeShared<FManagedArrayCollection>(MoveTemp(InSelectionCollection));
 
-			TSet<int32>& SelectionSet = SelectionFacade.FindOrAddSelectionSet(SelectionName, SelectionGroupName);
+		FCollectionClothSelectionFacade SelectionFacade(SelectionCollection);
+		SelectionFacade.DefineSchema();
+		check(SelectionFacade.IsValid());
 
-			const int32 NumElementsInGroup = SelectionCollection->NumElements(SelectionGroupName);
-			bool bFoundAnyInvalidIndex = false;
-
-			for (const int32 Index : Indices)
-			{
-				if (Index < 0 || Index >= NumElementsInGroup)
-				{
-					const FText LogErrorMessage = FText::Format(LOCTEXT("SelectionIndexOutOfBoundsDetails", "Selection index {0} not valid for group \"{1}\" with {2} elements"),
-						Index,
-						FText::FromName(SelectionGroupName),
-						NumElementsInGroup);
-						
-					// Log all indices, but toast once
-					UE_LOG(LogChaosClothAssetDataflowNodes, Warning, TEXT("%s"), *LogErrorMessage.ToString());
-					bFoundAnyInvalidIndex = true;
-				}
-				else
-				{
-					SelectionSet.Add(Index);
-				}
-			}
+		TSet<int32>& SelectionSet = SelectionFacade.FindOrAddSelectionSet(SelectionName, SelectionGroupName);
+		CopyIntoSelection(SelectionCollection, SelectionGroupName, Indices, SelectionSet);
 
-			if (bFoundAnyInvalidIndex)
-			{
-				// Toast once
-				const FText ToastErrorMessage = FText::Format(LOCTEXT("AnySelectionIndexOutOfBoundsDetails", "Found invalid selection indices for group \"{0}.\" See log for details"),
-					FText::FromName(SelectionGroupName));
-				FClothDataflowTools::LogAndToastWarning(*this, LOCTEXT("AnySelectionIndexOutOfBoundsHeadline", "Invalid selection"), ToastErrorMessage);
-			}
+		if (!SecondaryGroup.Name.IsEmpty() && !SecondaryIndices.IsEmpty())
+		{
+			const FName SecondarySelectionName(*Name);
+			const FName SecondarySelectionGroupName = (*SecondaryGroup.Name);		
+			TSet<int32>& SecondarySelectionSet = SelectionFacade.FindOrAddSelectionSecondarySet(SecondarySelectionName, SecondarySelectionGroupName);
 
-			SetValue(Context, MoveTemp(*SelectionCollection), &Collection);
+			CopyIntoSelection(SelectionCollection, SecondarySelectionGroupName, SecondaryIndices, SecondarySelectionSet);
 		}
+
+		SetValue(Context, MoveTemp(*SelectionCollection), &Collection);
 	}
 	else if (Out->IsA<FString>(&Name))
 	{
@@ -118,7 +134,7 @@ PRAGMA_DISABLE_DEPRECATION_WARNINGS
 
 		FClothDataflowTools::LogAndToastWarning(*this,
 			LOCTEXT("DeprecatedSelectionType", "Outdated Dataflow asset."),
-			LOCTEXT("DeprecatedSelectionDetails", "This node is out of data and contain deprecated data. The asset needs to be re-saved before it stops working at the next version update."));
+			LOCTEXT("DeprecatedSelectionDetails", "This node is out of date and contains deprecated data. The asset needs to be re-saved before it stops working at the next version update."));
 	}
 PRAGMA_ENABLE_DEPRECATION_WARNINGS
 }
diff --git a/Engine/Plugins/Experimental/ChaosClothAssetEditor/Source/ChaosClothAssetDataflowNodes/Public/ChaosClothAsset/SelectionNode.h b/Engine/Plugins/Experimental/ChaosClothAssetEditor/Source/ChaosClothAssetDataflowNodes/Public/ChaosClothAsset/SelectionNode.h
--- a/Engine/Plugins/Experimental/ChaosClothAssetEditor/Source/ChaosClothAssetDataflowNodes/Public/ChaosClothAsset/SelectionNode.h
+++ b/Engine/Plugins/Experimental/ChaosClothAssetEditor/Source/ChaosClothAssetDataflowNodes/Public/ChaosClothAsset/SelectionNode.h
@@ -73,6 +73,14 @@ public:
 	UPROPERTY(EditAnywhere, Category = "Selection", Meta = (ClampMin = "0"))
 	TSet<int32> Indices;
 
+	/** The type of element the secondary selection refers to */
+	UPROPERTY(EditAnywhere, Category = "Secondary Selection")
+	FChaosClothAssetNodeSelectionGroup SecondaryGroup;
+
+	/** Secondary set of element indices */
+	UPROPERTY(EditAnywhere, Category = "Secondary Selection", Meta = (ClampMin = "0"))
+	TSet<int32> SecondaryIndices;
+
 	FChaosClothAssetSelectionNode(const Dataflow::FNodeParameters& InParam, FGuid InGuid = FGuid::NewGuid());
 
 	/** Return a cached array of all the groups used by the input collection during at the time of the latest evaluation. */
diff --git a/Engine/Plugins/Experimental/ChaosClothAssetEditor/Source/ChaosClothAssetEditorTools/Private/ChaosClothAsset/ClothMeshSelectionTool.cpp b/Engine/Plugins/Experimental/ChaosClothAssetEditor/Source/ChaosClothAssetEditorTools/Private/ChaosClothAsset/ClothMeshSelectionTool.cpp
--- a/Engine/Plugins/Experimental/ChaosClothAssetEditor/Source/ChaosClothAssetEditorTools/Private/ChaosClothAsset/ClothMeshSelectionTool.cpp
+++ b/Engine/Plugins/Experimental/ChaosClothAssetEditor/Source/ChaosClothAssetEditorTools/Private/ChaosClothAsset/ClothMeshSelectionTool.cpp
@@ -188,27 +188,6 @@ void UClothMeshSelectionTool::Setup()
 	// Hide input target mesh
 	UE::ToolTarget::HideSourceObject(Target);
 
-	// Initialize the Selection from the selected Dataflow node
-	FString ExistingSelectionName;
-	FGroupTopologySelection ExistingNodeSelection;
-	GetSelectedNodeInfo(ExistingSelectionName, ExistingNodeSelection);
-
-	constexpr bool bBroadcastChange = false;
-	SelectionMechanic->SetSelection(ExistingNodeSelection, bBroadcastChange);
-
-	if (ExistingNodeSelection.SelectedCornerIDs.Num() > 0)
-	{
-		check(ExistingNodeSelection.SelectedEdgeIDs.Num() == 0);
-		check(ExistingNodeSelection.SelectedGroupIDs.Num() == 0);
-		SelectionMechanic->Properties->bSelectVertices = true;
-		SelectionMechanic->Properties->bSelectFaces = false;
-	} 
-	else
-	{
-		SelectionMechanic->Properties->bSelectVertices = false;
-		SelectionMechanic->Properties->bSelectFaces = true;
-	}
-
 	// Setup non-manifold mapping if necessary
 	if (ClothEditorContextObject)
 	{
@@ -241,13 +220,6 @@ void UClothMeshSelectionTool::Setup()
 		}
 	}
 
-	// 
-	// Actions
-	//
-
-	ActionsProps = NewObject<UClothMeshSelectionToolActions>();
-	ActionsProps->Initialize(this);
-	AddToolPropertySource(ActionsProps);
 
 	//
 	// Properties
@@ -264,8 +236,35 @@ void UClothMeshSelectionTool::Setup()
 		SelectionMechanic->SetShowEdges(bNewShowEdges);
 	});
 
-	ToolProperties->Name = ExistingSelectionName;
+	// Order of operations is important here:
+	// ToolProperties->RestoreProperties should happen before GetSelectedNodeInfo so that we can know whether we need to get Primary or Secondary set
+	// ToolProperties->WatchProperty(ToolProperties->Name) should happen after GetSelectedNodeInfo so that we can set the OriginalName
+
+	ToolProperties->RestoreProperties(this);
+	 
+	// Initialize the Selection from the selected Dataflow node
+	FString ExistingSelectionName;
+	FGroupTopologySelection ExistingNodeSelection;
+	GetSelectedNodeInfo(ExistingSelectionName, ExistingNodeSelection);
+
+	constexpr bool bBroadcastChange = false;
+	SelectionMechanic->SetSelection(ExistingNodeSelection, bBroadcastChange);
+
+	if (ExistingNodeSelection.SelectedCornerIDs.Num() > 0)
+	{
+		check(ExistingNodeSelection.SelectedEdgeIDs.Num() == 0);
+		check(ExistingNodeSelection.SelectedGroupIDs.Num() == 0);
+		SelectionMechanic->Properties->bSelectVertices = true;
+		SelectionMechanic->Properties->bSelectFaces = false;
+	}
+	else
+	{
+		SelectionMechanic->Properties->bSelectVertices = false;
+		SelectionMechanic->Properties->bSelectFaces = true;
+	}
 
+
+    ToolProperties->Name = ExistingSelectionName;
 	ToolProperties->WatchProperty(ToolProperties->Name, [this, OriginalName = ToolProperties->Name](const FString& NewName)
 	{
 		if (NewName != OriginalName)
@@ -274,15 +273,25 @@ void UClothMeshSelectionTool::Setup()
 		}
 	});
 
+
+	// 
+	// Actions
+	//
+
+	ActionsProps = NewObject<UClothMeshSelectionToolActions>();
+	ActionsProps->Initialize(this);
+	AddToolPropertySource(ActionsProps);
+
 	AddToolPropertySource(ToolProperties);
-	ToolProperties->RestoreProperties(this);
 
 	AddToolPropertySource(SelectionMechanic->Properties);
+
+	UpdatePrimarySecondaryMessage();
 }
 
 void UClothMeshSelectionTool::OnShutdown(EToolShutdownType ShutdownType)
 {
-	if (ShutdownType == EToolShutdownType::Accept)
+	if (ShutdownType == EToolShutdownType::Accept && CanAccept())
 	{
 		GetToolManager()->BeginUndoTransaction(LOCTEXT("SelectionToolTransactionName", "Mesh Selection"));
 		UpdateSelectedNode();
@@ -373,27 +382,41 @@ bool UClothMeshSelectionTool::GetSelectedNodeInfo(FString& OutSelectionName, UE:
 		});
 	};
 
-
-	if (MeshSelectionNode->Group.Name == ClothCollectionGroup::SimVertices2D.ToString() ||
-		MeshSelectionNode->Group.Name == ClothCollectionGroup::SimVertices3D.ToString() ||
-		MeshSelectionNode->Group.Name == ClothCollectionGroup::RenderVertices.ToString())
+	auto ReadFromNode = [this, &AppendVerticesIfValid, &AppendFacesIfValid, &OutSelection](const FChaosClothAssetNodeSelectionGroup& SourceGroup, const TSet<int32>& SourceIndices)
 	{
-		if (bHasNonManifoldMapping)
+		if (SourceGroup.Name == ClothCollectionGroup::SimVertices2D.ToString() ||
+			SourceGroup.Name == ClothCollectionGroup::SimVertices3D.ToString() ||
+			SourceGroup.Name == ClothCollectionGroup::RenderVertices.ToString())
 		{
-			for (const int32 SelectionIndex : MeshSelectionNode->Indices)
+			if (bHasNonManifoldMapping)
 			{
-				AppendVerticesIfValid(OutSelection.SelectedCornerIDs, SelectionToDynamicMesh[SelectionIndex]);
+				for (const int32 SelectionIndex : SourceIndices)
+				{
+					if (SelectionIndex < SelectionToDynamicMesh.Num())		// Could be loading a render mesh selection where NumRenderVertices > NumSimVertices
+					{
+						AppendVerticesIfValid(OutSelection.SelectedCornerIDs, SelectionToDynamicMesh[SelectionIndex]);
+					}
+				}
+			}
+			else
+			{
+				AppendVerticesIfValid(OutSelection.SelectedCornerIDs, SourceIndices);
 			}
 		}
-		else
+		else if (SourceGroup.Name == ClothCollectionGroup::SimFaces.ToString() ||
+			     SourceGroup.Name == ClothCollectionGroup::RenderFaces)
 		{
-			AppendVerticesIfValid(OutSelection.SelectedCornerIDs, MeshSelectionNode->Indices);
+			AppendFacesIfValid(OutSelection.SelectedGroupIDs, SourceIndices);
 		}
+	};
+
+	if (ToolProperties && ToolProperties->bSecondarySelection)
+	{
+		ReadFromNode(MeshSelectionNode->SecondaryGroup, MeshSelectionNode->SecondaryIndices);
 	}
-	else if (MeshSelectionNode->Group.Name == ClothCollectionGroup::SimFaces.ToString() ||
-		MeshSelectionNode->Group.Name == ClothCollectionGroup::RenderFaces)
+	else
 	{
-		AppendFacesIfValid(OutSelection.SelectedGroupIDs, MeshSelectionNode->Indices);
+		ReadFromNode(MeshSelectionNode->Group, MeshSelectionNode->Indices);
 	}
 
 	OutSelectionName = MeshSelectionNode->Name;
@@ -451,24 +474,35 @@ void UClothMeshSelectionTool::UpdateSelectedNode()
 	FChaosClothAssetSelectionNode* const MeshSelectionNode = ClothEditorContextObject->GetSingleSelectedNodeOfType<FChaosClothAssetSelectionNode>();
 	check(MeshSelectionNode);
 
-	MeshSelectionNode->Name = ToolProperties->Name;
-	MeshSelectionNode->Group.Name = GroupName.ToString();
-
-	if (SelectionMechanic->Properties->bSelectVertices && bHasNonManifoldMapping)
+	auto WriteToNode = [this, &Indices, &GroupName, MeshSelectionNode](FChaosClothAssetNodeSelectionGroup& TargetGroup, TSet<int32>& TargetIndices)
 	{
-		MeshSelectionNode->Indices.Reset();
-		for (const int32 DynamicMeshIdx : Indices)
+		MeshSelectionNode->Name = ToolProperties->Name;
+		TargetGroup.Name = GroupName.ToString();
+
+		if (SelectionMechanic->Properties->bSelectVertices && bHasNonManifoldMapping)
+		{
+			TargetIndices.Reset();
+			for (const int32 DynamicMeshIdx : Indices)
+			{
+				const int32 MappedSelectionIndex = DynamicMeshToSelection[DynamicMeshIdx];
+				TargetIndices.Add(MappedSelectionIndex);
+			}
+		}
+		else
 		{
-			const int32 MappedSelectionIndex = DynamicMeshToSelection[DynamicMeshIdx];
-			MeshSelectionNode->Indices.Add(MappedSelectionIndex);
+			TargetIndices = MoveTemp(Indices);
 		}
+	};
+
+	if (!ToolProperties->bSecondarySelection)
+	{
+		WriteToNode(MeshSelectionNode->Group, MeshSelectionNode->Indices);
 	}
 	else
 	{
-		MeshSelectionNode->Indices = MoveTemp(Indices);
+		WriteToNode(MeshSelectionNode->SecondaryGroup, MeshSelectionNode->SecondaryIndices);
 	}
-
-	MeshSelectionNode->Invalidate();
+	
 }
 
 void UClothMeshSelectionTool::RequestAction(EClothMeshSelectionToolActions ActionType)
@@ -486,12 +520,18 @@ void UClothMeshSelectionTool::ApplyAction(EClothMeshSelectionToolActions ActionT
 	switch (ActionType)
 	{
 	case EClothMeshSelectionToolActions::ImportFromCollection:
-		ImportFromCollectionAction();
+		ImportFromCollection(/*bImportFromSecondarySet = */ false);
+		break;
+	case EClothMeshSelectionToolActions::ImportSecondaryFromCollection:
+		ImportFromCollection(/*bImportFromSecondarySet = */ true);
+		break;
+	case EClothMeshSelectionToolActions::TogglePrimarySecondary:
+		TogglePrimarySecondaryAction();
 		break;
 	}
 }
 
-void UClothMeshSelectionTool::ImportFromCollectionAction()
+void UClothMeshSelectionTool::ImportFromCollection(bool bImportFromSecondarySet)
 {
 	if (const TSharedPtr<const FManagedArrayCollection> ClothCollection = ClothEditorContextObject->GetSelectedInputClothCollection().Pin())
 	{
@@ -533,9 +573,11 @@ void UClothMeshSelectionTool::ImportFromCollectionAction()
 				}
 			}
 			const FName InSelectionName(ToolProperties->Name);
-			if (const TSet<int32>* const SelectionSet = SelectionFacade.FindSelectionSet(InSelectionName))
+
+			if (const TSet<int32>* const SelectionSet = bImportFromSecondarySet ? SelectionFacade.FindSelectionSecondarySet(InSelectionName) : SelectionFacade.FindSelectionSet(InSelectionName))
 			{
-				const FName& ExistingSelectionGroup = SelectionFacade.GetSelectionGroup(InSelectionName);
+				const FName& ExistingSelectionGroup = bImportFromSecondarySet ? SelectionFacade.GetSelectionSecondaryGroup(InSelectionName) : SelectionFacade.GetSelectionGroup(InSelectionName);
+
 				if (ExistingSelectionGroup == GroupName)
 				{
 					auto AppendVerticesIfValid = [this]<typename T>(TSet<int32>&Dest, const T & Source)
@@ -576,7 +618,10 @@ void UClothMeshSelectionTool::ImportFromCollectionAction()
 						{
 							for (const int32 SelectionIndex : *SelectionSet)
 							{
-								AppendVerticesIfValid(OutSelection.SelectedCornerIDs, SelectionToDynamicMesh[SelectionIndex]);
+								if (SelectionIndex < SelectionToDynamicMesh.Num())		// Could be loading a render mesh selection where NumRenderVertices > NumSimVertices
+								{
+									AppendVerticesIfValid(OutSelection.SelectedCornerIDs, SelectionToDynamicMesh[SelectionIndex]);
+								}
 							}
 						}
 						else
@@ -596,4 +641,50 @@ void UClothMeshSelectionTool::ImportFromCollectionAction()
 	}
 }
 
+void UClothMeshSelectionTool::TogglePrimarySecondaryAction()
+{
+	// Save any changes made in the current mode to the node
+	if (bAnyChangeMade)
+	{
+		UpdateSelectedNode();
+		bAnyChangeMade = false;
+	}
+
+	// Toggle
+	ToolProperties->bSecondarySelection = !ToolProperties->bSecondarySelection;
+
+	// Re-initialize the Selection from the selected Dataflow node
+	FString ExistingSelectionName;
+	FGroupTopologySelection ExistingNodeSelection;
+	GetSelectedNodeInfo(ExistingSelectionName, ExistingNodeSelection);
+
+	constexpr bool bBroadcastChange = false;
+	SelectionMechanic->SetSelection(ExistingNodeSelection, bBroadcastChange);
+
+	if (ExistingNodeSelection.SelectedCornerIDs.Num() > 0)
+	{
+		check(ExistingNodeSelection.SelectedEdgeIDs.Num() == 0);
+		check(ExistingNodeSelection.SelectedGroupIDs.Num() == 0);
+		SelectionMechanic->Properties->bSelectVertices = true;
+		SelectionMechanic->Properties->bSelectFaces = false;
+	}
+	else
+	{
+		SelectionMechanic->Properties->bSelectVertices = false;
+		SelectionMechanic->Properties->bSelectFaces = true;
+	}
+
+	ToolProperties->Name = ExistingSelectionName;
+
+	PreviewMesh->FastNotifySecondaryTrianglesChanged();
+
+	UpdatePrimarySecondaryMessage();
+}
+
+void UClothMeshSelectionTool::UpdatePrimarySecondaryMessage()
+{
+	const FText Message = ToolProperties->bSecondarySelection ? LOCTEXT("SecondarySelectionMode", "Secondary Selection") : LOCTEXT("PrimarySelectionMode", "Primary Selection");
+	GetToolManager()->DisplayMessage(Message, EToolMessageLevel::UserWarning);
+}
+
 #undef LOCTEXT_NAMESPACE
diff --git a/Engine/Plugins/Experimental/ChaosClothAssetEditor/Source/ChaosClothAssetEditorTools/Public/ChaosClothAsset/ClothMeshSelectionTool.h b/Engine/Plugins/Experimental/ChaosClothAssetEditor/Source/ChaosClothAssetEditorTools/Public/ChaosClothAsset/ClothMeshSelectionTool.h
--- a/Engine/Plugins/Experimental/ChaosClothAssetEditor/Source/ChaosClothAssetEditorTools/Public/ChaosClothAsset/ClothMeshSelectionTool.h
+++ b/Engine/Plugins/Experimental/ChaosClothAssetEditor/Source/ChaosClothAssetEditorTools/Public/ChaosClothAsset/ClothMeshSelectionTool.h
@@ -35,7 +35,9 @@ enum class EClothMeshSelectionToolActions
 {
 	NoAction,
 
-	ImportFromCollection
+	ImportFromCollection,
+	ImportSecondaryFromCollection,
+	TogglePrimarySecondary,
 };
 
 UCLASS()
@@ -56,6 +58,19 @@ public:
 	{
 		PostAction(EClothMeshSelectionToolActions::ImportFromCollection);
 	}
+
+	UFUNCTION(CallInEditor, Category = Operations)
+	void ImportSecondaryFromCollection()
+	{
+		PostAction(EClothMeshSelectionToolActions::ImportSecondaryFromCollection);
+	}
+
+	UFUNCTION(CallInEditor, Category = Operations)
+	void TogglePrimarySecondary()
+	{
+		PostAction(EClothMeshSelectionToolActions::TogglePrimarySecondary);
+	}
+
 };
 
 UCLASS()
@@ -68,6 +83,13 @@ public:
 	UPROPERTY(EditAnywhere, Transient, Category = Name, meta = (DisplayName = "Name", TransientToolProperty))
 	FString Name;
 
+	/**
+	* Whether the user is editing the primary or secondary selection set. Stored as a property here so that the tool will remember what the
+	* user was doing the last time the tool shut down.
+	*/
+	UPROPERTY()
+	bool bSecondarySelection = false;
+
 	UPROPERTY(EditAnywhere, Category = Visualization, meta = (DisplayName = "Show Vertices"))
 	bool bShowVertices = false;
 	
@@ -139,7 +161,9 @@ private:
 	EClothMeshSelectionToolActions PendingAction;
 	virtual void ApplyAction(EClothMeshSelectionToolActions ActionType);
 
-	void ImportFromCollectionAction();
-
+	void ImportFromCollection(bool bImportFromSecondarySet);
+	void TogglePrimarySecondaryAction();
+	
+	void UpdatePrimarySecondaryMessage();
 };
 
