Chaos Cloth Asset - Temporarily disable the merging of properties as it is crashing.

diff --git a/Engine/Plugins/Experimental/ChaosClothAssetEditor/Source/ChaosClothAssetDataflowNodes/Private/ChaosClothAsset/MergeClothCollectionsNode.cpp b/Engine/Plugins/Experimental/ChaosClothAssetEditor/Source/ChaosClothAssetDataflowNodes/Private/ChaosClothAsset/MergeClothCollectionsNode.cpp
--- a/Engine/Plugins/Experimental/ChaosClothAssetEditor/Source/ChaosClothAssetDataflowNodes/Private/ChaosClothAsset/MergeClothCollectionsNode.cpp
+++ b/Engine/Plugins/Experimental/ChaosClothAssetEditor/Source/ChaosClothAssetDataflowNodes/Private/ChaosClothAsset/MergeClothCollectionsNode.cpp
@@ -308,7 +308,7 @@ void FChaosClothAssetMergeClothCollectionsNode::Evaluate(Dataflow::FContext& Con
 			if (OtherPropertyFacade.IsValid())
 			{
 				// Change that boolean to come back to the old behavior
-				static constexpr bool bOverrideProperties = false;
+				static constexpr bool bOverrideProperties = true;
 				if(bOverrideProperties)
 				{
 					constexpr bool bUpdateExistingProperties = true; // Want last one wins.
