When converting from Collection to Dynamic Mesh, include the render mesh weight maps data.

#rnx

diff --git a/Engine/Plugins/Experimental/ChaosClothAssetEditor/Source/ChaosClothAssetTools/Private/ChaosClothAsset/ClothPatternToDynamicMesh.cpp b/Engine/Plugins/Experimental/ChaosClothAssetEditor/Source/ChaosClothAssetTools/Private/ChaosClothAsset/ClothPatternToDynamicMesh.cpp
--- a/Engine/Plugins/Experimental/ChaosClothAssetEditor/Source/ChaosClothAssetTools/Private/ChaosClothAsset/ClothPatternToDynamicMesh.cpp
+++ b/Engine/Plugins/Experimental/ChaosClothAssetEditor/Source/ChaosClothAssetTools/Private/ChaosClothAsset/ClothPatternToDynamicMesh.cpp
@@ -6,6 +6,7 @@
 #include "ChaosClothAsset/ClothAsset.h"
 #include "ChaosClothAsset/ClothComponent.h"
 #include "ChaosClothAsset/CollectionClothFacade.h"
+#include "ChaosClothAsset/ClothCollectionGroup.h"
 #include "DynamicMesh/MeshNormals.h"
 #include "DynamicMesh/NonManifoldMappingSupport.h"
 #include "Engine/SkeletalMesh.h"
@@ -210,14 +211,17 @@ public:
 			}
 		}
 
+		//
+		// Weight map layers precomputation
+		//
 		if(VertexDataType != EClothPatternVertexType::Render)
 		{
-			//
-			// Weight map layers precomputation
-			//
-			// No weight maps for render mesh
 			WeightMapNames = Cloth.GetWeightMapNames();
 		}
+		else 
+		{
+			WeightMapNames = Cloth.GetUserDefinedAttributeNames<float>(ClothCollectionGroup::RenderVertices);
+		}
 
 		// Set the reference skeleton if available
 		const FString& SkeletalMeshPathName = ClothFacade.GetSkeletalMeshPathName();
@@ -244,7 +248,7 @@ public:
 
 	int32 NumWeightMapLayers() const
 	{
-		return (VertexDataType == EClothPatternVertexType::Render) ? 0 : WeightMapNames.Num(); // No weight maps for render mesh
+		return WeightMapNames.Num();
 	}
 
 	FName GetWeightMapName(int32 LayerIndex) const
@@ -254,9 +258,31 @@ public:
 
 	float GetVertexWeight(int32 LayerIndex, VertIDType VertexIndex) const
 	{
+		checkSlow(LayerIndex < WeightMapNames.Num());
+
 		// All weights live on 3D indices.
 		const int32 VertexWeightIndex = (VertexDataType == EClothPatternVertexType::Sim2D) ? Cloth.GetSimVertex3DLookup()[VertexIndex] : VertexIndex;
-		return Cloth.GetWeightMap(WeightMapNames[LayerIndex])[VertexWeightIndex];
+
+		TConstArrayView<float> WeightMap;
+		
+		if (VertexDataType != EClothPatternVertexType::Render)
+		{
+			WeightMap = Cloth.GetWeightMap(WeightMapNames[LayerIndex]);
+		}
+		else
+		{
+			WeightMap = Cloth.GetUserDefinedAttribute<float>(WeightMapNames[LayerIndex], ClothCollectionGroup::RenderVertices);
+		}
+
+		checkSlow(VertexWeightIndex < WeightMap.Num());
+		if (VertexWeightIndex < WeightMap.Num())
+		{
+			return WeightMap[VertexWeightIndex];
+		}
+		else
+		{
+			return 0;
+		}
 	}
 
 	// --"Vertex Buffer" info
