Chaos Cloth Asset: switch Fabrics to use AnisoSpring rather than AnisoStretch.

diff --git a/Engine/Plugins/Experimental/ChaosClothAsset/Source/ChaosClothAsset/Private/ChaosClothAsset/CollectionClothFabricFacade.cpp b/Engine/Plugins/Experimental/ChaosClothAsset/Source/ChaosClothAsset/Private/ChaosClothAsset/CollectionClothFabricFacade.cpp
--- a/Engine/Plugins/Experimental/ChaosClothAsset/Source/ChaosClothAsset/Private/ChaosClothAsset/CollectionClothFabricFacade.cpp
+++ b/Engine/Plugins/Experimental/ChaosClothAsset/Source/ChaosClothAsset/Private/ChaosClothAsset/CollectionClothFabricFacade.cpp
@@ -41,7 +41,7 @@ namespace UE::Chaos::ClothAsset
 			ClothCollection->GetElements(ClothCollection->GetFabricBucklingStiffness())[GetElementIndex()]);
 	}
 	
-	FCollectionClothFabricConstFacade::FAnisotropicData FCollectionClothFabricConstFacade::GetXPBDAnisoStretchStiffness() const
+	FCollectionClothFabricConstFacade::FAnisotropicData FCollectionClothFabricConstFacade::GetXPBDAnisoSpringStiffness() const
 	{
 		return FCollectionClothFabricConstFacade::FAnisotropicData(
 			ClothCollection->GetElements(ClothCollection->GetFabricStretchStiffness())[GetElementIndex()]);
@@ -97,7 +97,7 @@ namespace UE::Chaos::ClothAsset
 	void FCollectionClothFabricFacade::Initialize(const FCollectionClothFabricConstFacade& OtherFabricFacade)
 	{
 		Initialize( OtherFabricFacade.GetXPBDAnisoBendingStiffness(), OtherFabricFacade.GetXPBDAnisoBucklingRatio(), OtherFabricFacade.GetXPBDAnisoBucklingStiffness(),
-			OtherFabricFacade.GetXPBDAnisoStretchStiffness(), OtherFabricFacade.GetDensityWeighted(),
+			OtherFabricFacade.GetXPBDAnisoSpringStiffness(), OtherFabricFacade.GetDensityWeighted(),
 			OtherFabricFacade.GetFrictionCoefficient(), OtherFabricFacade.GetXPBDAnisoDamping(), OtherFabricFacade.GetCollisionThickness());
 	}
 
diff --git a/Engine/Plugins/Experimental/ChaosClothAsset/Source/ChaosClothAsset/Public/ChaosClothAsset/CollectionClothFabricFacade.h b/Engine/Plugins/Experimental/ChaosClothAsset/Source/ChaosClothAsset/Public/ChaosClothAsset/CollectionClothFabricFacade.h
--- a/Engine/Plugins/Experimental/ChaosClothAsset/Source/ChaosClothAsset/Public/ChaosClothAsset/CollectionClothFabricFacade.h
+++ b/Engine/Plugins/Experimental/ChaosClothAsset/Source/ChaosClothAsset/Public/ChaosClothAsset/CollectionClothFabricFacade.h
@@ -49,7 +49,7 @@ namespace UE::Chaos::ClothAsset
 		FAnisotropicData GetXPBDAnisoBucklingStiffness() const;
 
 		/** Return the anisotropic stretch stiffness */
-		FAnisotropicData GetXPBDAnisoStretchStiffness() const;
+		FAnisotropicData GetXPBDAnisoSpringStiffness() const;
 
 		/** Return the cloth density */
 		float GetDensityWeighted() const;
diff --git a/Engine/Plugins/Experimental/ChaosClothAssetEditor/Source/ChaosClothAssetDataflowNodes/Private/ChaosClothAsset/SimulationStretchConfigNode.cpp b/Engine/Plugins/Experimental/ChaosClothAssetEditor/Source/ChaosClothAssetDataflowNodes/Private/ChaosClothAsset/SimulationStretchConfigNode.cpp
--- a/Engine/Plugins/Experimental/ChaosClothAssetEditor/Source/ChaosClothAssetDataflowNodes/Private/ChaosClothAsset/SimulationStretchConfigNode.cpp
+++ b/Engine/Plugins/Experimental/ChaosClothAssetEditor/Source/ChaosClothAssetDataflowNodes/Private/ChaosClothAsset/SimulationStretchConfigNode.cpp
@@ -26,54 +26,28 @@ void FChaosClothAssetSimulationStretchConfigNode::AddProperties(FPropertyHelper&
 	{
 		if(SolverType == EChaosClothAssetConstraintSolverType::XPBD)
 		{
-			if(DistributionType == EChaosClothAssetConstraintDistributionType::Anisotropic)
+			if (DistributionType == EChaosClothAssetConstraintDistributionType::Anisotropic)
 			{
-				if(ConstraintType == EChaosClothAssetStretchConstraintType::StretchShear)
-				{
-					PropertyHelper.SetPropertyBool(FName(TEXT("XPBDAnisoStretchUse3dRestLengths")), bStretchUse3dRestLengths, {
-						FName(TEXT("XPBDAnisoSpringUse3dRestLengths"))}, ECollectionPropertyFlags::None);  // Non animatable
+				PropertyHelper.SetPropertyBool(FName(TEXT("XPBDAnisoSpringUse3dRestLengths")), bStretchUse3dRestLengths, {
+					FName(TEXT("XPBDAnisoStretchUse3dRestLengths")) }, ECollectionPropertyFlags::None);  // Non animatable
 
-					PropertyHelper.SetPropertyWeighted(FName(TEXT("XPBDAnisoStretchStiffnessWarp")), StretchStiffnessWarp, {
-						FName(TEXT("EdgeSpringStiffness")),
-						FName(TEXT("XPBDEdgeSpringStiffness")),
-						FName(TEXT("XPBDAnisoSpringStiffnessWarp"))});
-				
-					PropertyHelper.SetPropertyWeighted(FName(TEXT("XPBDAnisoStretchStiffnessWeft")), StretchStiffnessWeft, {
-						FName(TEXT("XPBDAnisoSpringStiffnessWeft"))});
-					PropertyHelper.SetPropertyWeighted(FName(TEXT("XPBDAnisoStretchStiffnessBias")), StretchStiffnessBias, {
-						FName(TEXT("XPBDAnisoSpringStiffnessBias"))});
-					PropertyHelper.SetPropertyWeighted(FName(TEXT("XPBDAnisoStretchDamping")), StretchDamping, {
-						FName(TEXT("XPBDEdgeSpringDamping")),
-						FName(TEXT("XPBDAnisoSpringDamping"))});
-				
-					PropertyHelper.SetPropertyWeighted(FName(TEXT("XPBDAnisoStretchWarpScale")), StretchWarpScale, {
-						FName(TEXT("XPBDAnisoSpringWarpScale"))});
-					PropertyHelper.SetPropertyWeighted(FName(TEXT("XPBDAnisoStretchWeftScale")), StretchWeftScale, {
-						FName(TEXT("XPBDAnisoSpringWeftScale"))});
-				}
-				else
-				{
-					PropertyHelper.SetPropertyBool(FName(TEXT("XPBDAnisoSpringUse3dRestLengths")), bStretchUse3dRestLengths, {
-						FName(TEXT("XPBDAnisoStretchUse3dRestLengths"))}, ECollectionPropertyFlags::None);  // Non animatable
+				PropertyHelper.SetPropertyWeighted(FName(TEXT("XPBDAnisoSpringStiffnessWarp")), StretchStiffnessWarp, {
+					FName(TEXT("EdgeSpringStiffness")),
+					FName(TEXT("XPBDEdgeSpringStiffness")),
+					FName(TEXT("XPBDAnisoStretchStiffnessWarp")) });
 
-					PropertyHelper.SetPropertyWeighted(FName(TEXT("XPBDAnisoSpringStiffnessWarp")), StretchStiffnessWarp, {
-						FName(TEXT("EdgeSpringStiffness")),
-						FName(TEXT("XPBDEdgeSpringStiffness")),
-						FName(TEXT("XPBDAnisoStretchStiffnessWarp"))});
-				
-					PropertyHelper.SetPropertyWeighted(FName(TEXT("XPBDAnisoSpringStiffnessWeft")), StretchStiffnessWeft, {
-						FName(TEXT("XPBDAnisoStretchStiffnessWeft"))});
-					PropertyHelper.SetPropertyWeighted(FName(TEXT("XPBDAnisoSpringStiffnessBias")), StretchStiffnessBias, {
-						FName(TEXT("XPBDAnisoStretchStiffnessBias"))});
-					PropertyHelper.SetPropertyWeighted(FName(TEXT("XPBDAnisoSpringDamping")), StretchDamping, {
-						FName(TEXT("XPBDAnisoStretchDamping")),
-						FName(TEXT("XPBDEdgeSpringDamping"))});
-				
-					PropertyHelper.SetPropertyWeighted(FName(TEXT("XPBDAnisoSpringWarpScale")), StretchWarpScale, {
-						FName(TEXT("XPBDAnisoStretchWarpScale"))});
-					PropertyHelper.SetPropertyWeighted(FName(TEXT("XPBDAnisoSpringWeftScale")), StretchWeftScale, {
-						FName(TEXT("XPBDAnisoStretchWeftScale"))});
-				}
+				PropertyHelper.SetPropertyWeighted(FName(TEXT("XPBDAnisoSpringStiffnessWeft")), StretchStiffnessWeft, {
+					FName(TEXT("XPBDAnisoStretchStiffnessWeft")) });
+				PropertyHelper.SetPropertyWeighted(FName(TEXT("XPBDAnisoSpringStiffnessBias")), StretchStiffnessBias, {
+					FName(TEXT("XPBDAnisoStretchStiffnessBias")) });
+				PropertyHelper.SetPropertyWeighted(FName(TEXT("XPBDAnisoSpringDamping")), StretchDamping, {
+					FName(TEXT("XPBDAnisoStretchDamping")),
+					FName(TEXT("XPBDEdgeSpringDamping")) });
+
+				PropertyHelper.SetPropertyWeighted(FName(TEXT("XPBDAnisoSpringWarpScale")), StretchWarpScale, {
+					FName(TEXT("XPBDAnisoStretchWarpScale")) });
+				PropertyHelper.SetPropertyWeighted(FName(TEXT("XPBDAnisoSpringWeftScale")), StretchWeftScale, {
+					FName(TEXT("XPBDAnisoStretchWeftScale")) });
 			}
 			else
 			{
@@ -109,37 +83,38 @@ void FChaosClothAssetSimulationStretchConfigNode::AddProperties(FPropertyHelper&
 	}
 	else
 	{
-		PropertyHelper.SetPropertyBool(FName(TEXT("XPBDAnisoStretchUse3dRestLengths")), bStretchUse3dRestLengths, {
-						FName(TEXT("XPBDAnisoSpringUse3dRestLengths"))}, ECollectionPropertyFlags::None);  // Non animatable
+		PropertyHelper.SetPropertyBool(FName(TEXT("XPBDAnisoSpringUse3dRestLengths")), bStretchUse3dRestLengths, {
+						FName(TEXT("XPBDAnisoStretchUse3dRestLengths"))}, ECollectionPropertyFlags::None);  // Non animatable
 
-		SetFabricPropertyWeighted(FName(TEXT("XPBDAnisoStretchStiffnessWarp")), StretchStiffnessWarp,ClothFacade, PropertyHelper, [](
+		SetFabricPropertyWeighted(FName(TEXT("XPBDAnisoSpringStiffnessWarp")), StretchStiffnessWarp,ClothFacade, PropertyHelper, [](
 			const UE::Chaos::ClothAsset::FCollectionClothFabricFacade& FabricFacade)-> float
 		 {
-			 return FabricFacade.GetXPBDAnisoStretchStiffness().Warp;
+			 return FabricFacade.GetXPBDAnisoSpringStiffness().Warp;
 		 }, {
 			FName(TEXT("EdgeSpringStiffness")),
-			FName(TEXT("XPBDEdgeSpringStiffness"))});
+			FName(TEXT("XPBDEdgeSpringStiffness")),
+			FName(TEXT("XPBDAnisoStretchStiffnessWarp"))});
 
-		SetFabricPropertyWeighted(FName(TEXT("XPBDAnisoStretchStiffnessWeft")), StretchStiffnessWeft, ClothFacade, PropertyHelper, [](
+		SetFabricPropertyWeighted(FName(TEXT("XPBDAnisoSpringStiffnessWeft")), StretchStiffnessWeft, ClothFacade, PropertyHelper, [](
 			const UE::Chaos::ClothAsset::FCollectionClothFabricFacade& FabricFacade)-> float
 		 {
-			 return FabricFacade.GetXPBDAnisoStretchStiffness().Weft;
-		 }, {});
+			 return FabricFacade.GetXPBDAnisoSpringStiffness().Weft;
+		 }, { FName(TEXT("XPBDAnisoStretchStiffnessWeft")) });
 		
-		SetFabricPropertyWeighted(FName(TEXT("XPBDAnisoStretchStiffnessBias")), StretchStiffnessBias,ClothFacade, PropertyHelper, [](
+		SetFabricPropertyWeighted(FName(TEXT("XPBDAnisoSpringStiffnessBias")), StretchStiffnessBias, ClothFacade, PropertyHelper, [](
 			const UE::Chaos::ClothAsset::FCollectionClothFabricFacade& FabricFacade)-> float
 		 {
-			 return FabricFacade.GetXPBDAnisoStretchStiffness().Bias;
-		 }, {});
+			 return FabricFacade.GetXPBDAnisoSpringStiffness().Bias;
+		 }, { FName(TEXT("XPBDAnisoStretchStiffnessBias")) });
 
-		SetFabricPropertyWeighted(FName(TEXT("XPBDAnisoStretchDamping")), StretchDamping,ClothFacade, PropertyHelper, [](
+		SetFabricPropertyWeighted(FName(TEXT("XPBDAnisoSpringDamping")), StretchDamping, ClothFacade, PropertyHelper, [](
 			const UE::Chaos::ClothAsset::FCollectionClothFabricFacade& FabricFacade)-> float
 		 {
 			 return FabricFacade.GetXPBDAnisoDamping();
-		 }, {
-			FName(TEXT("XPBDEdgeSpringDamping"))});
+		 }, { FName(TEXT("XPBDEdgeSpringDamping")),
+			FName(TEXT("XPBDAnisoStretchDamping")) });
 
-		PropertyHelper.SetPropertyWeighted(FName(TEXT("XPBDAnisoStretchWarpScale")), StretchWarpScale, {});
-		PropertyHelper.SetPropertyWeighted(FName(TEXT("XPBDAnisoStretchWeftScale")), StretchWeftScale, {});
+		PropertyHelper.SetPropertyWeighted(FName(TEXT("XPBDAnisoSpringWarpScale")), StretchWarpScale, { FName(TEXT("XPBDAnisoStretchWarpScale")) });
+		PropertyHelper.SetPropertyWeighted(FName(TEXT("XPBDAnisoSpringWeftScale")), StretchWeftScale, { FName(TEXT("XPBDAnisoStretchWeftScale")) });
 	}
 }
diff --git a/Engine/Plugins/Experimental/ChaosClothAssetEditor/Source/ChaosClothAssetDataflowNodes/Public/ChaosClothAsset/SimulationStretchConfigNode.h b/Engine/Plugins/Experimental/ChaosClothAssetEditor/Source/ChaosClothAssetDataflowNodes/Public/ChaosClothAsset/SimulationStretchConfigNode.h
--- a/Engine/Plugins/Experimental/ChaosClothAssetEditor/Source/ChaosClothAssetDataflowNodes/Public/ChaosClothAsset/SimulationStretchConfigNode.h
+++ b/Engine/Plugins/Experimental/ChaosClothAssetEditor/Source/ChaosClothAssetDataflowNodes/Public/ChaosClothAsset/SimulationStretchConfigNode.h
@@ -27,22 +27,18 @@ public:
 	UPROPERTY(EditAnywhere, Category = "Stretch Types", Meta = (EditCondition = "bImportProperties == false && SolverType == EChaosClothAssetConstraintSolverType::XPBD", EditConditionHides))
 	EChaosClothAssetConstraintDistributionType DistributionType = EChaosClothAssetConstraintDistributionType::Isotropic;
 	
-	/** Constraint stretching type.  */
-	UPROPERTY(EditAnywhere, Category = "Stretch Types", Meta = (EditCondition = "bImportProperties == false && SolverType == EChaosClothAssetConstraintSolverType::XPBD && DistributionType == EChaosClothAssetConstraintDistributionType::Anisotropic", EditConditionHides))
-	EChaosClothAssetStretchConstraintType ConstraintType = EChaosClothAssetStretchConstraintType::StretchShear;
-	
 	/** Add an area constraint in case of isotropic distribution  */
 	UPROPERTY(EditAnywhere, Category = "Stretch Types", Meta = (EditCondition = "bImportProperties == false && ((SolverType == EChaosClothAssetConstraintSolverType::XPBD && DistributionType == EChaosClothAssetConstraintDistributionType::Isotropic) || SolverType == EChaosClothAssetConstraintSolverType::PBD)", EditConditionHides))
 	bool bAddAreaConstraint = true;
 
 	/**
-	 * The stiffness of the stretch constraints.
+	 * The stiffness of the stretch constraints. Note that PBD stiffnesses will be internally clamped to [0,1].
 	 * If a valid weight map is found with the given Weight Map name, then both Low and High values
 	 * are interpolated with the per particle weight to make the final value used for the simulation.
 	 * Otherwise all particles are considered to have a zero weight, and only the Low value is meaningful.
 	 */
 	UPROPERTY(EditAnywhere, Category = "Stretch Properties", Meta = (UIMin = "0", UIMax = "10000", ClampMin = "0", ClampMax = "1000000000", EditCondition = "bImportProperties == false && (SolverType == EChaosClothAssetConstraintSolverType::PBD || DistributionType == EChaosClothAssetConstraintDistributionType::Isotropic)", EditConditionHides))
-	FChaosClothAssetWeightedValue StretchStiffness = { true, 100.f, 100.f, TEXT("StretchStiffness") };
+	FChaosClothAssetWeightedValue StretchStiffness = { true, 1.f, 1.f, TEXT("StretchStiffness") };
 	
 	/**
 	 * The stiffness of the stretch constraints in the warp (vertical) direction.
@@ -100,14 +96,14 @@ public:
 	FChaosClothAssetWeightedValue StretchWeftScale = { true, 1.f, 1.f, TEXT("StretchWeftScale") };
 
 	/**
-	 * The stiffness of the surface area preservation constraints.
+	 * The stiffness of the surface area preservation constraints. Note that PBD stiffnesses will be internally clamped to [0,1].
 	 * Increase the solver iteration count for stiffer materials.
 	 * If a valid weight map is found with the given Weight Map name, then both Low and High values
 	 * are interpolated with the per particle weight to make the final value used for the simulation.
 	 * Otherwise all particles are considered to have a zero weight, and only the Low value is meaningful.
 	 */
 	UPROPERTY(EditAnywhere, Category = "Stretch Properties", Meta = (UIMin = "0", UIMax = "10000", ClampMin = "0", ClampMax = "10000000", EditCondition = "bImportProperties == false && bAddAreaConstraint == true && ((SolverType == EChaosClothAssetConstraintSolverType::XPBD && DistributionType == EChaosClothAssetConstraintDistributionType::Isotropic) || SolverType == EChaosClothAssetConstraintSolverType::PBD)", EditConditionHides))
-	FChaosClothAssetWeightedValue AreaStiffness = { true, 100.f, 100.f, TEXT("AreaStiffness") };
+	FChaosClothAssetWeightedValue AreaStiffness = { true, 1.f, 1.f, TEXT("AreaStiffness") };
 
 	FChaosClothAssetSimulationStretchConfigNode(const Dataflow::FNodeParameters& InParam, FGuid InGuid = FGuid::NewGuid());
 
