Cloth Editor Selection Tool: request graph evaluation when leaving selection tool

diff --git a/Engine/Plugins/Experimental/ChaosClothAssetEditor/Source/ChaosClothAssetEditorTools/Private/ChaosClothAsset/ClothMeshSelectionTool.cpp b/Engine/Plugins/Experimental/ChaosClothAssetEditor/Source/ChaosClothAssetEditorTools/Private/ChaosClothAsset/ClothMeshSelectionTool.cpp
--- a/Engine/Plugins/Experimental/ChaosClothAssetEditor/Source/ChaosClothAssetEditorTools/Private/ChaosClothAsset/ClothMeshSelectionTool.cpp
+++ b/Engine/Plugins/Experimental/ChaosClothAssetEditor/Source/ChaosClothAssetEditorTools/Private/ChaosClothAsset/ClothMeshSelectionTool.cpp
@@ -298,6 +298,13 @@ void UClothMeshSelectionTool::OnShutdown(EToolShutdownType ShutdownType)
 		GetToolManager()->EndUndoTransaction();
 	}
 
+	// Invalidate the node even if we are hitting cancel. We could have saved new selection information to the node by switching primary/secondary modes and we'd expect 
+	// that information to make its way into the ClothCollection
+	if (FChaosClothAssetSelectionNode* const MeshSelectionNode = ClothEditorContextObject->GetSingleSelectedNodeOfType<FChaosClothAssetSelectionNode>())
+	{
+		MeshSelectionNode->Invalidate();
+	}
+
 	SelectionMechanic->Properties->SaveProperties(this);
 	ToolProperties->SaveProperties(this);
 
