Chaos Cloth Asset:
- Remove the USD Import Transform property.
- Rephrase the doc appearing in the Refresh Asset tooltip.

diff --git a/Engine/Plugins/Experimental/ChaosClothAssetEditor/Source/ChaosClothAssetDataflowNodes/Private/ChaosClothAsset/USDImportNode.cpp b/Engine/Plugins/Experimental/ChaosClothAssetEditor/Source/ChaosClothAssetDataflowNodes/Private/ChaosClothAsset/USDImportNode.cpp
--- a/Engine/Plugins/Experimental/ChaosClothAssetEditor/Source/ChaosClothAssetDataflowNodes/Private/ChaosClothAsset/USDImportNode.cpp
+++ b/Engine/Plugins/Experimental/ChaosClothAssetEditor/Source/ChaosClothAssetDataflowNodes/Private/ChaosClothAsset/USDImportNode.cpp
@@ -1015,17 +1015,6 @@ bool FChaosClothAssetUSDImportNode::ImportFromCache(const TSharedRef<FManagedArr
 		}
 	}
 
-	// Apply transforms  TODO: Use a transform node/tool instead
-	FCollectionClothFacade ClothFacade(OutClothCollection);
-	for (FVector3f& SimPosition3D : ClothFacade.GetSimPosition3D())
-	{
-		SimPosition3D = Transform.TransformPosition(SimPosition3D);
-	}
-	for (FVector3f& RenderPosition : ClothFacade.GetRenderPosition())
-	{
-		RenderPosition = Transform.TransformPosition(RenderPosition);
-	}
-
 	return true;
 }
 
diff --git a/Engine/Plugins/Experimental/ChaosClothAssetEditor/Source/ChaosClothAssetDataflowNodes/Public/ChaosClothAsset/TerminalNode.h b/Engine/Plugins/Experimental/ChaosClothAssetEditor/Source/ChaosClothAssetDataflowNodes/Public/ChaosClothAsset/TerminalNode.h
--- a/Engine/Plugins/Experimental/ChaosClothAssetEditor/Source/ChaosClothAssetDataflowNodes/Public/ChaosClothAsset/TerminalNode.h
+++ b/Engine/Plugins/Experimental/ChaosClothAssetEditor/Source/ChaosClothAssetDataflowNodes/Public/ChaosClothAsset/TerminalNode.h
@@ -48,7 +48,11 @@ public:
 	/** The number of LODs currently exposed to the node UI. */
 	UPROPERTY()
 	int32 NumLods = 1;
-	/** Refresh the asset even if the ClothCollection hasn't changed to take into account changes in the build code for example. */
+	/**
+	 * Refresh the asset even if the ClothCollection hasn't changed.
+	 * Note that it is not required to manually refresh the cloth asset, this is done automatically when there is a change in the Dataflow.
+	 * This function is a developper utility used for debugging.
+	 */
 	UPROPERTY(EditAnywhere, Category = "Cloth Asset Terminal")
 	mutable FChaosClothAssetTerminalNodeRefreshAsset RefreshAsset;
 
diff --git a/Engine/Plugins/Experimental/ChaosClothAssetEditor/Source/ChaosClothAssetDataflowNodes/Public/ChaosClothAsset/USDImportNode.h b/Engine/Plugins/Experimental/ChaosClothAssetEditor/Source/ChaosClothAssetDataflowNodes/Public/ChaosClothAsset/USDImportNode.h
--- a/Engine/Plugins/Experimental/ChaosClothAssetEditor/Source/ChaosClothAssetDataflowNodes/Public/ChaosClothAsset/USDImportNode.h
+++ b/Engine/Plugins/Experimental/ChaosClothAssetEditor/Source/ChaosClothAssetDataflowNodes/Public/ChaosClothAsset/USDImportNode.h
@@ -22,13 +22,6 @@ public:
 	UPROPERTY(EditAnywhere, Category = "USD Import", Meta = (DisplayName = "USD File"))
 	FChaosClothAssetImportFilePath UsdFile;
 
-	/**
-	 * The transform to apply to the imported 3D positions. 
-	 * Experimental, will soon be removed in favour of a tool supported transform node.
-	 */
-	UPROPERTY(EditAnywhere, Category = "USD Import", Meta = (Exprimental))
-	FTransform3f Transform = FTransform3f::Identity;
-
 	FChaosClothAssetUSDImportNode(const Dataflow::FNodeParameters& InParam, FGuid InGuid = FGuid::NewGuid());
 
 private:
