Fis merging of the cloth collection

diff --git a/Engine/Plugins/Experimental/ChaosClothAssetEditor/Source/ChaosClothAssetDataflowNodes/Private/ChaosClothAsset/MergeClothCollectionsNode.cpp b/Engine/Plugins/Experimental/ChaosClothAssetEditor/Source/ChaosClothAssetDataflowNodes/Private/ChaosClothAsset/MergeClothCollectionsNode.cpp
--- a/Engine/Plugins/Experimental/ChaosClothAssetEditor/Source/ChaosClothAssetDataflowNodes/Private/ChaosClothAsset/MergeClothCollectionsNode.cpp
+++ b/Engine/Plugins/Experimental/ChaosClothAssetEditor/Source/ChaosClothAssetDataflowNodes/Private/ChaosClothAsset/MergeClothCollectionsNode.cpp
@@ -51,6 +51,7 @@ namespace UE::Chaos::ClothAsset::Private
 				IncrWeightMapName = WeightMapName;
 				IncrWeightMapName.AppendInt(++IncrCount);
 			}
+			WeightMapName = IncrWeightMapName;
 		}
 		
 		// If the low high values of the merged property are the same we don't need to build a weight map
@@ -63,9 +64,10 @@ namespace UE::Chaos::ClothAsset::Private
 			}
 			// Create if necessary a new weight map
 			OutClothFacade.AddWeightMap(FName(WeightMapName));
+			TArrayView<float> WeightMap = OutClothFacade.GetWeightMap(FName(WeightMapName));
 
-			auto FillWeightMap = [&PropertyName, &PropertyBounds](const TConstArrayView<float> InWeightMap, const FVector2f& InPropertyBounds,
-				const int32 InNumVertices, const int32 OutVertexOffset, TArrayView<float> OutWeightMap)
+			auto FillWeightMap = [&PropertyBounds, &WeightMap](const TConstArrayView<float> InWeightMap, const FVector2f& InPropertyBounds,
+				const int32 InNumVertices, const int32 OutVertexOffset)
 			{
 				const bool bHasAlreadyValues = (InWeightMap.Num() > 0);
 				for(int32 VertexIndex = 0; VertexIndex < InNumVertices; ++VertexIndex)
@@ -73,17 +75,14 @@ namespace UE::Chaos::ClothAsset::Private
                 	// If no values in the weight map we are using the low value
 					const float WeightMapValue = bHasAlreadyValues ? (InWeightMap[VertexIndex] * (InPropertyBounds[1] - InPropertyBounds[0]) +
 						InPropertyBounds[0]) : InPropertyBounds[0];
-					OutWeightMap[OutVertexOffset+VertexIndex] = (WeightMapValue - PropertyBounds[0]) / (PropertyBounds[1] - PropertyBounds[0]);
+					WeightMap[OutVertexOffset+VertexIndex] = (WeightMapValue - PropertyBounds[0]) / (PropertyBounds[1] - PropertyBounds[0]);
                 }
 			};
 			const int32 InNumVertices = InClothFacade.GetNumSimVertices3D();
 			const int32 OutNumVertices = OutClothFacade.GetNumSimVertices3D() - InNumVertices;
 			
-			TArrayView<float> OutWeightMap = OutClothFacade.GetWeightMap(FName(OutWeightMapName));
-			FillWeightMap(OutWeightMap, OutPropertyBounds, OutNumVertices, 0, OutWeightMap);
-			
-			const TConstArrayView<float> InWeightMap = InClothFacade.GetWeightMap(FName(InWeightMapName));
-			FillWeightMap(InWeightMap, InPropertyBounds, InNumVertices, OutNumVertices, OutWeightMap);
+			FillWeightMap(OutClothFacade.GetWeightMap(FName(OutWeightMapName)), OutPropertyBounds, OutNumVertices, 0);
+			FillWeightMap(InClothFacade.GetWeightMap(FName(InWeightMapName)), InPropertyBounds, InNumVertices, OutNumVertices);
 		}
 		return WeightMapName;
 	}
@@ -209,7 +208,7 @@ namespace UE::Chaos::ClothAsset::Private
                     // We keep the string value to be the one in the output if defined
                     const FString WeightMapName = BuildWeightMaps(DataflowNode, InClothFacade, OutClothFacade,
                     	InPropertyBounds, OutPropertyBounds, PropertyBounds, InPropertyKey,
-							InPropertyFacade.GetStringValue(InKeyIndex), OutPropertyFacade.GetStringValue(InKeyIndex));
+							InPropertyFacade.GetStringValue(InKeyIndex), OutPropertyFacade.GetStringValue(OutKeyIndex));
 
 					OutPropertyFacade.SetStringValue(OutKeyIndex, WeightMapName);
 					bOverrideProperty = false;
@@ -308,7 +307,7 @@ void FChaosClothAssetMergeClothCollectionsNode::Evaluate(Dataflow::FContext& Con
 			if (OtherPropertyFacade.IsValid())
 			{
 				// Change that boolean to come back to the old behavior
-				static constexpr bool bOverrideProperties = true;
+				static constexpr bool bOverrideProperties = false;
 				if(bOverrideProperties)
 				{
 					constexpr bool bUpdateExistingProperties = true; // Want last one wins.
