Cloth Editor Remesher: better handling of open meshes in the Render mesh. Add the option to detect connected boundary vertices that are coincident with another connected set of boundary vertices on another boundary elsewhere. We can then explicitly remesh these "seams" while maintaining a vertex pairing, before then going on to remesh the interior of the mesh islands.

diff --git a/Engine/Plugins/Experimental/ChaosClothAsset/Source/ChaosClothAsset/Private/ChaosClothAsset/ClothGeometryTools.cpp b/Engine/Plugins/Experimental/ChaosClothAsset/Source/ChaosClothAsset/Private/ChaosClothAsset/ClothGeometryTools.cpp
--- a/Engine/Plugins/Experimental/ChaosClothAsset/Source/ChaosClothAsset/Private/ChaosClothAsset/ClothGeometryTools.cpp
+++ b/Engine/Plugins/Experimental/ChaosClothAsset/Source/ChaosClothAsset/Private/ChaosClothAsset/ClothGeometryTools.cpp
@@ -1115,46 +1115,36 @@ namespace UE::Chaos::ClothAsset
 	}
 
 
-	void FClothGeometryTools::BuildConnectedSeams2D(const TSharedRef<const FManagedArrayCollection>& ClothCollection,
-		int32 SeamIndex,
+	void FClothGeometryTools::BuildConnectedSeams(const TArray<FIntVector2>& InputStitches,
 		const UE::Geometry::FDynamicMesh3& Mesh,
 		TArray<TArray<FIntVector2>>& Seams)
 	{
-		using namespace UE::Geometry;
-
-		FNonManifoldMappingSupport NonManifold(Mesh);
-		checkf(!NonManifold.IsNonManifoldVertexInSource(), TEXT("Cloth source is non-manifold. Cannot use FDynamicMesh to build connected seams"));
-
-		const FCollectionClothConstFacade ClothFacade(ClothCollection);
-		const FCollectionClothSeamConstFacade SeamFacade = ClothFacade.GetSeam(SeamIndex);
-		
-		TArray<FIntVector2> InputStitches(SeamFacade.GetSeamStitch2DEndIndices());
+		TArray<FIntVector2> Stitches = InputStitches;
 
 		// filter out any stitches referencing deleted vertices
-		InputStitches.SetNum(Algo::RemoveIf(InputStitches, [](const FIntVector2& Stitch)
+		Stitches.SetNum(Algo::RemoveIf(Stitches, [](const FIntVector2& Stitch)
 		{
 			return Stitch[0] == INDEX_NONE || Stitch[1] == INDEX_NONE;
 		}));
 
-
-		while (InputStitches.Num() > 0)
+		while (Stitches.Num() > 0)
 		{
 			TArray<FIntVector2> Seam;
 
-			const FIntVector2 FirstStitch = InputStitches.Last();
+			const FIntVector2 FirstStitch = Stitches.Last();
 			Seam.Add(FirstStitch);
-			InputStitches.RemoveAt(InputStitches.Num() - 1);
+			Stitches.RemoveAt(Stitches.Num() - 1);
 
 			FIntVector2 CurrStitch = FirstStitch;
 			bool bFoundNextStitch = true;
 			bool bReverseSearch = false;
-			while (InputStitches.Num() > 0 && (bFoundNextStitch || !bReverseSearch))
+			while (Stitches.Num() > 0 && (bFoundNextStitch || !bReverseSearch))
 			{
 				bFoundNextStitch = false;
 
-				for (int32 TestStitchIndex = 0; TestStitchIndex < InputStitches.Num(); ++TestStitchIndex)
+				for (int32 TestStitchIndex = 0; TestStitchIndex < Stitches.Num(); ++TestStitchIndex)
 				{
-					FIntVector2 TestStitch = InputStitches[TestStitchIndex];
+					FIntVector2 TestStitch = Stitches[TestStitchIndex];
 
 					// Stitch (A, B) is connected to stitch (C, D) if there exist edges {(A, C), (B, D)} *or* {(A, D), (B, C)} in the given DynamicMesh.
 
@@ -1163,12 +1153,12 @@ namespace UE::Chaos::ClothAsset
 					const int32 C = TestStitch[0];
 					const int32 D = TestStitch[1];
 
-					if (Mesh.FindEdge(A, C) != FDynamicMesh3::InvalidID && Mesh.FindEdge(B, D) != FDynamicMesh3::InvalidID)
+					if (Mesh.FindEdge(A, C) != UE::Geometry::FDynamicMesh3::InvalidID && Mesh.FindEdge(B, D) != UE::Geometry::FDynamicMesh3::InvalidID)
 					{
 						Seam.Add(TestStitch);
 						bFoundNextStitch = true;
 					}
-					else if (Mesh.FindEdge(A, D) != FDynamicMesh3::InvalidID && Mesh.FindEdge(B, C) != FDynamicMesh3::InvalidID)
+					else if (Mesh.FindEdge(A, D) != UE::Geometry::FDynamicMesh3::InvalidID && Mesh.FindEdge(B, C) != UE::Geometry::FDynamicMesh3::InvalidID)
 					{
 						Swap(TestStitch[0], TestStitch[1]);
 						Seam.Add(TestStitch);
@@ -1177,7 +1167,7 @@ namespace UE::Chaos::ClothAsset
 
 					if (bFoundNextStitch)
 					{
-						InputStitches.RemoveAt(TestStitchIndex);
+						Stitches.RemoveAt(TestStitchIndex);
 						CurrStitch = TestStitch;
 						break;
 					}
@@ -1198,6 +1188,25 @@ namespace UE::Chaos::ClothAsset
 		}
 	}
 
+
+	void FClothGeometryTools::BuildConnectedSeams2D(const TSharedRef<const FManagedArrayCollection>& ClothCollection,
+		int32 SeamIndex,
+		const UE::Geometry::FDynamicMesh3& Mesh,
+		TArray<TArray<FIntVector2>>& Seams)
+	{
+		using namespace UE::Geometry;
+
+		FNonManifoldMappingSupport NonManifold(Mesh);
+		checkf(!NonManifold.IsNonManifoldVertexInSource(), TEXT("Cloth source is non-manifold. Cannot use FDynamicMesh to build connected seams"));
+
+		const FCollectionClothConstFacade ClothFacade(ClothCollection);
+		const FCollectionClothSeamConstFacade SeamFacade = ClothFacade.GetSeam(SeamIndex);
+		
+		const TArray<FIntVector2> Stitches(SeamFacade.GetSeamStitch2DEndIndices());
+
+		BuildConnectedSeams(Stitches, Mesh, Seams);
+	}
+
 	void FClothGeometryTools::SampleVertices(const TConstArrayView<FVector3f> VertexPositions, float CullDiameterSq, TSet<int32>& OutVertexSet)
 	{
 		check(CullDiameterSq > 0.0f);
diff --git a/Engine/Plugins/Experimental/ChaosClothAsset/Source/ChaosClothAsset/Public/ChaosClothAsset/ClothGeometryTools.h b/Engine/Plugins/Experimental/ChaosClothAsset/Source/ChaosClothAsset/Public/ChaosClothAsset/ClothGeometryTools.h
--- a/Engine/Plugins/Experimental/ChaosClothAsset/Source/ChaosClothAsset/Public/ChaosClothAsset/ClothGeometryTools.h
+++ b/Engine/Plugins/Experimental/ChaosClothAsset/Source/ChaosClothAsset/Public/ChaosClothAsset/ClothGeometryTools.h
@@ -89,6 +89,14 @@ namespace UE::Chaos::ClothAsset
 		*/
 		static void CleanupAndCompactMesh(const TSharedRef<FManagedArrayCollection>& ClothCollection);
 
+		/**
+		 * Find sets of connected stitches for the input stitches given in random order.
+	     * Stitch (A, B) is connected to stitch (C, D) if there exist edges {(A, C), (B, D)} *or* {(A, D), (B, C)} in the given DynamicMesh.
+	     */
+		static void BuildConnectedSeams(const TArray<FIntVector2>& InputStitches,
+			const UE::Geometry::FDynamicMesh3& Mesh,
+			TArray<TArray<FIntVector2>>& Seams);
+
 		/**
 		* Find sets of connected stitches for the given seam.
 		* Stitch (A, B) is connected to stitch (C, D) if there exist edges {(A, C), (B, D)} *or* {(A, D), (B, C)} in the given DynamicMesh.
diff --git a/Engine/Plugins/Experimental/ChaosClothAssetEditor/Source/ChaosClothAssetDataflowNodes/Private/ChaosClothAsset/RemeshNode.cpp b/Engine/Plugins/Experimental/ChaosClothAssetEditor/Source/ChaosClothAssetDataflowNodes/Private/ChaosClothAsset/RemeshNode.cpp
--- a/Engine/Plugins/Experimental/ChaosClothAssetEditor/Source/ChaosClothAssetDataflowNodes/Private/ChaosClothAsset/RemeshNode.cpp
+++ b/Engine/Plugins/Experimental/ChaosClothAssetEditor/Source/ChaosClothAssetDataflowNodes/Private/ChaosClothAsset/RemeshNode.cpp
@@ -19,6 +19,8 @@
 #include "Chaos/CollectionPropertyFacade.h"
 #include "Algo/Find.h"
 #include "IMeshReductionManagerModule.h"
+#include "Algo/RemoveIf.h"
+#include "Spatial/PointSetHashTable.h"
 
 #include UE_INLINE_GENERATED_CPP_BY_NAME(RemeshNode)
 
@@ -88,6 +90,20 @@ namespace UE::Chaos::ClothAsset::Private
 		const int EdgeA = Mesh.FindEdge(SeamEdge.Stitches[0][0], SeamEdge.Stitches[1][0]);
 		const int EdgeB = Mesh.FindEdge(SeamEdge.Stitches[0][1], SeamEdge.Stitches[1][1]);
 
+		if (SeamEdge.Stitches[0][0] == SeamEdge.Stitches[0][1])
+		{
+			FSeamCollapseParameters Constraints;
+			Constraints.bCanCollapse = false;
+			return Constraints;
+		}
+
+		if (SeamEdge.Stitches[1][0] == SeamEdge.Stitches[1][1])
+		{
+			FSeamCollapseParameters Constraints;
+			Constraints.bCanCollapse = false;
+			return Constraints;
+		}
+
 		if (EdgeA == EdgeB)
 		{
 			// Don't collapse the same edge twice
@@ -286,6 +302,79 @@ namespace UE::Chaos::ClothAsset::Private
 	}
 
 
+	void FindCoincidentBoundaryVertices(const UE::Geometry::FDynamicMesh3& Mesh, TArray<FIntVector2>& Pairs)
+	{
+		using namespace UE::Geometry;
+
+		constexpr double ProximityTolerance = FMathf::ZeroTolerance;
+
+		TSet<int32> BoundaryVertices;
+		for (int32 EdgeID : Mesh.BoundaryEdgeIndicesItr())
+		{
+			const UE::Geometry::FIndex2i Vertices = Mesh.GetEdgeV(EdgeID);
+			BoundaryVertices.Add(Vertices[0]);
+			BoundaryVertices.Add(Vertices[1]);
+		}
+
+		//
+		// Create a spatial hash to speed up matching vertex search (this setup code was mostly copied from FMergeCoincidentMeshEdges)
+		//
+		
+		// use denser grid as vertex count increases
+		const int HashN = (Mesh.TriangleCount() < 100000) ? 64 : 128;
+		const UE::Geometry::FAxisAlignedBox3d Bounds = Mesh.GetBounds(true);
+		const double CellSize = FMath::Max(FMathd::ZeroTolerance, Bounds.MaxDim() / (double)HashN);
+
+		UE::Geometry::FPointSetAdapterd BoundaryVertAdapter;
+		BoundaryVertAdapter.MaxPointID = [&Mesh]() { return Mesh.MaxVertexID(); };
+		BoundaryVertAdapter.PointCount = [&BoundaryVertices]() { return BoundaryVertices.Num(); };
+		BoundaryVertAdapter.IsPoint = [&Mesh](int Idx) { return Mesh.IsVertex(Idx) && Mesh.IsBoundaryVertex(Idx); };
+		BoundaryVertAdapter.GetPoint = [&Mesh](int Idx) { return Mesh.GetVertex(Idx); };
+		BoundaryVertAdapter.HasNormals = [] { return false; };
+		BoundaryVertAdapter.GetPointNormal = [](int Idx) { return FVector3f::UnitY(); };
+
+		FPointSetHashtable BoundaryVertsHash(&BoundaryVertAdapter);
+		BoundaryVertsHash.Build(CellSize, Bounds.Min);
+		const double UseMergeSearchTol = FMathd::Min(CellSize, 2 * ProximityTolerance);
+
+		// Now look for coincident vertices
+		
+		TSet<FIntVector2> PairSet;
+
+		for (const int32 VertexAIndex : BoundaryVertices)
+		{
+			const FVector3d VertexAPosition = Mesh.GetVertex(VertexAIndex);
+
+			TArray<int> NearbyVertices;
+			BoundaryVertsHash.FindPointsInBall(VertexAPosition, UseMergeSearchTol, NearbyVertices);
+
+			for (const int32 VertexBIndex : NearbyVertices)
+			{
+				if (VertexAIndex == VertexBIndex)
+				{
+					continue;
+				}
+
+				if (Mesh.FindEdge(VertexAIndex, VertexBIndex) != UE::Geometry::FDynamicMesh3::InvalidID)
+				{
+					continue;
+				}
+
+				const FVector3d VertexBPosition = Mesh.GetVertex(VertexBIndex);
+				const double DistSqr = FVector3d::DistSquared(VertexAPosition, VertexBPosition);
+
+				if (DistSqr < ProximityTolerance * ProximityTolerance)
+				{
+					const FIntVector2 SortedPair = VertexAIndex < VertexBIndex ? FIntVector2{ VertexAIndex, VertexBIndex } : FIntVector2{ VertexBIndex, VertexAIndex };
+					PairSet.Add(SortedPair);
+				}
+			}
+		}
+
+		Pairs = PairSet.Array();
+	}
+
+
 	void RemeshSeams(UE::Geometry::FDynamicMesh3& Mesh, TArray<TArray<FIntVector2>>& Seams, double TargetEdgeLength)
 	{
 		using namespace UE::Geometry;
@@ -769,7 +858,7 @@ namespace UE::Chaos::ClothAsset::Private
 
 
 
-	bool Simplify(UE::Geometry::FDynamicMesh3& Mesh, int TargetVertexCount, UE::Geometry::FCompactMaps* CompactMaps = nullptr)
+	bool Simplify(UE::Geometry::FDynamicMesh3& Mesh, int TargetVertexCount, bool bCoarsenBoundaries, UE::Geometry::FCompactMaps* CompactMaps = nullptr)
 	{
 		using namespace UE::Geometry;
 
@@ -810,7 +899,7 @@ namespace UE::Chaos::ClothAsset::Private
 		Op.TargetMode = ESimplifyTargetType::VertexCount;
 		Op.TargetCount = TargetVertexCount;
 
-		Op.MeshBoundaryConstraint = EEdgeRefineFlags::CollapseOnly;
+		Op.MeshBoundaryConstraint = bCoarsenBoundaries ? EEdgeRefineFlags::CollapseOnly : EEdgeRefineFlags::FullyConstrained;
 		Op.GroupBoundaryConstraint = EEdgeRefineFlags::CollapseOnly;
 		Op.MaterialBoundaryConstraint = EEdgeRefineFlags::CollapseOnly;
 
@@ -1188,23 +1277,46 @@ void FChaosClothAssetRemeshNode::RemeshRenderMesh(const TSharedRef<const FManage
 
 	const bool bHasUVs = (DynamicMesh.Attributes()->PrimaryUV() != nullptr);
 
-	UE::Geometry::FCompactMaps CompactMaps;
-	if (RemeshMethodRender == EChaosClothAssetRemeshMethod::Remesh)
+	const double MeshArea = TMeshQueries<FDynamicMesh3>::GetVolumeArea(DynamicMesh).Y;
+	const int TargetTriangleCount = FMath::RoundToInt(static_cast<float>(TargetPercentRender) / 100.0f * static_cast<float>(InputMeshTriangleCount));
+	const double TargetEdgeLength = FRemeshMeshOp::CalculateTargetEdgeLength(nullptr, TargetTriangleCount, MeshArea);
+
+	TArray<TArray<FIntVector2>> Seams;
+
+	if (bRemeshRenderSeams)
 	{
-		const double MeshArea = TMeshQueries<FDynamicMesh3>::GetVolumeArea(DynamicMesh).Y;
-		const int TargetTriangleCount = FMath::RoundToInt(static_cast<float>(TargetPercentRender) / 100.0f * static_cast<float>(InputMeshTriangleCount));
-		const double TargetEdgeLength = FRemeshMeshOp::CalculateTargetEdgeLength(nullptr, TargetTriangleCount, MeshArea);
+		// Create pseudo-stitches based on boundary vertex proximity. These stitches aren't going to actually weld vertices together, but they will guide boundary remeshing.
+		// The goal is to maintain a vertex pairing along boundaries in order to avoid holes opening up when the mesh deforms due to skinning.
+		TArray<FIntVector2> Stitches;
+		Private::FindCoincidentBoundaryVertices(DynamicMesh, Stitches);
 
-		TArray<TArray<FIntVector2>> Seams;
+		FClothGeometryTools::BuildConnectedSeams(Stitches, DynamicMesh, Seams);
 
+		for (int RemeshPass = 0; RemeshPass < RenderSeamRemeshIterations; ++RemeshPass)
+		{
+			Private::RemeshSeams(DynamicMesh, Seams, TargetEdgeLength);
+		}
+
+		// Also remesh the open boundaries that are not constrained by seams
+		for (int RemeshPass = 0; RemeshPass < RenderSeamRemeshIterations; ++RemeshPass)
+		{
+			Private::RemeshBoundaries(DynamicMesh, Seams, TargetEdgeLength);
+		}
+	}
+
+
+	UE::Geometry::FCompactMaps CompactMaps;
+	if (RemeshMethodRender == EChaosClothAssetRemeshMethod::Remesh)
+	{
 		constexpr bool bUniformSmoothing = false;	// uniform smoothing can distort the UV layer pretty badly
 		const bool bSuccess = Private::Remesh(DynamicMesh, TargetEdgeLength, IterationsRender, SmoothingRender, bUniformSmoothing, Seams, &CompactMaps);
 		check(bSuccess);
 	}
 	else
 	{
+		const bool bCoarsenBoundariesDuringSimplify = !bRemeshRenderSeams;
 		const int TargetVertexCount = FMath::RoundToInt(static_cast<float>(TargetPercentRender) / 100.0f * static_cast<float>(InputMeshVertexCount));
-		Private::Simplify(DynamicMesh, TargetVertexCount, &CompactMaps);
+		Private::Simplify(DynamicMesh, TargetVertexCount, bCoarsenBoundariesDuringSimplify, &CompactMaps);
 	}
 
 	// Collect outputs
diff --git a/Engine/Plugins/Experimental/ChaosClothAssetEditor/Source/ChaosClothAssetDataflowNodes/Public/ChaosClothAsset/RemeshNode.h b/Engine/Plugins/Experimental/ChaosClothAssetEditor/Source/ChaosClothAssetDataflowNodes/Public/ChaosClothAsset/RemeshNode.h
--- a/Engine/Plugins/Experimental/ChaosClothAssetEditor/Source/ChaosClothAssetDataflowNodes/Public/ChaosClothAsset/RemeshNode.h
+++ b/Engine/Plugins/Experimental/ChaosClothAssetEditor/Source/ChaosClothAssetDataflowNodes/Public/ChaosClothAsset/RemeshNode.h
@@ -56,6 +56,14 @@ public:
 	UPROPERTY(EditAnywhere, Category = "Render Mesh", meta = (UIMin = "0.0", UIMax = "1.0", ClampMin = "0.0", ClampMax = "1.0", EditCondition = "bRemeshRender && RemeshMethodRender == EChaosClothAssetRemeshMethod::Remesh"))
 	double SmoothingRender = 0.25;
 
+	/** If checked, attempt to find matching vertices along Render mesh boundaries and remesh these separately */ 
+	UPROPERTY(EditAnywhere, Category = "Render Mesh", meta = (EditCondition = "bRemeshRender"))
+	bool bRemeshRenderSeams = false;
+
+	/** Number of remesh iterations over the Render mesh seams */
+	UPROPERTY(EditAnywhere, Category = "Render Mesh", meta = (UIMin = "0", UIMax = "20", ClampMin = "0", ClampMax = "100", EditCondition = "bRemeshRender && bRemeshRenderSeams"))
+	int32 RenderSeamRemeshIterations = 1;
+
 	FChaosClothAssetRemeshNode(const Dataflow::FNodeParameters& InParam, FGuid InGuid = FGuid::NewGuid());
 
 	virtual void Evaluate(Dataflow::FContext& Context, const FDataflowOutput* Out) const override;
