Chaos Cloth Asset:
- Added a default template Dataflow asset when creating a new Cloth Asset.
- Remove all the Dataflow creation Dialogs, and only show a warning when the suggested name already exists.
- Fixed a crash in the Proxy Deformer node on empty collection that prevent the creation of a default Dataflow.

#tests Created new Cloth Assets, and created Dataflows with exisiting/duplicate (or not) names through the changed open Cloth Asset dialog.

diff --git a/Engine/Plugins/Experimental/ChaosClothAssetEditor/Source/ChaosClothAssetDataflowNodes/Private/ChaosClothAsset/ProxyDeformerNode.cpp b/Engine/Plugins/Experimental/ChaosClothAssetEditor/Source/ChaosClothAssetDataflowNodes/Private/ChaosClothAsset/ProxyDeformerNode.cpp
--- a/Engine/Plugins/Experimental/ChaosClothAssetEditor/Source/ChaosClothAssetDataflowNodes/Private/ChaosClothAsset/ProxyDeformerNode.cpp
+++ b/Engine/Plugins/Experimental/ChaosClothAssetEditor/Source/ChaosClothAssetDataflowNodes/Private/ChaosClothAsset/ProxyDeformerNode.cpp
@@ -140,7 +140,7 @@ void FChaosClothAssetProxyDeformerNode::Evaluate(Dataflow::FContext& Context, co
 
 		// Always check for a valid cloth collection/facade to avoid processing non cloth collections
 		FCollectionClothFacade ClothFacade(ClothCollection);
-		if (ClothFacade.IsValid())
+		if (ClothFacade.IsValid() && ClothFacade.HasValidData())
 		{
 			// Retrieve the MaxDistance weight map name
 			FName MaxDistanceWeightMapName = FName(*GetValue<FString>(Context, &MaxDistance.WeightMap));
diff --git a/Engine/Plugins/Experimental/ChaosClothAssetEditor/Source/ChaosClothAssetEditor/Private/ChaosClothAsset/AssetDefinition_ClothAsset.cpp b/Engine/Plugins/Experimental/ChaosClothAssetEditor/Source/ChaosClothAssetEditor/Private/ChaosClothAsset/AssetDefinition_ClothAsset.cpp
--- a/Engine/Plugins/Experimental/ChaosClothAssetEditor/Source/ChaosClothAssetEditor/Private/ChaosClothAsset/AssetDefinition_ClothAsset.cpp
+++ b/Engine/Plugins/Experimental/ChaosClothAssetEditor/Source/ChaosClothAssetEditor/Private/ChaosClothAsset/AssetDefinition_ClothAsset.cpp
@@ -33,15 +33,17 @@ namespace ClothAssetDefinitionHelpers
 			const FString ClothName = ClothAsset->GetName();
 			NewDataflowAssetDialogConfig.DefaultAssetName = ClothName + "_Dataflow";
 			NewDataflowAssetDialogConfig.AssetClassNames.Add(DataflowClass->GetClassPathName());
-			NewDataflowAssetDialogConfig.ExistingAssetPolicy = ESaveAssetDialogExistingAssetPolicy::Disallow;
+			NewDataflowAssetDialogConfig.ExistingAssetPolicy = ESaveAssetDialogExistingAssetPolicy::AllowButWarn;
 			NewDataflowAssetDialogConfig.DialogTitleOverride = LOCTEXT("NewDataflowAssetDialogTitle", "Save Dataflow Asset As");
 		}
 
 		FContentBrowserModule& ContentBrowserModule = FModuleManager::LoadModuleChecked<FContentBrowserModule>("ContentBrowser");
 
-		FString NewPackageName;
+		FString NewPackageName = FPaths::Combine(NewDataflowAssetDialogConfig.DefaultPath, NewDataflowAssetDialogConfig.DefaultAssetName);
 		FText OutError;
-		for (bool bFilenameValid = false; !bFilenameValid; bFilenameValid = FFileHelper::IsFilenameValidForSaving(NewPackageName, OutError))
+
+		while (!FFileHelper::IsFilenameValidForSaving(NewPackageName, OutError) ||
+			LoadObject<UObject>(nullptr, *NewPackageName, nullptr, LOAD_NoWarn | LOAD_Quiet))  // Check if an object with this name already exists
 		{
 			const FString AssetSavePath = ContentBrowserModule.Get().CreateModalSaveAssetDialog(NewDataflowAssetDialogConfig);
 			if (AssetSavePath.IsEmpty())
@@ -54,15 +56,9 @@ namespace ClothAssetDefinitionHelpers
 
 		const FName NewAssetName(FPackageName::GetLongPackageAssetName(NewPackageName));
 		UPackage* const NewPackage = CreatePackage(*NewPackageName);
-		UDataflow* const NewAsset = NewObject<UDataflow>(NewPackage, DataflowClass, NewAssetName, RF_Public | RF_Standalone | RF_Transactional);
 
-		// Add a ClothAsset Terminal Node to the empty graph
-		const TSharedPtr<FAssetSchemaAction_Dataflow_CreateNode_DataflowEdNode> NodeAction =
-			FAssetSchemaAction_Dataflow_CreateNode_DataflowEdNode::CreateAction(NewAsset, FChaosClothAssetTerminalNode::StaticType());
-		constexpr UEdGraphPin* FromPin = nullptr;
-		constexpr bool bSelectNewNode = true;
-		UEdGraphNode* const NewEdNode = NodeAction->PerformAction(NewAsset, FromPin, FVector2D::Zero(), bSelectNewNode);
-		check(NewEdNode);
+		UDataflow* const ClothAssetTemplate = LoadObject<UDataflow>(NewPackage, TEXT("/ChaosClothAssetEditor/ClothAssetTemplate.ClothAssetTemplate"));
+		UDataflow* const NewAsset = DuplicateObject(ClothAssetTemplate, NewPackage, NewAssetName);
 
 		NewAsset->MarkPackageDirty();
 
@@ -135,7 +131,7 @@ namespace ClothAssetDefinitionHelpers
 		bool bDialogDone = false;
 		while (!bDialogDone)
 		{
-			bDialogDone = NewOrOpenDialog(ClothAsset, DataflowAsset);
+			bDialogDone = CreateNewDataflowAsset(ClothAsset, DataflowAsset);
 		}
 
 		return DataflowAsset;
