Per cloth property choice to import fabrics and build weight maps

diff --git a/Engine/Plugins/Experimental/ChaosClothAssetEditor/Source/ChaosClothAssetDataflowNodes/Private/ChaosClothAsset/SimulationBendingConfigNode.cpp b/Engine/Plugins/Experimental/ChaosClothAssetEditor/Source/ChaosClothAssetDataflowNodes/Private/ChaosClothAsset/SimulationBendingConfigNode.cpp
--- a/Engine/Plugins/Experimental/ChaosClothAssetEditor/Source/ChaosClothAssetDataflowNodes/Private/ChaosClothAsset/SimulationBendingConfigNode.cpp
+++ b/Engine/Plugins/Experimental/ChaosClothAssetEditor/Source/ChaosClothAssetDataflowNodes/Private/ChaosClothAsset/SimulationBendingConfigNode.cpp
@@ -16,6 +16,7 @@ FChaosClothAssetSimulationBendingConfigNode::FChaosClothAssetSimulationBendingCo
 	RegisterInputConnection(&BendingStiffnessWeft.WeightMap);
 	RegisterInputConnection(&BendingStiffnessBias.WeightMap);
 	RegisterInputConnection(&BendingDamping.WeightMap);
+	RegisterInputConnection(&BendingAnisoDamping.WeightMap);
 	RegisterInputConnection(&BucklingStiffness.WeightMap);
 	RegisterInputConnection(&BucklingStiffnessWarp.WeightMap);
 	RegisterInputConnection(&BucklingStiffnessWeft.WeightMap);
@@ -25,213 +26,173 @@ FChaosClothAssetSimulationBendingConfigNode::FChaosClothAssetSimulationBendingCo
 void FChaosClothAssetSimulationBendingConfigNode::AddProperties(FPropertyHelper& PropertyHelper) const
 {
 	UE::Chaos::ClothAsset::FCollectionClothFacade ClothFacade(PropertyHelper.GetClothCollection());
-	if(!CanUseFabrics(ClothFacade))
+	if(!ClothFacade.IsValid())
 	{
-		if(SolverType == EChaosClothAssetConstraintSolverType::XPBD)
+		return;
+	}
+	
+	if(SolverType == EChaosClothAssetConstraintSolverType::XPBD)
+	{
+		if(DistributionType == EChaosClothAssetConstraintDistributionType::Isotropic)
 		{
-			if(DistributionType == EChaosClothAssetConstraintDistributionType::Isotropic)
+			if(ConstraintType == EChaosClothAssetBendingConstraintType::FacesSpring)
 			{
-				if(ConstraintType == EChaosClothAssetBendingConstraintType::FacesSpring)
-				{
-					PropertyHelper.SetPropertyWeighted(FName(TEXT("XPBDBendingSpringStiffness")), BendingStiffness,{
-						FName(TEXT("XPBDBendingElementStiffness")),
-						FName(TEXT("XPBDAnisoBendingStiffnessWarp")),
-						FName(TEXT("BendingSpringStiffness")),
-						FName(TEXT("BendingElementStiffness"))});
-					
-					PropertyHelper.SetPropertyWeighted(FName(TEXT("XPBDBendingSpringDamping")), BendingDamping,{
-						FName(TEXT("XPBDBendingElementDamping")),
-						FName(TEXT("XPBDAnisoBendingDamping"))});
-				}
-				else
-				{
-					PropertyHelper.SetPropertyWeighted(FName(TEXT("XPBDBendingElementStiffness")), BendingStiffness,{
-						FName(TEXT("XPBDBendingSpringStiffness")),
-						FName(TEXT("XPBDAnisoBendingStiffnessWarp")),
-						FName(TEXT("BendingSpringStiffness")),
-						FName(TEXT("BendingElementStiffness"))});
-					
-					PropertyHelper.SetPropertyWeighted(FName(TEXT("XPBDBendingElementDamping")), BendingDamping,{
-						FName(TEXT("XPBDBendingSpringDamping")),
-						FName(TEXT("XPBDAnisoBendingDamping"))});
-					
-					PropertyHelper.SetPropertyWeighted(FName(TEXT("XPBDBucklingStiffness")), BucklingStiffness,{
-						FName(TEXT("XPBDAnisoBucklingStiffnessWarp")),
-						FName(TEXT("BucklingStiffness"))});
-					
-					PropertyHelper.SetProperty(FName(TEXT("XPBDBucklingRatio")), BucklingRatio, {
-						FName(TEXT("XPBDAnisoBucklingRatio")),
-						FName(TEXT("BucklingRatios"))});
-				}
-				PropertyHelper.SetPropertyEnum(FName(TEXT("XPBDRestAngleType")), RestAngleType, {
-					FName(TEXT("XPBDAnisoRestAngleType")),
-					FName(TEXT("RestAngleType")) 
-				}, ECollectionPropertyFlags::None);
-				
-				PropertyHelper.SetPropertyWeighted(FName(TEXT("XPBDFlatnessRatio")), FlatnessRatio, {
-					FName(TEXT("XPBDAnisoFlatnessRatio")),
-					FName(TEXT("FlatnessRatio"))       
-				});
-				
-				PropertyHelper.SetPropertyWeighted(FName(TEXT("XPBDRestAngle")), RestAngle, {
-					FName(TEXT("XPBDAnisoRestAngle")),
-					FName(TEXT("RestAngle"))       
-				});
+				PropertyHelper.SetPropertyWeighted(FName(TEXT("XPBDBendingSpringStiffness")), BendingStiffness,{
+					FName(TEXT("XPBDBendingElementStiffness")),
+					FName(TEXT("XPBDAnisoBendingStiffnessWarp")),
+					FName(TEXT("BendingSpringStiffness")),
+					FName(TEXT("BendingElementStiffness"))});
+
+				PropertyHelper.SetPropertyWeighted(FName(TEXT("XPBDBendingSpringDamping")), BendingDamping,{
+					FName(TEXT("XPBDBendingElementDamping")),
+					FName(TEXT("XPBDAnisoBendingDamping"))});
 			}
 			else
 			{
-				PropertyHelper.SetPropertyWeighted(FName(TEXT("XPBDAnisoBendingStiffnessWarp")), BendingStiffnessWarp, {
-					FName(TEXT("BendingSpringStiffness")),     
-					FName(TEXT("BendingElementStiffness")),     
-					FName(TEXT("XPBDBendingSpringStiffness")),  
-					FName(TEXT("XPBDBendingElementStiffness"))});
-				
-				PropertyHelper.SetPropertyWeighted(FName(TEXT("XPBDAnisoBendingStiffnessWeft")), BendingStiffnessWeft, {});
-				PropertyHelper.SetPropertyWeighted(FName(TEXT("XPBDAnisoBendingStiffnessBias")), BendingStiffnessBias, {});
-
-				PropertyHelper.SetPropertyWeighted(FName(TEXT("XPBDAnisoBucklingStiffnessWarp")), BucklingStiffnessWarp, {
-					FName(TEXT("BucklingStiffness")),    
-					FName(TEXT("XPBDBucklingStiffness"))});
-				
-				PropertyHelper.SetPropertyWeighted(FName(TEXT("XPBDAnisoBucklingStiffnessWeft")), BucklingStiffnessWeft, {});
-				PropertyHelper.SetPropertyWeighted(FName(TEXT("XPBDAnisoBucklingStiffnessBias")), BucklingStiffnessBias, {});
+				PropertyHelper.SetPropertyWeighted(FName(TEXT("XPBDBendingElementStiffness")), BendingStiffness,{
+					FName(TEXT("XPBDBendingSpringStiffness")),
+					FName(TEXT("XPBDAnisoBendingStiffnessWarp")),
+					FName(TEXT("BendingSpringStiffness")),
+					FName(TEXT("BendingElementStiffness"))});
 
-				PropertyHelper.SetPropertyWeighted(FName(TEXT("XPBDAnisoBendingDamping")), BendingDamping, {
-					FName(TEXT("XPBDBendingSpringDamping")),  
-					FName(TEXT("XPBDBendingElementDamping"))});
+				PropertyHelper.SetPropertyWeighted(FName(TEXT("XPBDBendingElementDamping")), BendingDamping,{
+					FName(TEXT("XPBDBendingSpringDamping")),
+					FName(TEXT("XPBDAnisoBendingDamping"))});
 				
-				PropertyHelper.SetProperty(FName(TEXT("XPBDAnisoBucklingRatio")), BucklingRatio, {
-					FName(TEXT("BucklingRatio")),
-					FName(TEXT("XPBDBucklingRatio"))});
-
-				PropertyHelper.SetPropertyEnum(FName(TEXT("XPBDAnisoRestAngleType")), RestAngleType,{
-					FName(TEXT("XPBDRestAngleType")),
-					FName(TEXT("RestAngleType"))  
-				}, ECollectionPropertyFlags::None);
+				PropertyHelper.SetPropertyWeighted(FName(TEXT("XPBDBucklingStiffness")), BucklingStiffness,{
+					FName(TEXT("XPBDAnisoBucklingStiffnessWarp")),
+					FName(TEXT("BucklingStiffness"))});
 				
-				PropertyHelper.SetPropertyWeighted(FName(TEXT("XPBDAnisoFlatnessRatio")), FlatnessRatio, {
-					FName(TEXT("XPBDFlatnessRatio")),
-					FName(TEXT("FlatnessRatio"))       
-				});
-				PropertyHelper.SetPropertyWeighted(FName(TEXT("XPBDAnisoRestAngle")), RestAngle, {
-					FName(TEXT("XPBDRestAngle")),
-					FName(TEXT("RestAngle"))       
-				});
+				PropertyHelper.SetProperty(FName(TEXT("XPBDBucklingRatio")), BucklingRatio, {
+					FName(TEXT("XPBDAnisoBucklingRatio")),
+					FName(TEXT("BucklingRatios"))});
 			}
+			PropertyHelper.SetPropertyEnum(FName(TEXT("XPBDRestAngleType")), RestAngleType, {
+				FName(TEXT("XPBDAnisoRestAngleType")),
+				FName(TEXT("RestAngleType")) 
+			}, ECollectionPropertyFlags::None);
+			
+			PropertyHelper.SetPropertyWeighted(FName(TEXT("XPBDFlatnessRatio")), FlatnessRatio, {
+				FName(TEXT("XPBDAnisoFlatnessRatio")),
+				FName(TEXT("FlatnessRatio"))       
+			});
+			
+			PropertyHelper.SetPropertyWeighted(FName(TEXT("XPBDRestAngle")), RestAngle, {
+				FName(TEXT("XPBDAnisoRestAngle")),
+				FName(TEXT("RestAngle"))       
+			});
 		}
 		else
 		{
-			if(ConstraintType == EChaosClothAssetBendingConstraintType::FacesSpring)
+			SetFabricPropertyWeighted(FName(TEXT("XPBDAnisoBendingStiffnessWarp")), BendingStiffnessWarp, ClothFacade, PropertyHelper, [](
+				const UE::Chaos::ClothAsset::FCollectionClothFabricFacade& FabricFacade)-> float
+			 {
+				 return FabricFacade.GetXPBDAnisoBendingStiffness().Warp;
+			 }, {
+				FName(TEXT("BendingSpringStiffness")),     
+				FName(TEXT("BendingElementStiffness")),     
+				FName(TEXT("XPBDBendingSpringStiffness")),  
+				FName(TEXT("XPBDBendingElementStiffness"))});
+
+			SetFabricPropertyWeighted(FName(TEXT("XPBDAnisoBendingStiffnessWeft")), BendingStiffnessWeft, ClothFacade, PropertyHelper, [](
+				const UE::Chaos::ClothAsset::FCollectionClothFabricFacade& FabricFacade)-> float
+			 {
+				 return FabricFacade.GetXPBDAnisoBendingStiffness().Weft;
+			 }, {});
+			
+			SetFabricPropertyWeighted(FName(TEXT("XPBDAnisoBendingStiffnessBias")), BendingStiffnessBias,ClothFacade, PropertyHelper, [](
+				const UE::Chaos::ClothAsset::FCollectionClothFabricFacade& FabricFacade)-> float
+			 {
+				 return FabricFacade.GetXPBDAnisoBendingStiffness().Bias;
+			 }, {});
+			
+			SetFabricPropertyWeighted(FName(TEXT("XPBDAnisoBucklingStiffnessWarp")), BucklingStiffnessWarp,ClothFacade, PropertyHelper, [](
+				const UE::Chaos::ClothAsset::FCollectionClothFabricFacade& FabricFacade)-> float
+			 {
+				 return FabricFacade.GetXPBDAnisoBucklingStiffness().Warp;
+			 }, {
+				FName(TEXT("BucklingStiffness")),    
+				FName(TEXT("XPBDBucklingStiffness"))});
+			
+			SetFabricPropertyWeighted(FName(TEXT("XPBDAnisoBucklingStiffnessWeft")), BucklingStiffnessWeft, ClothFacade, PropertyHelper, [](
+				const UE::Chaos::ClothAsset::FCollectionClothFabricFacade& FabricFacade)-> float
 			{
-				PropertyHelper.SetPropertyWeighted(FName(TEXT("BendingSpringStiffness")), BendingStiffness, {
-					FName(TEXT("BendingElementStiffness")),
-					FName(TEXT("XPBDBendingSpringStiffness")),    
-					FName(TEXT("XPBDBendingElementStiffness")),   
-					FName(TEXT("XPBDAnisoBendingStiffnessWarp"))});
-			}
-			else
+				return FabricFacade.GetXPBDAnisoBucklingStiffness().Weft;
+			},{});
+			
+			SetFabricPropertyWeighted(FName(TEXT("XPBDAnisoBucklingStiffnessBias")), BucklingStiffnessBias, ClothFacade, PropertyHelper, [](
+				const UE::Chaos::ClothAsset::FCollectionClothFabricFacade& FabricFacade)-> float
 			{
-				PropertyHelper.SetPropertyWeighted(FName(TEXT("BendingElementStiffness")), BendingStiffness, {
-					FName(TEXT("BendingSpringStiffness")),        // Existing properties to warn against
-					FName(TEXT("XPBDBendingSpringStiffness")),
-					FName(TEXT("XPBDBendingElementStiffness")),   
-					FName(TEXT("XPBDAnisoBendingStiffnessWarp"))});
-				
-				PropertyHelper.SetPropertyWeighted(FName(TEXT("BucklingStiffness")), BucklingStiffness, {
-					FName(TEXT("XPBDBucklingStiffness")),
-					FName(TEXT("XPBDAnisoBucklingStiffnessWarp"))});
-				
-				PropertyHelper.SetProperty(FName(TEXT("BucklingRatios")), BucklingRatio, {
-					FName(TEXT("XPBDBucklingRatio")),
-					FName(TEXT("XPBDAnisoBucklingRatio"))});
-			}
+				return FabricFacade.GetXPBDAnisoBucklingStiffness().Bias;
+			},{});
+
+			SetFabricPropertyWeighted(FName(TEXT("XPBDAnisoBendingDamping")), BendingAnisoDamping, ClothFacade, PropertyHelper, [](
+				const UE::Chaos::ClothAsset::FCollectionClothFabricFacade& FabricFacade)-> float
+			{
+				return FabricFacade.GetXPBDAnisoDamping();
+			},{
+				FName(TEXT("XPBDBendingSpringDamping")),  
+				FName(TEXT("XPBDBendingElementDamping"))});
+
+			PropertyHelper.SetProperty(FName(TEXT("XPBDAnisoBucklingRatio")), BucklingRatio, {
+				FName(TEXT("BucklingRatio")),
+				FName(TEXT("XPBDBucklingRatio"))});
+
+			PropertyHelper.SetPropertyEnum(FName(TEXT("XPBDAnisoRestAngleType")), RestAngleType,{
+					FName(TEXT("XPBDRestAngleType")),
+					FName(TEXT("RestAngleType"))  
+				}, ECollectionPropertyFlags::None);
 			
-			PropertyHelper.SetPropertyEnum(FName(TEXT("RestAngleType")), RestAngleType, {
-				FName(TEXT("XPBDAnisoRestAngleType")),
-				FName(TEXT("XPBDRestAngleType")) 
-			}, ECollectionPropertyFlags::None);
-			PropertyHelper.SetPropertyWeighted(FName(TEXT("FlatnessRatio")), FlatnessRatio, {
-					FName(TEXT("XPBDAnisoFlatnessRatio")),
-					FName(TEXT("XPBDFlatnessRatio"))       
-				});
-			PropertyHelper.SetPropertyWeighted(FName(TEXT("RestAngle")), RestAngle, {
-				FName(TEXT("XPBDAnisoRestAngle")), 
-				FName(TEXT("XPBDRestAngle"))       
+			PropertyHelper.SetPropertyWeighted(FName(TEXT("XPBDAnisoFlatnessRatio")), FlatnessRatio, {
+				FName(TEXT("XPBDFlatnessRatio")),
+				FName(TEXT("FlatnessRatio"))       
+			});
+			
+			PropertyHelper.SetPropertyWeighted(FName(TEXT("XPBDAnisoRestAngle")), RestAngle, {
+				FName(TEXT("XPBDRestAngle")),
+				FName(TEXT("RestAngle"))       
 			});
 		}
 	}
 	else
 	{
-		SetFabricPropertyWeighted(FName(TEXT("XPBDAnisoBendingStiffnessWarp")), BendingStiffnessWarp, ClothFacade, PropertyHelper, [](
-			const UE::Chaos::ClothAsset::FCollectionClothFabricFacade& FabricFacade)-> float
-		 {
-			 return FabricFacade.GetXPBDAnisoBendingStiffness().Warp;
-		 }, {
-			FName(TEXT("BendingSpringStiffness")),     
-			FName(TEXT("BendingElementStiffness")),     
-			FName(TEXT("XPBDBendingSpringStiffness")),  
-			FName(TEXT("XPBDBendingElementStiffness"))});
-
-		SetFabricPropertyWeighted(FName(TEXT("XPBDAnisoBendingStiffnessWeft")), BendingStiffnessWeft, ClothFacade, PropertyHelper, [](
-			const UE::Chaos::ClothAsset::FCollectionClothFabricFacade& FabricFacade)-> float
-		 {
-			 return FabricFacade.GetXPBDAnisoBendingStiffness().Weft;
-		 }, {});
-		
-		SetFabricPropertyWeighted(FName(TEXT("XPBDAnisoBendingStiffnessBias")), BendingStiffnessBias,ClothFacade, PropertyHelper, [](
-			const UE::Chaos::ClothAsset::FCollectionClothFabricFacade& FabricFacade)-> float
-		 {
-			 return FabricFacade.GetXPBDAnisoBendingStiffness().Bias;
-		 }, {});
-		
-		SetFabricPropertyWeighted(FName(TEXT("XPBDAnisoBucklingStiffnessWarp")), BucklingStiffnessWarp,ClothFacade, PropertyHelper, [](
-			const UE::Chaos::ClothAsset::FCollectionClothFabricFacade& FabricFacade)-> float
-		 {
-			 return FabricFacade.GetXPBDAnisoBucklingStiffness().Warp;
-		 }, {
-			FName(TEXT("BucklingStiffness")),    
-			FName(TEXT("XPBDBucklingStiffness"))});
-		
-		SetFabricPropertyWeighted(FName(TEXT("XPBDAnisoBucklingStiffnessWeft")), BucklingStiffnessWeft, ClothFacade, PropertyHelper, [](
-			const UE::Chaos::ClothAsset::FCollectionClothFabricFacade& FabricFacade)-> float
-		{
-			return FabricFacade.GetXPBDAnisoBucklingStiffness().Weft;
-		},{});
-		
-		SetFabricPropertyWeighted(FName(TEXT("XPBDAnisoBucklingStiffnessBias")), BucklingStiffnessBias, ClothFacade, PropertyHelper, [](
-			const UE::Chaos::ClothAsset::FCollectionClothFabricFacade& FabricFacade)-> float
+		if(ConstraintType == EChaosClothAssetBendingConstraintType::FacesSpring)
 		{
-			return FabricFacade.GetXPBDAnisoBucklingStiffness().Bias;
-		},{});
-
-		SetFabricPropertyWeighted(FName(TEXT("XPBDAnisoBendingDamping")), BendingDamping, ClothFacade, PropertyHelper, [](
-			const UE::Chaos::ClothAsset::FCollectionClothFabricFacade& FabricFacade)-> float
+			PropertyHelper.SetPropertyWeighted(FName(TEXT("BendingSpringStiffness")), BendingStiffness, {
+				FName(TEXT("BendingElementStiffness")),
+				FName(TEXT("XPBDBendingSpringStiffness")),    
+				FName(TEXT("XPBDBendingElementStiffness")),   
+				FName(TEXT("XPBDAnisoBendingStiffnessWarp"))});
+		}
+		else
 		{
-			return FabricFacade.GetXPBDAnisoDamping();
-		},{
-			FName(TEXT("XPBDBendingSpringDamping")),  
-			FName(TEXT("XPBDBendingElementDamping"))});
-		
-		SetFabricProperty(FName(TEXT("XPBDAnisoBucklingRatio")), ClothFacade, PropertyHelper, [](const UE::Chaos::ClothAsset::FCollectionClothFabricFacade& FabricFacade)-> float
-		 {
-			return FabricFacade.GetXPBDAnisoBucklingRatio();
-		 }, {
-			FName(TEXT("BucklingRatio")),
-			FName(TEXT("XPBDBucklingRatio"))});
-
-		PropertyHelper.SetPropertyEnum(FName(TEXT("XPBDAnisoRestAngleType")), RestAngleType,{
-				FName(TEXT("XPBDRestAngleType")),
-				FName(TEXT("RestAngleType"))  
-			}, ECollectionPropertyFlags::None);
+			PropertyHelper.SetPropertyWeighted(FName(TEXT("BendingElementStiffness")), BendingStiffness, {
+				FName(TEXT("BendingSpringStiffness")),        // Existing properties to warn against
+				FName(TEXT("XPBDBendingSpringStiffness")),
+				FName(TEXT("XPBDBendingElementStiffness")),   
+				FName(TEXT("XPBDAnisoBendingStiffnessWarp"))});
+			
+			PropertyHelper.SetPropertyWeighted(FName(TEXT("BucklingStiffness")), BucklingStiffness, {
+				FName(TEXT("XPBDBucklingStiffness")),
+				FName(TEXT("XPBDAnisoBucklingStiffnessWarp"))});
+			
+			PropertyHelper.SetProperty(FName(TEXT("BucklingRatios")), BucklingRatio, {
+				FName(TEXT("XPBDBucklingRatio")),
+				FName(TEXT("XPBDAnisoBucklingRatio"))});
+		}
 		
-		PropertyHelper.SetPropertyWeighted(FName(TEXT("XPBDAnisoFlatnessRatio")), FlatnessRatio, {
-			FName(TEXT("XPBDFlatnessRatio")),
-			FName(TEXT("FlatnessRatio"))       
-		});
+		PropertyHelper.SetPropertyEnum(FName(TEXT("RestAngleType")), RestAngleType, {
+			FName(TEXT("XPBDAnisoRestAngleType")),
+			FName(TEXT("XPBDRestAngleType")) 
+		}, ECollectionPropertyFlags::None);
 		
-		PropertyHelper.SetPropertyWeighted(FName(TEXT("XPBDAnisoRestAngle")), RestAngle, {
-			FName(TEXT("XPBDRestAngle")),
-			FName(TEXT("RestAngle"))       
+		PropertyHelper.SetPropertyWeighted(FName(TEXT("FlatnessRatio")), FlatnessRatio, {
+				FName(TEXT("XPBDAnisoFlatnessRatio")),
+				FName(TEXT("XPBDFlatnessRatio"))       
+			});
+		PropertyHelper.SetPropertyWeighted(FName(TEXT("RestAngle")), RestAngle, {
+			FName(TEXT("XPBDAnisoRestAngle")), 
+			FName(TEXT("XPBDRestAngle"))       
 		});
 	}
 }
diff --git a/Engine/Plugins/Experimental/ChaosClothAssetEditor/Source/ChaosClothAssetDataflowNodes/Private/ChaosClothAsset/SimulationCollisionConfigNode.cpp b/Engine/Plugins/Experimental/ChaosClothAssetEditor/Source/ChaosClothAssetDataflowNodes/Private/ChaosClothAsset/SimulationCollisionConfigNode.cpp
--- a/Engine/Plugins/Experimental/ChaosClothAssetEditor/Source/ChaosClothAssetDataflowNodes/Private/ChaosClothAsset/SimulationCollisionConfigNode.cpp
+++ b/Engine/Plugins/Experimental/ChaosClothAssetEditor/Source/ChaosClothAssetDataflowNodes/Private/ChaosClothAsset/SimulationCollisionConfigNode.cpp
@@ -6,30 +6,16 @@
 #include UE_INLINE_GENERATED_CPP_BY_NAME(SimulationCollisionConfigNode)
 
 FChaosClothAssetSimulationCollisionConfigNode::FChaosClothAssetSimulationCollisionConfigNode(const Dataflow::FNodeParameters& InParam, FGuid InGuid)
-	: FChaosClothAssetSimulationFabricConfigNode(InParam, InGuid)
+	: FChaosClothAssetSimulationBaseConfigNode(InParam, InGuid)
 {
 	RegisterCollectionConnections();
 }
 
 void FChaosClothAssetSimulationCollisionConfigNode::AddProperties(FPropertyHelper& PropertyHelper) const
 {
-	UE::Chaos::ClothAsset::FCollectionClothFacade ClothFacade(PropertyHelper.GetClothCollection());
-	if(!CanUseFabrics(ClothFacade))
-	{
-		PropertyHelper.SetProperty(this, &FrictionCoefficient);
-		PropertyHelper.SetProperty(this, &CollisionThickness);
-	}
-	else
-	{
-		SetFabricProperty(FName(TEXT("FrictionCoefficient")), ClothFacade, PropertyHelper, [](const UE::Chaos::ClothAsset::FCollectionClothFabricFacade& FabricFacade)-> float
-		 {
-			 return FabricFacade.GetFrictionCoefficient();
-		 }, {});
-		SetFabricProperty(FName(TEXT("CollisionThickness")), ClothFacade, PropertyHelper, [](const UE::Chaos::ClothAsset::FCollectionClothFabricFacade& FabricFacade)-> float
-		 {
-			 return FabricFacade.GetCollisionThickness();
-		 }, {});
-	}
+	PropertyHelper.SetProperty(this, &FrictionCoefficient);
+	PropertyHelper.SetProperty(this, &CollisionThickness);
+	
 	PropertyHelper.SetPropertyBool(this, &bUseCCD);
 	PropertyHelper.SetProperty(this, &ProximityStiffness);
 }
diff --git a/Engine/Plugins/Experimental/ChaosClothAssetEditor/Source/ChaosClothAssetDataflowNodes/Private/ChaosClothAsset/SimulationFabricConfigNode.cpp b/Engine/Plugins/Experimental/ChaosClothAssetEditor/Source/ChaosClothAssetDataflowNodes/Private/ChaosClothAsset/SimulationFabricConfigNode.cpp
--- a/Engine/Plugins/Experimental/ChaosClothAssetEditor/Source/ChaosClothAssetDataflowNodes/Private/ChaosClothAsset/SimulationFabricConfigNode.cpp
+++ b/Engine/Plugins/Experimental/ChaosClothAssetEditor/Source/ChaosClothAssetDataflowNodes/Private/ChaosClothAsset/SimulationFabricConfigNode.cpp
@@ -114,10 +114,26 @@ void FChaosClothAssetSimulationFabricConfigNode::SetFabricPropertyWeighted(
 	FPropertyHelper& PropertyHelper, const TFunction<float(const UE::Chaos::ClothAsset::FCollectionClothFabricFacade&)>& FabricValueFunction,
 	const TArray<FName>& SimilarPropertyNames, ECollectionPropertyFlags PropertyFlags)
 {
-	const UE::Chaos::ClothAsset::FWeightedValueBounds WeightedValueBounds = UE::Chaos::ClothAsset::Private::BuildFabricWeightedValue(ClothFacade, 
+	if(PropertyValue.bCouldUseFabrics && (PropertyValue.bImportFabricBounds || PropertyValue.bBuildFabricMaps))
+	{
+		UE::Chaos::ClothAsset::FWeightedValueBounds WeightedValueBounds;
+		if(PropertyValue.bBuildFabricMaps)
+		{
+			WeightedValueBounds = UE::Chaos::ClothAsset::Private::BuildFabricWeightedValue(ClothFacade, 
 			 PropertyHelper.GetPropertyString(&PropertyValue.WeightMap), FabricValueFunction);
-
-	PropertyHelper.SetPropertyWeighted(PropertyName, WeightedValueBounds, PropertyValue, SimilarPropertyNames, PropertyFlags);
+		}
+		else 
+		{
+			WeightedValueBounds = UE::Chaos::ClothAsset::Private::BuildFabricValue(
+				ClothFacade, FabricValueFunction);
+		}
+		if(PropertyValue.bImportFabricBounds)
+		{
+			PropertyValue.Low = WeightedValueBounds.Low;
+			PropertyValue.High = WeightedValueBounds.High;
+		}
+	}
+	PropertyHelper.SetPropertyWeighted(PropertyName, PropertyValue, SimilarPropertyNames, PropertyFlags);
 }
 
 template void FChaosClothAssetSimulationFabricConfigNode::SetFabricPropertyWeighted<FChaosClothAssetWeightedValue>(const FName& PropertyName,
diff --git a/Engine/Plugins/Experimental/ChaosClothAssetEditor/Source/ChaosClothAssetDataflowNodes/Private/ChaosClothAsset/SimulationMassConfigNode.cpp b/Engine/Plugins/Experimental/ChaosClothAssetEditor/Source/ChaosClothAssetDataflowNodes/Private/ChaosClothAsset/SimulationMassConfigNode.cpp
--- a/Engine/Plugins/Experimental/ChaosClothAssetEditor/Source/ChaosClothAssetDataflowNodes/Private/ChaosClothAsset/SimulationMassConfigNode.cpp
+++ b/Engine/Plugins/Experimental/ChaosClothAssetEditor/Source/ChaosClothAssetDataflowNodes/Private/ChaosClothAsset/SimulationMassConfigNode.cpp
@@ -19,38 +19,32 @@ FChaosClothAssetSimulationMassConfigNode::FChaosClothAssetSimulationMassConfigNo
 void FChaosClothAssetSimulationMassConfigNode::AddProperties(FPropertyHelper& PropertyHelper) const
 {
 	UE::Chaos::ClothAsset::FCollectionClothFacade ClothFacade(PropertyHelper.GetClothCollection());
-	if(!CanUseFabrics(ClothFacade))
+	if(!ClothFacade.IsValid())
 	{
-		PropertyHelper.SetPropertyEnum(this, &MassMode, {}, ECollectionPropertyFlags::Intrinsic);
-		switch (MassMode)
-		{
-		default:
-		case EClothMassMode::UniformMass:
-			{
-				PropertyHelper.SetPropertyWeighted(FName(TEXT("MassValue")), UniformMassWeighted, {}, ECollectionPropertyFlags::Intrinsic);
-			}
-			break;
-		case EClothMassMode::TotalMass:
-			{
-				PropertyHelper.SetProperty(FName(TEXT("MassValue")), TotalMass, {}, ECollectionPropertyFlags::Intrinsic);
-			}
-			break;
-		case EClothMassMode::Density:
-			{
-				PropertyHelper.SetPropertyWeighted(FName(TEXT("MassValue")), DensityWeighted, {}, ECollectionPropertyFlags::Intrinsic);
-			}
-			break;
-		}
+		return;
 	}
-	else
+	PropertyHelper.SetPropertyEnum(this, &MassMode, {}, ECollectionPropertyFlags::Intrinsic);
+	switch (MassMode)
 	{
-		constexpr EClothMassMode MassModeProperty = EClothMassMode::Density;
-		PropertyHelper.SetPropertyEnum(FName(TEXT("MassMode")), MassModeProperty, {}, ECollectionPropertyFlags::Intrinsic);
-
-		SetFabricPropertyWeighted(FName(TEXT("MassValue")), DensityWeighted, ClothFacade, PropertyHelper, [](const UE::Chaos::ClothAsset::FCollectionClothFabricFacade& FabricFacade)-> float
-		 {
-			 return FabricFacade.GetDensityWeighted();
-		 }, {}, ECollectionPropertyFlags::Intrinsic);
+	default:
+	case EClothMassMode::UniformMass:
+		{
+			PropertyHelper.SetPropertyWeighted(FName(TEXT("MassValue")), UniformMassWeighted, {}, ECollectionPropertyFlags::Intrinsic);
+		}
+		break;
+	case EClothMassMode::TotalMass:
+		{
+			PropertyHelper.SetProperty(FName(TEXT("MassValue")), TotalMass, {}, ECollectionPropertyFlags::Intrinsic);
+		}
+		break;
+	case EClothMassMode::Density:
+		{
+			SetFabricPropertyWeighted(FName(TEXT("MassValue")), DensityWeighted, ClothFacade, PropertyHelper, [](const UE::Chaos::ClothAsset::FCollectionClothFabricFacade& FabricFacade)-> float
+			 {
+				 return FabricFacade.GetDensityWeighted();
+			 }, {}, ECollectionPropertyFlags::Intrinsic);
+		}
+		break;
 	}
 	PropertyHelper.SetProperty(this, &MinPerParticleMass, {}, ECollectionPropertyFlags::Intrinsic);
 }
diff --git a/Engine/Plugins/Experimental/ChaosClothAssetEditor/Source/ChaosClothAssetDataflowNodes/Private/ChaosClothAsset/SimulationStretchConfigNode.cpp b/Engine/Plugins/Experimental/ChaosClothAssetEditor/Source/ChaosClothAssetDataflowNodes/Private/ChaosClothAsset/SimulationStretchConfigNode.cpp
--- a/Engine/Plugins/Experimental/ChaosClothAssetEditor/Source/ChaosClothAssetDataflowNodes/Private/ChaosClothAsset/SimulationStretchConfigNode.cpp
+++ b/Engine/Plugins/Experimental/ChaosClothAssetEditor/Source/ChaosClothAssetDataflowNodes/Private/ChaosClothAsset/SimulationStretchConfigNode.cpp
@@ -14,6 +14,7 @@ FChaosClothAssetSimulationStretchConfigNode::FChaosClothAssetSimulationStretchCo
 	RegisterInputConnection(&StretchStiffnessWeft.WeightMap);
 	RegisterInputConnection(&StretchStiffnessBias.WeightMap);
 	RegisterInputConnection(&StretchDamping.WeightMap);
+	RegisterInputConnection(&StretchAnisoDamping.WeightMap);
 	RegisterInputConnection(&StretchWarpScale.WeightMap);
 	RegisterInputConnection(&StretchWeftScale.WeightMap);
 	RegisterInputConnection(&AreaStiffness.WeightMap);
@@ -22,99 +23,78 @@ FChaosClothAssetSimulationStretchConfigNode::FChaosClothAssetSimulationStretchCo
 void FChaosClothAssetSimulationStretchConfigNode::AddProperties(FPropertyHelper& PropertyHelper) const
 {
 	UE::Chaos::ClothAsset::FCollectionClothFacade ClothFacade(PropertyHelper.GetClothCollection());
-	if(!CanUseFabrics(ClothFacade))
+	if(!ClothFacade.IsValid())
 	{
-		if(SolverType == EChaosClothAssetConstraintSolverType::XPBD)
+		return;
+	}
+	
+	if(SolverType == EChaosClothAssetConstraintSolverType::XPBD)
+	{
+		if (DistributionType == EChaosClothAssetConstraintDistributionType::Anisotropic)
 		{
-			if (DistributionType == EChaosClothAssetConstraintDistributionType::Anisotropic)
-			{
-				PropertyHelper.SetPropertyBool(FName(TEXT("XPBDAnisoSpringUse3dRestLengths")), bStretchUse3dRestLengths, {
-					FName(TEXT("XPBDAnisoStretchUse3dRestLengths")) }, ECollectionPropertyFlags::None);  // Non animatable
+			PropertyHelper.SetPropertyBool(FName(TEXT("XPBDAnisoSpringUse3dRestLengths")), bStretchUse3dRestLengths, {
+				FName(TEXT("XPBDAnisoStretchUse3dRestLengths"))}, ECollectionPropertyFlags::None);  // Non animatable
 
-				PropertyHelper.SetPropertyWeighted(FName(TEXT("XPBDAnisoSpringStiffnessWarp")), StretchStiffnessWarp, {
-					FName(TEXT("EdgeSpringStiffness")),
-					FName(TEXT("XPBDEdgeSpringStiffness")),
-					FName(TEXT("XPBDAnisoStretchStiffnessWarp")) });
+			SetFabricPropertyWeighted(FName(TEXT("XPBDAnisoSpringStiffnessWarp")), StretchStiffnessWarp,ClothFacade, PropertyHelper, [](
+				const UE::Chaos::ClothAsset::FCollectionClothFabricFacade& FabricFacade)-> float
+			 {
+				 return FabricFacade.GetXPBDAnisoSpringStiffness().Warp;
+			 }, {
+				FName(TEXT("EdgeSpringStiffness")),
+				FName(TEXT("XPBDEdgeSpringStiffness")),
+				FName(TEXT("XPBDAnisoStretchStiffnessWarp"))});
 
-				PropertyHelper.SetPropertyWeighted(FName(TEXT("XPBDAnisoSpringStiffnessWeft")), StretchStiffnessWeft, {
-					FName(TEXT("XPBDAnisoStretchStiffnessWeft")) });
-				PropertyHelper.SetPropertyWeighted(FName(TEXT("XPBDAnisoSpringStiffnessBias")), StretchStiffnessBias, {
-					FName(TEXT("XPBDAnisoStretchStiffnessBias")) });
-				PropertyHelper.SetPropertyWeighted(FName(TEXT("XPBDAnisoSpringDamping")), StretchDamping, {
-					FName(TEXT("XPBDAnisoStretchDamping")),
-					FName(TEXT("XPBDEdgeSpringDamping")) });
+			SetFabricPropertyWeighted(FName(TEXT("XPBDAnisoSpringStiffnessWeft")), StretchStiffnessWeft, ClothFacade, PropertyHelper, [](
+				const UE::Chaos::ClothAsset::FCollectionClothFabricFacade& FabricFacade)-> float
+			 {
+				 return FabricFacade.GetXPBDAnisoSpringStiffness().Weft;
+			 }, { FName(TEXT("XPBDAnisoStretchStiffnessWeft")) });
+	
+			SetFabricPropertyWeighted(FName(TEXT("XPBDAnisoSpringStiffnessBias")), StretchStiffnessBias, ClothFacade, PropertyHelper, [](
+				const UE::Chaos::ClothAsset::FCollectionClothFabricFacade& FabricFacade)-> float
+			 {
+				 return FabricFacade.GetXPBDAnisoSpringStiffness().Bias;
+			 }, { FName(TEXT("XPBDAnisoStretchStiffnessBias")) });
 
-				PropertyHelper.SetPropertyWeighted(FName(TEXT("XPBDAnisoSpringWarpScale")), StretchWarpScale, {
-					FName(TEXT("XPBDAnisoStretchWarpScale")) });
-				PropertyHelper.SetPropertyWeighted(FName(TEXT("XPBDAnisoSpringWeftScale")), StretchWeftScale, {
-					FName(TEXT("XPBDAnisoStretchWeftScale")) });
-			}
-			else
-			{
-				PropertyHelper.SetPropertyWeighted(FName(TEXT("XPBDEdgeSpringStiffness")), StretchStiffness, {
-					FName(TEXT("EdgeSpringStiffness")),
-					FName(TEXT("XPBDAnisoStretchStiffnessWarp")),
-					FName(TEXT("XPBDAnisoSpringStiffnessWarp"))});
-				
-				PropertyHelper.SetPropertyWeighted(FName(TEXT("XPBDEdgeSpringDamping")), StretchDamping, {
-					FName(TEXT("XPBDAnisoStretchDamping")),
-					FName(TEXT("XPBDAnisoSpringDamping"))});
+			SetFabricPropertyWeighted(FName(TEXT("XPBDAnisoSpringDamping")), StretchAnisoDamping, ClothFacade, PropertyHelper, [](
+				const UE::Chaos::ClothAsset::FCollectionClothFabricFacade& FabricFacade)-> float
+			 {
+				 return FabricFacade.GetXPBDAnisoDamping();
+			 }, { FName(TEXT("XPBDEdgeSpringDamping")),
+				FName(TEXT("XPBDAnisoStretchDamping")) });
 
-				if(bAddAreaConstraint)
-				{
-					PropertyHelper.SetPropertyWeighted(FName(TEXT("XPBDAreaSpringStiffness")), AreaStiffness,{
-						FName(TEXT("AreaSpringStiffness"))});
-				}
-			}
+			PropertyHelper.SetPropertyWeighted(FName(TEXT("XPBDAnisoSpringWarpScale")), StretchWarpScale, { FName(TEXT("XPBDAnisoStretchWarpScale")) });
+			PropertyHelper.SetPropertyWeighted(FName(TEXT("XPBDAnisoSpringWeftScale")), StretchWeftScale, { FName(TEXT("XPBDAnisoStretchWeftScale")) });
 		}
 		else
 		{
-			PropertyHelper.SetPropertyWeighted(FName(TEXT("EdgeSpringStiffness")), StretchStiffness, {
-				FName(TEXT("XPBDEdgeSpringStiffness")), 
+			PropertyHelper.SetPropertyWeighted(FName(TEXT("XPBDEdgeSpringStiffness")), StretchStiffness, {
+				FName(TEXT("EdgeSpringStiffness")),
 				FName(TEXT("XPBDAnisoStretchStiffnessWarp")),
 				FName(TEXT("XPBDAnisoSpringStiffnessWarp"))});
-			
+
+			PropertyHelper.SetPropertyWeighted(FName(TEXT("XPBDEdgeSpringDamping")), StretchDamping, {
+				FName(TEXT("XPBDAnisoStretchDamping")),
+				FName(TEXT("XPBDAnisoSpringDamping"))});
+
 			if(bAddAreaConstraint)
 			{
-				PropertyHelper.SetPropertyWeighted(FName(TEXT("AreaSpringStiffness")), AreaStiffness,{
-					FName(TEXT("XPBDAreaSpringStiffness"))});
+				PropertyHelper.SetPropertyWeighted(FName(TEXT("XPBDAreaSpringStiffness")), AreaStiffness,{
+					FName(TEXT("AreaSpringStiffness"))});
 			}
 		}
 	}
 	else
 	{
-		PropertyHelper.SetPropertyBool(FName(TEXT("XPBDAnisoSpringUse3dRestLengths")), bStretchUse3dRestLengths, {
-						FName(TEXT("XPBDAnisoStretchUse3dRestLengths"))}, ECollectionPropertyFlags::None);  // Non animatable
-
-		SetFabricPropertyWeighted(FName(TEXT("XPBDAnisoSpringStiffnessWarp")), StretchStiffnessWarp,ClothFacade, PropertyHelper, [](
-			const UE::Chaos::ClothAsset::FCollectionClothFabricFacade& FabricFacade)-> float
-		 {
-			 return FabricFacade.GetXPBDAnisoSpringStiffness().Warp;
-		 }, {
-			FName(TEXT("EdgeSpringStiffness")),
-			FName(TEXT("XPBDEdgeSpringStiffness")),
-			FName(TEXT("XPBDAnisoStretchStiffnessWarp"))});
-
-		SetFabricPropertyWeighted(FName(TEXT("XPBDAnisoSpringStiffnessWeft")), StretchStiffnessWeft, ClothFacade, PropertyHelper, [](
-			const UE::Chaos::ClothAsset::FCollectionClothFabricFacade& FabricFacade)-> float
-		 {
-			 return FabricFacade.GetXPBDAnisoSpringStiffness().Weft;
-		 }, { FName(TEXT("XPBDAnisoStretchStiffnessWeft")) });
+		PropertyHelper.SetPropertyWeighted(FName(TEXT("EdgeSpringStiffness")), StretchStiffness, {
+			FName(TEXT("XPBDEdgeSpringStiffness")), 
+			FName(TEXT("XPBDAnisoStretchStiffnessWarp")),
+			FName(TEXT("XPBDAnisoSpringStiffnessWarp"))});
 		
-		SetFabricPropertyWeighted(FName(TEXT("XPBDAnisoSpringStiffnessBias")), StretchStiffnessBias, ClothFacade, PropertyHelper, [](
-			const UE::Chaos::ClothAsset::FCollectionClothFabricFacade& FabricFacade)-> float
-		 {
-			 return FabricFacade.GetXPBDAnisoSpringStiffness().Bias;
-		 }, { FName(TEXT("XPBDAnisoStretchStiffnessBias")) });
-
-		SetFabricPropertyWeighted(FName(TEXT("XPBDAnisoSpringDamping")), StretchDamping, ClothFacade, PropertyHelper, [](
-			const UE::Chaos::ClothAsset::FCollectionClothFabricFacade& FabricFacade)-> float
-		 {
-			 return FabricFacade.GetXPBDAnisoDamping();
-		 }, { FName(TEXT("XPBDEdgeSpringDamping")),
-			FName(TEXT("XPBDAnisoStretchDamping")) });
-
-		PropertyHelper.SetPropertyWeighted(FName(TEXT("XPBDAnisoSpringWarpScale")), StretchWarpScale, { FName(TEXT("XPBDAnisoStretchWarpScale")) });
-		PropertyHelper.SetPropertyWeighted(FName(TEXT("XPBDAnisoSpringWeftScale")), StretchWeftScale, { FName(TEXT("XPBDAnisoStretchWeftScale")) });
+		if(bAddAreaConstraint)
+		{
+			PropertyHelper.SetPropertyWeighted(FName(TEXT("AreaSpringStiffness")), AreaStiffness,{
+				FName(TEXT("XPBDAreaSpringStiffness"))});
+		}
 	}
 }
diff --git a/Engine/Plugins/Experimental/ChaosClothAssetEditor/Source/ChaosClothAssetDataflowNodes/Private/ChaosClothAsset/WeightedValueCustomization.cpp b/Engine/Plugins/Experimental/ChaosClothAssetEditor/Source/ChaosClothAssetDataflowNodes/Private/ChaosClothAsset/WeightedValueCustomization.cpp
--- a/Engine/Plugins/Experimental/ChaosClothAssetEditor/Source/ChaosClothAssetDataflowNodes/Private/ChaosClothAsset/WeightedValueCustomization.cpp
+++ b/Engine/Plugins/Experimental/ChaosClothAssetEditor/Source/ChaosClothAssetDataflowNodes/Private/ChaosClothAsset/WeightedValueCustomization.cpp
@@ -1,6 +1,7 @@
 // Copyright Epic Games, Inc. All Rights Reserved.
 
 #include "ChaosClothAsset/WeightedValueCustomization.h"
+#include "ChaosClothAsset/ClothAssetEditorStyle.h"
 #include "ChaosClothAsset/WeightedValue.h"
 #include "Widgets/Input/SCheckBox.h"
 #include "Widgets/Input/SNumericEntryBox.h"
@@ -16,19 +17,68 @@ namespace UE::Chaos::ClothAsset
 	namespace Private
 	{
 		static const FString OverridePrefix = TEXT("_Override");
+		
+		static const FString ImportFabricBounds = TEXT("ImportFabricBounds");
+		static const FString BuildFabricMaps = TEXT("BuildFabricMaps");
+		static const FString CouldUseFabrics = TEXT("CouldUseFabrics");
 
 		static bool IsOverrideProperty(const TSharedPtr<IPropertyHandle>& Property)
 		{
 			const FStringView PropertyPath = Property ? Property->GetPropertyPath() : FStringView();
 			return PropertyPath.EndsWith(OverridePrefix, ESearchCase::CaseSensitive);
 		}
-
 		static bool IsOverridePropertyOf(const TSharedPtr<IPropertyHandle>& OverrideProperty, const TSharedPtr<IPropertyHandle>& Property)
 		{
 			const FStringView OverridePropertyPath = OverrideProperty ? OverrideProperty->GetPropertyPath() : FStringView();
 			const FStringView PropertyPath = Property ? Property->GetPropertyPath() : FStringView();
 			return OverridePropertyPath == FString(PropertyPath) + OverridePrefix;
 		}
+		static bool IsImportFabricBoundsProperty(const TSharedPtr<IPropertyHandle>& Property)
+		{
+			const FStringView PropertyPath = Property ? Property->GetPropertyPath() : FStringView();
+			return PropertyPath.EndsWith(ImportFabricBounds, ESearchCase::CaseSensitive);
+		}
+		static bool IsBuildFabricMapsProperty(const TSharedPtr<IPropertyHandle>& Property)
+		{
+			const FStringView PropertyPath = Property ? Property->GetPropertyPath() : FStringView();
+			return PropertyPath.EndsWith(BuildFabricMaps, ESearchCase::CaseSensitive);
+		}
+		static bool CouldUseFabricsProperty(const TSharedPtr<IPropertyHandle>& Property)
+		{
+			const FStringView PropertyPath = Property ? Property->GetPropertyPath() : FStringView();
+			return PropertyPath.EndsWith(CouldUseFabrics, ESearchCase::CaseSensitive);
+		}
+		
+		static void AddToggledCheckBox(const TSharedRef<IPropertyHandle>& PropertyHandle, const TSharedPtr<SHorizontalBox>& HorizontalBox, const FSlateBrush* SlateBrush)
+		{
+			TWeakPtr<IPropertyHandle> WeakHandle = PropertyHandle;
+		
+			HorizontalBox->AddSlot()
+				.VAlign(VAlign_Center)
+				.HAlign(HAlign_Left)
+				.AutoWidth()
+				[
+					SNew(SCheckBox)
+					.Style(&FAppStyle::Get().GetWidgetStyle<FCheckBoxStyle>("ToggleButtonCheckBox"))
+					.Type(ESlateCheckBoxType::ToggleButton)
+					.ToolTipText(WeakHandle.Pin()->GetToolTipText())
+					.IsChecked_Lambda([WeakHandle]()->ECheckBoxState
+						{
+							bool bValue = false;
+							WeakHandle.Pin()->GetValue(bValue);
+							return bValue ? ECheckBoxState::Checked : ECheckBoxState::Unchecked;
+						})
+						.OnCheckStateChanged_Lambda([WeakHandle](ECheckBoxState CheckBoxState)
+						{
+							WeakHandle.Pin()->SetValue(CheckBoxState == ECheckBoxState::Checked, EPropertyValueSetFlags::DefaultFlags);
+						})
+					[
+					SNew(SImage)
+						.ColorAndOpacity(FSlateColor::UseForeground())
+						.Image(SlateBrush)
+					]
+				];
+		}
 	}
 
 	TSharedRef<IPropertyTypeCustomization> FWeightedValueCustomization::MakeInstance()
@@ -40,24 +90,57 @@ namespace UE::Chaos::ClothAsset
 	
 	FWeightedValueCustomization::~FWeightedValueCustomization() = default;
 
+
 	void FWeightedValueCustomization::MakeHeaderRow(TSharedRef<class IPropertyHandle>& StructPropertyHandle, FDetailWidgetRow& Row)
 	{
-		TWeakPtr<IPropertyHandle> StructWeakHandlePtr = StructPropertyHandle;
+		const TWeakPtr<IPropertyHandle> StructWeakHandlePtr = StructPropertyHandle;
 
-		TSharedPtr<SHorizontalBox> HorizontalBox;
+		TSharedPtr<SHorizontalBox> ValueHorizontalBox;
+		TSharedPtr<SHorizontalBox> NameHorizontalBox;
 
 		Row.NameContent()
 			[
-				StructPropertyHandle->CreatePropertyNameWidget()
+				SAssignNew(NameHorizontalBox, SHorizontalBox)
+				.IsEnabled(this, &FMathStructCustomization::IsValueEnabled, StructWeakHandlePtr)
 			]
 		.ValueContent()
 			// Make enough space for each child handle
 			.MinDesiredWidth(125.f * SortedChildHandles.Num())
 			.MaxDesiredWidth(125.f * SortedChildHandles.Num())
 			[
-				SAssignNew(HorizontalBox, SHorizontalBox)
+				SAssignNew(ValueHorizontalBox, SHorizontalBox)
 				.IsEnabled(this, &FMathStructCustomization::IsValueEnabled, StructWeakHandlePtr)
 			];
+		
+		for (int32 ChildIndex = 0; ChildIndex < SortedChildHandles.Num(); ++ChildIndex)
+		{
+			TSharedRef<IPropertyHandle> ChildHandle = SortedChildHandles[ChildIndex];
+			
+			if (Private::CouldUseFabricsProperty(ChildHandle))
+			{
+				bool bValue = false;
+				ChildHandle->GetValue(bValue);
+				if(!bValue)
+				{
+					break;
+				}
+			}
+			if (Private::IsImportFabricBoundsProperty(ChildHandle))
+			{
+				Private::AddToggledCheckBox(ChildHandle, NameHorizontalBox, FAppStyle::Get().GetBrush("Icons.Import"));
+			}
+			else if (Private::IsBuildFabricMapsProperty(ChildHandle))
+			{
+				Private::AddToggledCheckBox(ChildHandle, NameHorizontalBox, UE::Chaos::ClothAsset::FClothAssetEditorStyle::Get().GetBrush("ClassIcon.ChaosClothPreset"));
+			}
+		}
+		
+		NameHorizontalBox->AddSlot().VAlign(VAlign_Center)
+				.HAlign(HAlign_Right)
+				.AutoWidth()
+				[
+					StructPropertyHandle->CreatePropertyNameWidget()
+				];
 
 		for (int32 ChildIndex = 0; ChildIndex < SortedChildHandles.Num(); ++ChildIndex)
 		{
@@ -84,28 +167,30 @@ namespace UE::Chaos::ClothAsset
 			const bool bLastChild = SortedChildHandles.Num() - 1 == ChildIndex;
 
 			TSharedRef<SWidget> ChildWidget = MakeChildWidget(StructPropertyHandle, ChildHandle);
-			
-			if (ChildHandle->GetPropertyClass() == FBoolProperty::StaticClass())
-			{
-				HorizontalBox->AddSlot()
-					.Padding(FMargin(0.f, 2.f, bLastChild ? 0.f : 3.f, 2.f))
-					.AutoWidth()  // keep the check box slots small
-					[
-						ChildWidget
-					];
-			}
-			else
+			if(ChildWidget != SNullWidget::NullWidget)
 			{
-				if (ChildHandle->GetPropertyClass() == FFloatProperty::StaticClass())
+				if (ChildHandle->GetPropertyClass() == FBoolProperty::StaticClass())
 				{
-					NumericEntryBoxWidgetList.Add(ChildWidget);
+					ValueHorizontalBox->AddSlot()
+						.Padding(FMargin(0.f, 2.f, bLastChild ? 0.f : 3.f, 2.f))
+						.AutoWidth()  // keep the check box slots small
+						[
+							ChildWidget
+						];
 				}
+				else
+				{
+					if (ChildHandle->GetPropertyClass() == FFloatProperty::StaticClass())
+					{
+						NumericEntryBoxWidgetList.Add(ChildWidget);
+					}
 
-				HorizontalBox->AddSlot()
-					.Padding(FMargin(0.f, 2.f, bLastChild ? 0.f : 3.f, 2.f))
-					[
-						ChildWidget
-					];
+					ValueHorizontalBox->AddSlot()
+						.Padding(FMargin(0.f, 2.f, bLastChild ? 0.f : 3.f, 2.f))
+						[
+							ChildWidget
+						];
+				}
 			}
 		}
 	}
@@ -123,24 +208,23 @@ namespace UE::Chaos::ClothAsset
 		if (PropertyClass == FBoolProperty::StaticClass())
 		{
 			TWeakPtr<IPropertyHandle> WeakHandlePtr = PropertyHandle;
-			return
-				SNew(SCheckBox)
-				.ToolTipText(
-					LOCTEXT(
-						"IsAnimatable",
-						"Whether the property can ever be updated/animated in real time.\n"
-						"This could make the simulation takes more CPU time, even more so if the weight maps needs updating."))
-				.Type(ESlateCheckBoxType::CheckBox)
-				.IsChecked_Lambda([WeakHandlePtr]()->ECheckBoxState
-					{
-						bool bValue = false;
-						WeakHandlePtr.Pin()->GetValue(bValue);
-						return bValue ? ECheckBoxState::Checked : ECheckBoxState::Unchecked;
-					})
-				.OnCheckStateChanged_Lambda([WeakHandlePtr](ECheckBoxState CheckBoxState)
-					{
-						WeakHandlePtr.Pin()->SetValue(CheckBoxState == ECheckBoxState::Checked, EPropertyValueSetFlags::DefaultFlags);
-					});
+			if (!Private::IsImportFabricBoundsProperty(PropertyHandle) && !Private::IsBuildFabricMapsProperty(PropertyHandle) && !Private::CouldUseFabricsProperty(PropertyHandle))
+			{
+				return
+					SNew(SCheckBox)
+					.ToolTipText(PropertyHandle->GetToolTipText())
+					.Type(ESlateCheckBoxType::CheckBox)
+					.IsChecked_Lambda([WeakHandlePtr]()->ECheckBoxState
+						{
+							bool bValue = false;
+							WeakHandlePtr.Pin()->GetValue(bValue);
+							return bValue ? ECheckBoxState::Checked : ECheckBoxState::Unchecked;
+						})
+					.OnCheckStateChanged_Lambda([WeakHandlePtr](ECheckBoxState CheckBoxState)
+						{
+							WeakHandlePtr.Pin()->SetValue(CheckBoxState == ECheckBoxState::Checked, EPropertyValueSetFlags::DefaultFlags);
+						});
+			}
 		}
 		if (PropertyClass == FStrProperty::StaticClass())
 		{
@@ -199,8 +283,6 @@ namespace UE::Chaos::ClothAsset
 					})
 				.Font(IPropertyTypeCustomizationUtils::GetRegularFont());
 		}
-
-		checkf(false, TEXT("Unsupported property class for the Weighted Values customization."));
 		return SNullWidget::NullWidget;
 	}
 
diff --git a/Engine/Plugins/Experimental/ChaosClothAssetEditor/Source/ChaosClothAssetDataflowNodes/Public/ChaosClothAsset/SimulationBaseConfigNode.h b/Engine/Plugins/Experimental/ChaosClothAssetEditor/Source/ChaosClothAssetDataflowNodes/Public/ChaosClothAsset/SimulationBaseConfigNode.h
--- a/Engine/Plugins/Experimental/ChaosClothAssetEditor/Source/ChaosClothAssetDataflowNodes/Public/ChaosClothAsset/SimulationBaseConfigNode.h
+++ b/Engine/Plugins/Experimental/ChaosClothAssetEditor/Source/ChaosClothAssetDataflowNodes/Public/ChaosClothAsset/SimulationBaseConfigNode.h
@@ -63,9 +63,6 @@ protected:
 		int32 SetPropertyWeighted(const FName& PropertyName, const FChaosClothAssetWeightedValue& PropertyValue, const TArray<FName>& SimilarPropertyNames = {}, ECollectionPropertyFlags PropertyFlags = ECollectionPropertyFlags::None);
 		int32 SetPropertyWeighted(const FName& PropertyName, const FChaosClothAssetWeightedValueNonAnimatable& PropertyValue, const TArray<FName>& SimilarPropertyNames = {}, ECollectionPropertyFlags PropertyFlags = ECollectionPropertyFlags::None);
 		int32 SetPropertyWeighted(const FName& PropertyName, const FChaosClothAssetWeightedValueNonAnimatableNoLowHighRange& PropertyValue, const TArray<FName>& SimilarPropertyNames = {}, ECollectionPropertyFlags PropertyFlags = ECollectionPropertyFlags::None);
-
-		template<typename WeightedValueType>
-		inline int32 SetPropertyWeighted(const FName& PropertyName, const UE::Chaos::ClothAsset::FWeightedValueBounds& PropertyBounds, const WeightedValueType& PropertyValue, const TArray<FName>& SimilarPropertyNames = {}, ECollectionPropertyFlags PropertyFlags = ECollectionPropertyFlags::None);
 		
 		template<typename T, typename WeightedValueType, TEMPLATE_REQUIRES(TIsDerivedFrom<T, FChaosClothAssetSimulationBaseConfigNode>::Value)>
 		inline int32 SetPropertyWeighted(const T* ConfigStruct, const WeightedValueType* PropertyValue, const TArray<FName>& SimilarPropertyNames = {}, ECollectionPropertyFlags PropertyFlags = ECollectionPropertyFlags::None);
@@ -224,13 +221,6 @@ inline int32 FChaosClothAssetSimulationBaseConfigNode::FPropertyHelper::SetPrope
 	return SetPropertyWeighted(PropertyName, *PropertyValue, SimilarPropertyNames, PropertyFlags);
 }
 
-template<typename WeightedValueType>
-inline int32 FChaosClothAssetSimulationBaseConfigNode::FPropertyHelper::SetPropertyWeighted(const FName& PropertyName, const UE::Chaos::ClothAsset::FWeightedValueBounds& PropertyBounds, const WeightedValueType& PropertyValue, const TArray<FName>& SimilarPropertyNames, ECollectionPropertyFlags PropertyFlags)
-{
-	return SetPropertyWeighted(PropertyName, PropertyValue.bIsAnimatable, PropertyBounds.Low,
-		PropertyBounds.High, PropertyValue.WeightMap, PropertyValue.WeightMap_Override, SimilarPropertyNames, PropertyFlags);
-}
-
 template<typename T, typename ConnectableStringValueType, typename TEnableIf<TIsDerivedFrom<T, FChaosClothAssetSimulationBaseConfigNode>::Value, int>::type>
 inline int32 FChaosClothAssetSimulationBaseConfigNode::FPropertyHelper::SetPropertyString(
 	const T* ConfigStruct, 
diff --git a/Engine/Plugins/Experimental/ChaosClothAssetEditor/Source/ChaosClothAssetDataflowNodes/Public/ChaosClothAsset/SimulationBendingConfigNode.h b/Engine/Plugins/Experimental/ChaosClothAssetEditor/Source/ChaosClothAssetDataflowNodes/Public/ChaosClothAsset/SimulationBendingConfigNode.h
--- a/Engine/Plugins/Experimental/ChaosClothAssetEditor/Source/ChaosClothAssetDataflowNodes/Public/ChaosClothAsset/SimulationBendingConfigNode.h
+++ b/Engine/Plugins/Experimental/ChaosClothAssetEditor/Source/ChaosClothAssetDataflowNodes/Public/ChaosClothAsset/SimulationBendingConfigNode.h
@@ -44,25 +44,16 @@ public:
 	FChaosClothAssetWeightedValueNonAnimatable RestAngle = { 0.f, 0.f, TEXT("RestAngle") };
 
 	/** Constraint solver type */
-	UPROPERTY(EditAnywhere, Category = "Bending Types", Meta = (EditCondition = "bImportProperties == false", EditConditionHides))
+	UPROPERTY(EditAnywhere, Category = "Bending Types")
 	EChaosClothAssetConstraintSolverType SolverType = EChaosClothAssetConstraintSolverType::PBD;
 	
 	/** Constraint distribution type */
-    UPROPERTY(EditAnywhere, Category = "Bending Types", Meta = (EditCondition = "bImportProperties == false && SolverType == EChaosClothAssetConstraintSolverType::XPBD", EditConditionHides))
+    UPROPERTY(EditAnywhere, Category = "Bending Types", Meta = (EditCondition = "SolverType == EChaosClothAssetConstraintSolverType::XPBD", EditConditionHides))
     EChaosClothAssetConstraintDistributionType DistributionType = EChaosClothAssetConstraintDistributionType::Isotropic;
 	
 	/** Constraint method type */
-	UPROPERTY(EditAnywhere, Category = "Bending Types", Meta = (EditCondition = "bImportProperties == false && DistributionType == EChaosClothAssetConstraintDistributionType::Anisotropic", EditConditionHides))
+	UPROPERTY(EditAnywhere, Category = "Bending Types", Meta = (EditCondition = "DistributionType == EChaosClothAssetConstraintDistributionType::Isotropic", EditConditionHides))
 	EChaosClothAssetBendingConstraintType ConstraintType = EChaosClothAssetBendingConstraintType::HingeAngles;
-
-	/**
-	 * The Stiffness of the bending constraints. Increase the iteration count for stiffer materials.
-	 * If a valid weight map is found with the given Weight Map name, then both Low and High values
-	 * are interpolated with the per particle weight to make the final value used for the simulation.
-	 * Otherwise all particles are considered to have a zero weight, and only the Low value is meaningful.
-	 */
-	UPROPERTY(EditAnywhere, Category = "Bending Properties", Meta = (UIMin = "0", UIMax = "1", ClampMin = "0", ClampMax = "1"), Meta = (EditCondition = "bImportProperties == false && ((DistributionType == EChaosClothAssetConstraintDistributionType::Isotropic && SolverType == EChaosClothAssetConstraintSolverType::XPBD) || SolverType == EChaosClothAssetConstraintSolverType::PBD)", EditConditionHides))
-	FChaosClothAssetWeightedValue BendingStiffness = { true, 1.f, 1.f, TEXT("BendingStiffness") };
 	
 	/**
 	 * The stiffness of the bending constraints in the warp (vertical) direction.
@@ -70,8 +61,8 @@ public:
 	 * are interpolated with the per particle weight to make the final value used for the simulation.
 	 * Otherwise all particles are considered to have a zero weight, and only the Low value is meaningful.
 	 */
-	UPROPERTY(EditAnywhere, Category = "Bending Properties", Meta = (UIMin = "0", UIMax = "10000", ClampMin = "0", ClampMax = "10000000", EditCondition = "DistributionType == EChaosClothAssetConstraintDistributionType::Anisotropic && bImportProperties == false && SolverType == EChaosClothAssetConstraintSolverType::XPBD", EditConditionHides))
-	FChaosClothAssetWeightedValue BendingStiffnessWarp = { true, 100.f, 100.f, TEXT("BendingStiffnessWarp") };
+	UPROPERTY(EditAnywhere, Category = "Bending Properties", Meta = (UIMin = "0", UIMax = "10000", ClampMin = "0", ClampMax = "10000000", EditCondition = "DistributionType == EChaosClothAssetConstraintDistributionType::Anisotropic && SolverType == EChaosClothAssetConstraintSolverType::XPBD", EditConditionHides))
+	FChaosClothAssetWeightedValue BendingStiffnessWarp = { true,  100.f, 100.f, TEXT("BendingStiffnessWarp"), true };
 
 	/**
 	 * The stiffness of the bending constraints in the weft (horizontal) direction.
@@ -79,8 +70,8 @@ public:
 	 * are interpolated with the per particle weight to make the final value used for the simulation.
 	 * Otherwise all particles are considered to have a zero weight, and only the Low value is meaningful.
 	 */
-	UPROPERTY(EditAnywhere, Category = "Bending Properties", Meta = (UIMin = "0", UIMax = "10000", ClampMin = "0", ClampMax = "10000000", EditCondition = "DistributionType == EChaosClothAssetConstraintDistributionType::Anisotropic && bImportProperties == false && SolverType == EChaosClothAssetConstraintSolverType::XPBD", EditConditionHides))
-	FChaosClothAssetWeightedValue BendingStiffnessWeft = { true, 100.f, 100.f, TEXT("BendingStiffnessWeft") };
+	UPROPERTY(EditAnywhere, Category = "Bending Properties", Meta = (UIMin = "0", UIMax = "10000", ClampMin = "0", ClampMax = "10000000", EditCondition = "DistributionType == EChaosClothAssetConstraintDistributionType::Anisotropic && SolverType == EChaosClothAssetConstraintSolverType::XPBD", EditConditionHides))
+	FChaosClothAssetWeightedValue BendingStiffnessWeft = { true, 100.f, 100.f, TEXT("BendingStiffnessWeft"),  true };
 
 	/**
 	 * The stiffness of the bending constraints in the bias (diagonal) direction.
@@ -88,36 +79,17 @@ public:
 	 * are interpolated with the per particle weight to make the final value used for the simulation.
 	 * Otherwise all particles are considered to have a zero weight, and only the Low value is meaningful.
 	 */
-	UPROPERTY(EditAnywhere, Category = "Bending Properties", Meta = (UIMin = "0", UIMax = "10000", ClampMin = "0", ClampMax = "10000000", EditCondition = "DistributionType == EChaosClothAssetConstraintDistributionType::Anisotropic && bImportProperties == false && SolverType == EChaosClothAssetConstraintSolverType::XPBD", EditConditionHides))
-	FChaosClothAssetWeightedValue BendingStiffnessBias = { true, 100.f, 100.f, TEXT("BendingStiffnessBias") };
+	UPROPERTY(EditAnywhere, Category = "Bending Properties", Meta = (UIMin = "0", UIMax = "10000", ClampMin = "0", ClampMax = "10000000", EditCondition = "DistributionType == EChaosClothAssetConstraintDistributionType::Anisotropic && SolverType == EChaosClothAssetConstraintSolverType::XPBD", EditConditionHides))
+	FChaosClothAssetWeightedValue BendingStiffnessBias = { true, 100.f, 100.f, TEXT("BendingStiffnessBias"), true };
 
 	/**
-	 * The damping of the bending constraints, relative to critical damping.
+	 * The damping of the bending anisotropic constraints, relative to critical damping.
 	 * If a valid weight map is found with the given Weight Map name, then both Low and High values
 	 * are interpolated with the per particle weight to make the final value used for the simulation.
 	 * Otherwise all particles are considered to have a zero weight, and only the Low value is meaningful.
 	 */
-	UPROPERTY(EditAnywhere, Category = "Bending Properties", Meta = (UIMin = "0", UIMax = "10", ClampMin = "0", ClampMax = "1000"), Meta = (EditCondition = "bImportProperties == false && SolverType == EChaosClothAssetConstraintSolverType::XPBD", EditConditionHides))
-	FChaosClothAssetWeightedValue BendingDamping = { true, 1.f, 1.f, TEXT("BendingDamping") };
-
-	/**
-	 * Once the element has bent such that it's folded more than this ratio from its rest angle ("buckled"), switch to using Buckling Stiffness instead of BendingElement Stiffness.
-	 * When Buckling Ratio = 0, the Buckling Stiffness will never be used. When BucklingRatio = 1, the Buckling Stiffness will be used as soon as its bent past its rest configuration.
-	 */
-	UPROPERTY(EditAnywhere, Category = "Bending Properties", Meta = (UIMin = "0", UIMax = "1", ClampMin = "0", ClampMax = "1"), Meta = (EditCondition = "bImportProperties == false && ConstraintType == EChaosClothAssetBendingConstraintType::HingeAngles", EditConditionHides))
-	float BucklingRatio = 0.5f;
-
-	/**
-	 * The stiffness after buckling.
-	 * The constraint will use this stiffness instead of bending Stiffness once the cloth has buckled, i.e., bent beyond a certain angle.
-	 * Typically, Buckling Stiffness is set to be less than Bending Stiffness.
-	 * Buckling Ratio determines the switch point between using BendingElement Stiffness and Buckling Stiffness.
-	 * If a valid weight map is found with the given Weight Map name, then both Low and High values
-	 * are interpolated with the per particle weight to make the final value used for the simulation.
-	 * Otherwise all particles are considered to have a zero weight, and only the Low value is meaningful.
-	 */
-	UPROPERTY(EditAnywhere, Category = "Bending Properties", Meta = (UIMin = "0", UIMax = "1", ClampMin = "0", ClampMax = "1", EditCondition = "BucklingRatio != 0 && ((DistributionType == EChaosClothAssetConstraintDistributionType::Isotropic && SolverType == EChaosClothAssetConstraintSolverType::XPBD) || SolverType == EChaosClothAssetConstraintSolverType::PBD) && bImportProperties == false && ConstraintType == EChaosClothAssetBendingConstraintType::HingeAngles", EditConditionHides))
-	FChaosClothAssetWeightedValue BucklingStiffness = { true, 0.9f, 0.9f, TEXT("BucklingStiffness") };
+	UPROPERTY(EditAnywhere, Category = "Bending Properties", Meta = (UIMin = "0", UIMax = "10", ClampMin = "0", ClampMax = "1000"), Meta = (EditCondition = "SolverType == EChaosClothAssetConstraintSolverType::XPBD && DistributionType == EChaosClothAssetConstraintDistributionType::Anisotropic", EditConditionHides))
+	FChaosClothAssetWeightedValue BendingAnisoDamping = { true, 1.f, 1.f, TEXT("BendingAnisoDamping"), true };
 	
 	/**
 	 * The stiffness after buckling in the warp (vertical) direction.
@@ -128,8 +100,8 @@ public:
 	 * are interpolated with the per particle weight to make the final value used for the simulation.
 	 * Otherwise all particles are considered to have a zero weight, and only the Low value is meaningful.
 	 */
-	UPROPERTY(EditAnywhere, Category = "Bending Properties", Meta = (UIMin = "0", UIMax = "10000", ClampMin = "0", ClampMax = "10000000", EditCondition = "BucklingRatio != 0 && DistributionType == EChaosClothAssetConstraintDistributionType::Anisotropic && bImportProperties == false && SolverType == EChaosClothAssetConstraintSolverType::XPBD", EditConditionHides))
-	FChaosClothAssetWeightedValue BucklingStiffnessWarp = { true, 50.f, 50.f, TEXT("BucklingStiffnessWarp") };
+	UPROPERTY(EditAnywhere, Category = "Bending Properties", Meta = (UIMin = "0", UIMax = "10000", ClampMin = "0", ClampMax = "10000000", EditCondition = "BucklingRatio != 0 && DistributionType == EChaosClothAssetConstraintDistributionType::Anisotropic  && SolverType == EChaosClothAssetConstraintSolverType::XPBD", EditConditionHides))
+	FChaosClothAssetWeightedValue BucklingStiffnessWarp = { true, 50.f, 50.f, TEXT("BucklingStiffnessWarp"), true };
 
 	/**
 	 * The stiffness after buckling in the weft (horizontal) direction.
@@ -140,8 +112,8 @@ public:
 	 * are interpolated with the per particle weight to make the final value used for the simulation.
 	 * Otherwise all particles are considered to have a zero weight, and only the Low value is meaningful.
 	 */
-	UPROPERTY(EditAnywhere, Category = "Bending Properties", Meta = (UIMin = "0", UIMax = "10000", ClampMin = "0", ClampMax = "10000000", EditCondition = "BucklingRatio != 0 && DistributionType == EChaosClothAssetConstraintDistributionType::Anisotropic && bImportProperties == false && SolverType == EChaosClothAssetConstraintSolverType::XPBD", EditConditionHides))
-	FChaosClothAssetWeightedValue BucklingStiffnessWeft = { true, 50.f, 50.f, TEXT("BucklingStiffnessWeft") };
+	UPROPERTY(EditAnywhere, Category = "Bending Properties", Meta = (UIMin = "0", UIMax = "10000", ClampMin = "0", ClampMax = "10000000", EditCondition = "BucklingRatio != 0 && DistributionType == EChaosClothAssetConstraintDistributionType::Anisotropic && SolverType == EChaosClothAssetConstraintSolverType::XPBD", EditConditionHides))
+	FChaosClothAssetWeightedValue BucklingStiffnessWeft = { true, 50.f, 50.f, TEXT("BucklingStiffnessWeft"),  true };
 
 	/**
 	 * The stiffness after buckling in the bias (diagonal) direction.
@@ -152,8 +124,45 @@ public:
 	 * are interpolated with the per particle weight to make the final value used for the simulation.
 	 * Otherwise all particles are considered to have a zero weight, and only the Low value is meaningful.
 	 */
-	UPROPERTY(EditAnywhere, Category = "Bending Properties", Meta = (UIMin = "0", UIMax = "10000", ClampMin = "0", ClampMax = "10000000", EditCondition = "BucklingRatio != 0 && DistributionType == EChaosClothAssetConstraintDistributionType::Anisotropic && bImportProperties == false && SolverType == EChaosClothAssetConstraintSolverType::XPBD", EditConditionHides))
-	FChaosClothAssetWeightedValue BucklingStiffnessBias = { true, 50.f, 50.f, TEXT("BucklingStiffnessBias") };
+	UPROPERTY(EditAnywhere, Category = "Bending Properties", Meta = (UIMin = "0", UIMax = "10000", ClampMin = "0", ClampMax = "10000000", EditCondition = "BucklingRatio != 0 && DistributionType == EChaosClothAssetConstraintDistributionType::Anisotropic && SolverType == EChaosClothAssetConstraintSolverType::XPBD", EditConditionHides))
+	FChaosClothAssetWeightedValue BucklingStiffnessBias = { true, 50.f, 50.f, TEXT("BucklingStiffnessBias"), true };
+
+	/**
+	 * The Stiffness of the bending constraints. Increase the iteration count for stiffer materials.
+	 * If a valid weight map is found with the given Weight Map name, then both Low and High values
+	 * are interpolated with the per particle weight to make the final value used for the simulation.
+	 * Otherwise all particles are considered to have a zero weight, and only the Low value is meaningful.
+	 */
+	UPROPERTY(EditAnywhere, Category = "Bending Properties", Meta = (UIMin = "0", UIMax = "1", ClampMin = "0", ClampMax = "1"), Meta = (EditCondition = "((DistributionType == EChaosClothAssetConstraintDistributionType::Isotropic && SolverType == EChaosClothAssetConstraintSolverType::XPBD) || SolverType == EChaosClothAssetConstraintSolverType::PBD)", EditConditionHides))
+	FChaosClothAssetWeightedValue BendingStiffness = { true, 1.f, 1.f, TEXT("BendingStiffness") };
+
+	/**
+	 * The damping of the bending constraints, relative to critical damping.
+	 * If a valid weight map is found with the given Weight Map name, then both Low and High values
+	 * are interpolated with the per particle weight to make the final value used for the simulation.
+	 * Otherwise all particles are considered to have a zero weight, and only the Low value is meaningful.
+	 */
+	UPROPERTY(EditAnywhere, Category = "Bending Properties", Meta = (UIMin = "0", UIMax = "10", ClampMin = "0", ClampMax = "1000"), Meta = (EditCondition = "SolverType == EChaosClothAssetConstraintSolverType::XPBD && DistributionType == EChaosClothAssetConstraintDistributionType::Isotropic ", EditConditionHides))
+	FChaosClothAssetWeightedValue BendingDamping = { true, 1.f, 1.f, TEXT("BendingDamping") };
+	
+	/**
+	 * The stiffness after buckling.
+	 * The constraint will use this stiffness instead of bending Stiffness once the cloth has buckled, i.e., bent beyond a certain angle.
+	 * Typically, Buckling Stiffness is set to be less than Bending Stiffness.
+	 * Buckling Ratio determines the switch point between using BendingElement Stiffness and Buckling Stiffness.
+	 * If a valid weight map is found with the given Weight Map name, then both Low and High values
+	 * are interpolated with the per particle weight to make the final value used for the simulation.
+	 * Otherwise all particles are considered to have a zero weight, and only the Low value is meaningful.
+	 */
+	UPROPERTY(EditAnywhere, Category = "Bending Properties", Meta = (UIMin = "0", UIMax = "1", ClampMin = "0", ClampMax = "1", EditCondition = "BucklingRatio != 0 && ((DistributionType == EChaosClothAssetConstraintDistributionType::Isotropic && SolverType == EChaosClothAssetConstraintSolverType::XPBD) || SolverType == EChaosClothAssetConstraintSolverType::PBD) && ConstraintType == EChaosClothAssetBendingConstraintType::HingeAngles", EditConditionHides))
+	FChaosClothAssetWeightedValue BucklingStiffness = { true, 0.9f, 0.9f, TEXT("BucklingStiffness") };
+
+	/**
+	 * Once the element has bent such that it's folded more than this ratio from its rest angle ("buckled"), switch to using Buckling Stiffness instead of BendingElement Stiffness.
+	 * When Buckling Ratio = 0, the Buckling Stiffness will never be used. When BucklingRatio = 1, the Buckling Stiffness will be used as soon as its bent past its rest configuration.
+	 */
+	UPROPERTY(EditAnywhere, Category = "Bending Properties", Meta = (UIMin = "0", UIMax = "1", ClampMin = "0", ClampMax = "1"), Meta = (EditCondition = "ConstraintType == EChaosClothAssetBendingConstraintType::HingeAngles", EditConditionHides))
+	float BucklingRatio = 0.5f;
 
 	FChaosClothAssetSimulationBendingConfigNode(const Dataflow::FNodeParameters& InParam, FGuid InGuid = FGuid::NewGuid());
 
diff --git a/Engine/Plugins/Experimental/ChaosClothAssetEditor/Source/ChaosClothAssetDataflowNodes/Public/ChaosClothAsset/SimulationCollisionConfigNode.h b/Engine/Plugins/Experimental/ChaosClothAssetEditor/Source/ChaosClothAssetDataflowNodes/Public/ChaosClothAsset/SimulationCollisionConfigNode.h
--- a/Engine/Plugins/Experimental/ChaosClothAssetEditor/Source/ChaosClothAssetDataflowNodes/Public/ChaosClothAsset/SimulationCollisionConfigNode.h
+++ b/Engine/Plugins/Experimental/ChaosClothAssetEditor/Source/ChaosClothAssetDataflowNodes/Public/ChaosClothAsset/SimulationCollisionConfigNode.h
@@ -2,23 +2,23 @@
 
 #pragma once 
 
-#include "ChaosClothAsset/SimulationFabricConfigNode.h"
+#include "ChaosClothAsset/SimulationBaseConfigNode.h"
 #include "SimulationCollisionConfigNode.generated.h"
 
 /** Physics mesh collision properties configuration node. */
 USTRUCT(Meta = (DataflowCloth))
-struct FChaosClothAssetSimulationCollisionConfigNode : public FChaosClothAssetSimulationFabricConfigNode
+struct FChaosClothAssetSimulationCollisionConfigNode : public FChaosClothAssetSimulationBaseConfigNode
 {
 	GENERATED_USTRUCT_BODY()
 	DATAFLOW_NODE_DEFINE_INTERNAL(FChaosClothAssetSimulationCollisionConfigNode, "SimulationCollisionConfig", "Cloth", "Cloth Simulation Collision Config")
 
 public:
 	/** The added thickness of collision shapes. */
-	UPROPERTY(EditAnywhere, Category = "Collision Properties", Meta = (UIMin = "0", UIMax = "100", ClampMin = "0", ClampMax = "1000", EditCondition = "bImportProperties == false", EditConditionHides))
+	UPROPERTY(EditAnywhere, Category = "Collision Properties", Meta = (UIMin = "0", UIMax = "100", ClampMin = "0", ClampMax = "1000"))
 	float CollisionThickness = 1.0f;
 
 	/** Friction coefficient for cloth - collider interaction. */
-	UPROPERTY(EditAnywhere, Category = "Collision Properties", Meta = (UIMin = "0", UIMax = "1", ClampMin = "0", ClampMax = "10", EditCondition = "bImportProperties == false", EditConditionHides))
+	UPROPERTY(EditAnywhere, Category = "Collision Properties", Meta = (UIMin = "0", UIMax = "1", ClampMin = "0", ClampMax = "10"))
 	float FrictionCoefficient = 0.8f;
 
 	/** Stiffness for proximity repulsion forces (Force-based solver only). Units = kg cm/ s^2 (same as XPBD springs)*/
diff --git a/Engine/Plugins/Experimental/ChaosClothAssetEditor/Source/ChaosClothAssetDataflowNodes/Public/ChaosClothAsset/SimulationConfigNodePropertyTypes.h b/Engine/Plugins/Experimental/ChaosClothAssetEditor/Source/ChaosClothAssetDataflowNodes/Public/ChaosClothAsset/SimulationConfigNodePropertyTypes.h
--- a/Engine/Plugins/Experimental/ChaosClothAssetEditor/Source/ChaosClothAssetDataflowNodes/Public/ChaosClothAsset/SimulationConfigNodePropertyTypes.h
+++ b/Engine/Plugins/Experimental/ChaosClothAssetEditor/Source/ChaosClothAssetDataflowNodes/Public/ChaosClothAsset/SimulationConfigNodePropertyTypes.h
@@ -27,15 +27,6 @@ enum class EChaosClothAssetBendingConstraintType : uint8
 	HingeAngles,
 };
 
-UENUM()
-enum class EChaosClothAssetStretchConstraintType : uint8
-{
-	/** Add a UV style stretch and shear constraint */
-	StretchShear,
-	/** Add edge and area springs in the 2d plane */
-	EdgeArea,
-};
-
 UENUM()
 enum class EChaosClothAssetConstraintDistributionType : uint8
 {
diff --git a/Engine/Plugins/Experimental/ChaosClothAssetEditor/Source/ChaosClothAssetDataflowNodes/Public/ChaosClothAsset/SimulationFabricConfigNode.h b/Engine/Plugins/Experimental/ChaosClothAssetEditor/Source/ChaosClothAssetDataflowNodes/Public/ChaosClothAsset/SimulationFabricConfigNode.h
--- a/Engine/Plugins/Experimental/ChaosClothAssetEditor/Source/ChaosClothAssetDataflowNodes/Public/ChaosClothAsset/SimulationFabricConfigNode.h
+++ b/Engine/Plugins/Experimental/ChaosClothAssetEditor/Source/ChaosClothAssetDataflowNodes/Public/ChaosClothAsset/SimulationFabricConfigNode.h
@@ -14,21 +14,11 @@ struct FChaosClothAssetSimulationFabricConfigNode : public FChaosClothAssetSimul
 
 public:
 	
-	/** Import properties to boolean the  fabric override */
-	UPROPERTY(EditAnywhere, Category = "Fabric Properties")
-	bool bImportProperties = false;
-	
 	FChaosClothAssetSimulationFabricConfigNode() = default;
 	
 	FChaosClothAssetSimulationFabricConfigNode(const Dataflow::FNodeParameters& InParam, FGuid InGuid = FGuid::NewGuid());
 	
 protected :
-
-	/** We can only use fabrics to set property values if the import settings is true and if the cloth facade has fabrics */
-	bool CanUseFabrics(const UE::Chaos::ClothAsset::FCollectionClothFacade& ClothFacade) const
-	{
-		return bImportProperties && ClothFacade.IsValid() && (ClothFacade.GetNumFabrics() > 0);
-	}
 	
 	/** Set an imported fabric value onto the weighted value property (animatable or not) */
 	template<typename PropertyType>
diff --git a/Engine/Plugins/Experimental/ChaosClothAssetEditor/Source/ChaosClothAssetDataflowNodes/Public/ChaosClothAsset/SimulationMassConfigNode.h b/Engine/Plugins/Experimental/ChaosClothAssetEditor/Source/ChaosClothAssetDataflowNodes/Public/ChaosClothAsset/SimulationMassConfigNode.h
--- a/Engine/Plugins/Experimental/ChaosClothAssetEditor/Source/ChaosClothAssetDataflowNodes/Public/ChaosClothAsset/SimulationMassConfigNode.h
+++ b/Engine/Plugins/Experimental/ChaosClothAssetEditor/Source/ChaosClothAssetDataflowNodes/Public/ChaosClothAsset/SimulationMassConfigNode.h
@@ -20,15 +20,15 @@ public:
 	 * -	Total Mass: The total mass is distributed equally over all the particles. Useful when referencing a specific garment size and feel.
 	 * -	Density: A constant mass density is used. Density is usually the preferred way of setting mass since it allows matching real life materials' density values.
 	 */
-	UPROPERTY(EditAnywhere, Category = "Mass Properties", meta = (EditCondition = "bImportProperties == false"))
+	UPROPERTY(EditAnywhere, Category = "Mass Properties")
 	EClothMassMode MassMode = EClothMassMode::Density;
 
 	/** The value used when the Mass Mode is set to Uniform Mass. */
-	UPROPERTY(EditAnywhere, Category = "Mass Properties", DisplayName = "Uniform Mass", Meta = (UIMin = "0.000001", UIMax = "0.001", ClampMin = "0", EditCondition = "MassMode == EClothMassMode::UniformMass && bImportProperties == false"))
+	UPROPERTY(EditAnywhere, Category = "Mass Properties", DisplayName = "Uniform Mass", Meta = (UIMin = "0.000001", UIMax = "0.001", ClampMin = "0", EditCondition = "MassMode == EClothMassMode::UniformMass", EditConditionHides))
 	FChaosClothAssetWeightedValueNonAnimatable UniformMassWeighted = {0.00015f, 0.00015f, TEXT("UniformMass")};
 
 	/** The value used when Mass Mode is set to TotalMass. */
-	UPROPERTY(EditAnywhere, Category = "Mass Properties", Meta = (UIMin = "0.001", UIMax = "10", ClampMin = "0", EditCondition = "MassMode == EClothMassMode::TotalMass && bImportProperties == false"))
+	UPROPERTY(EditAnywhere, Category = "Mass Properties", Meta = (UIMin = "0.001", UIMax = "10", ClampMin = "0", EditCondition = "MassMode == EClothMassMode::TotalMass", EditConditionHides))
 	float TotalMass = 0.5f;
 
 	/**
@@ -41,8 +41,8 @@ public:
 	 * Cotton: 0.2
 	 * Silk: 0.1
 	 */
-	UPROPERTY(EditAnywhere, Category = "Mass Properties", DisplayName = "Density", meta = (UIMin = "0.001", UIMax = "1", ClampMin = "0", EditCondition = "MassMode == EClothMassMode::Density && bImportProperties == false"))
-	FChaosClothAssetWeightedValueNonAnimatable DensityWeighted = {0.35f, 0.35f, TEXT("Density")};
+	UPROPERTY(EditAnywhere, Category = "Mass Properties", DisplayName = "Density", meta = (UIMin = "0.001", UIMax = "1", ClampMin = "0", EditCondition = "MassMode == EClothMassMode::Density", EditConditionHides))
+	FChaosClothAssetWeightedValueNonAnimatable DensityWeighted = {0.35f, 0.35f, TEXT("Density"),  true};
 
 	/** Calculated particle masses will be clamped to this minimum value (or 1e-8, whichever is larger). */
 	UPROPERTY(EditAnywhere, Category = "Mass Properties", meta = (ClampMin = "0"))
diff --git a/Engine/Plugins/Experimental/ChaosClothAssetEditor/Source/ChaosClothAssetDataflowNodes/Public/ChaosClothAsset/SimulationStretchConfigNode.h b/Engine/Plugins/Experimental/ChaosClothAssetEditor/Source/ChaosClothAssetDataflowNodes/Public/ChaosClothAsset/SimulationStretchConfigNode.h
--- a/Engine/Plugins/Experimental/ChaosClothAssetEditor/Source/ChaosClothAssetDataflowNodes/Public/ChaosClothAsset/SimulationStretchConfigNode.h
+++ b/Engine/Plugins/Experimental/ChaosClothAssetEditor/Source/ChaosClothAssetDataflowNodes/Public/ChaosClothAsset/SimulationStretchConfigNode.h
@@ -20,25 +20,16 @@ public:
 	bool bStretchUse3dRestLengths = true;
 
 	/** Constraint solver type.  */
-	UPROPERTY(EditAnywhere, Category = "Stretch Types", Meta = (EditCondition = "bImportProperties == false", EditConditionHides))
+	UPROPERTY(EditAnywhere, Category = "Stretch Types")
 	EChaosClothAssetConstraintSolverType SolverType = EChaosClothAssetConstraintSolverType::PBD;
 
 	/**  Constraint distribution type. */
-	UPROPERTY(EditAnywhere, Category = "Stretch Types", Meta = (EditCondition = "bImportProperties == false && SolverType == EChaosClothAssetConstraintSolverType::XPBD", EditConditionHides))
+	UPROPERTY(EditAnywhere, Category = "Stretch Types", Meta = (EditCondition = "SolverType == EChaosClothAssetConstraintSolverType::XPBD", EditConditionHides))
 	EChaosClothAssetConstraintDistributionType DistributionType = EChaosClothAssetConstraintDistributionType::Isotropic;
 	
 	/** Add an area constraint in case of isotropic distribution  */
-	UPROPERTY(EditAnywhere, Category = "Stretch Types", Meta = (EditCondition = "bImportProperties == false && ((SolverType == EChaosClothAssetConstraintSolverType::XPBD && DistributionType == EChaosClothAssetConstraintDistributionType::Isotropic) || SolverType == EChaosClothAssetConstraintSolverType::PBD)", EditConditionHides))
+	UPROPERTY(EditAnywhere, Category = "Stretch Types", Meta = (EditCondition = "((SolverType == EChaosClothAssetConstraintSolverType::XPBD && DistributionType == EChaosClothAssetConstraintDistributionType::Isotropic) || SolverType == EChaosClothAssetConstraintSolverType::PBD)", EditConditionHides))
 	bool bAddAreaConstraint = true;
-
-	/**
-	 * The stiffness of the stretch constraints. Note that PBD stiffnesses will be internally clamped to [0,1].
-	 * If a valid weight map is found with the given Weight Map name, then both Low and High values
-	 * are interpolated with the per particle weight to make the final value used for the simulation.
-	 * Otherwise all particles are considered to have a zero weight, and only the Low value is meaningful.
-	 */
-	UPROPERTY(EditAnywhere, Category = "Stretch Properties", Meta = (UIMin = "0", UIMax = "10000", ClampMin = "0", ClampMax = "1000000000", EditCondition = "bImportProperties == false && (SolverType == EChaosClothAssetConstraintSolverType::PBD || DistributionType == EChaosClothAssetConstraintDistributionType::Isotropic)", EditConditionHides))
-	FChaosClothAssetWeightedValue StretchStiffness = { true, 1.f, 1.f, TEXT("StretchStiffness") };
 	
 	/**
 	 * The stiffness of the stretch constraints in the warp (vertical) direction.
@@ -46,8 +37,8 @@ public:
 	 * are interpolated with the per particle weight to make the final value used for the simulation.
 	 * Otherwise all particles are considered to have a zero weight, and only the Low value is meaningful.
 	 */
-	UPROPERTY(EditAnywhere, Category = "Stretch Properties", Meta = (UIMin = "0", UIMax = "10000", ClampMin = "0", ClampMax = "10000000", EditCondition = "bImportProperties == false && DistributionType == EChaosClothAssetConstraintDistributionType::Anisotropic && SolverType == EChaosClothAssetConstraintSolverType::XPBD", EditConditionHides))
-	FChaosClothAssetWeightedValue StretchStiffnessWarp = { true, 100.f, 100.f, TEXT("StretchStiffnessWarp") };
+	UPROPERTY(EditAnywhere, Category = "Stretch Properties", Meta = (UIMin = "0", UIMax = "10000", ClampMin = "0", ClampMax = "10000000", EditCondition = "DistributionType == EChaosClothAssetConstraintDistributionType::Anisotropic && SolverType == EChaosClothAssetConstraintSolverType::XPBD", EditConditionHides))
+	FChaosClothAssetWeightedValue StretchStiffnessWarp = { true, 100.f, 100.f, TEXT("StretchStiffnessWarp"),true };
 
 	/**
 	 * The stiffness of the stretch constraints in the weft (horizontal) direction.
@@ -55,8 +46,8 @@ public:
 	 * are interpolated with the per particle weight to make the final value used for the simulation.
 	 * Otherwise all particles are considered to have a zero weight, and only the Low value is meaningful.
 	 */
-	UPROPERTY(EditAnywhere, Category = "Stretch Properties", Meta = (UIMin = "0", UIMax = "10000", ClampMin = "0", ClampMax = "10000000", EditCondition = "bImportProperties == false && DistributionType == EChaosClothAssetConstraintDistributionType::Anisotropic && SolverType == EChaosClothAssetConstraintSolverType::XPBD", EditConditionHides))
-	FChaosClothAssetWeightedValue StretchStiffnessWeft = { true, 100.f, 100.f, TEXT("StretchStiffnessWeft") };
+	UPROPERTY(EditAnywhere, Category = "Stretch Properties", Meta = (UIMin = "0", UIMax = "10000", ClampMin = "0", ClampMax = "10000000", EditCondition = "DistributionType == EChaosClothAssetConstraintDistributionType::Anisotropic && SolverType == EChaosClothAssetConstraintSolverType::XPBD", EditConditionHides))
+	FChaosClothAssetWeightedValue StretchStiffnessWeft = { true, 100.f, 100.f, TEXT("StretchStiffnessWeft"), true };
 
 	/**
 	 * The stiffness of the stretch constraints in the bias (diagonal) direction.
@@ -64,9 +55,28 @@ public:
 	 * are interpolated with the per particle weight to make the final value used for the simulation.
 	 * Otherwise all particles are considered to have a zero weight, and only the Low value is meaningful.
 	 */
-	UPROPERTY(EditAnywhere, Category = "Stretch Properties", Meta = (UIMin = "0", UIMax = "10000", ClampMin = "0", ClampMax = "10000000", EditCondition = "bImportProperties == false && DistributionType == EChaosClothAssetConstraintDistributionType::Anisotropic && SolverType == EChaosClothAssetConstraintSolverType::XPBD", EditConditionHides))
-	FChaosClothAssetWeightedValue StretchStiffnessBias = { true, 100.f, 100.f, TEXT("StretchStiffnessBias") };
+	UPROPERTY(EditAnywhere, Category = "Stretch Properties", Meta = (UIMin = "0", UIMax = "10000", ClampMin = "0", ClampMax = "10000000", EditCondition = "DistributionType == EChaosClothAssetConstraintDistributionType::Anisotropic && SolverType == EChaosClothAssetConstraintSolverType::XPBD", EditConditionHides))
+	FChaosClothAssetWeightedValue StretchStiffnessBias = { true, 100.f, 100.f, TEXT("StretchStiffnessBias"), true };
 
+	/**
+	 * 
+	 * The damping of the stretch anisotropic constraints, relative to critical damping.
+	 * If a valid weight map is found with the given Weight Map name, then both Low and High values
+	 * are interpolated with the per particle weight to make the final value used for the simulation.
+	 * Otherwise all particles are considered to have a zero weight, and only the Low value is meaningful.
+	 */
+	UPROPERTY(EditAnywhere, Category = "Stretch Properties", Meta = (UIMin = "0", UIMax = "10", ClampMin = "0", ClampMax = "1000", EditCondition = "SolverType == EChaosClothAssetConstraintSolverType::XPBD && DistributionType == EChaosClothAssetConstraintDistributionType::Anisotropic", EditConditionHides))
+	FChaosClothAssetWeightedValue StretchAnisoDamping = { true, 1.f, 1.f, TEXT("StretchAnisoDamping"), true };
+	
+	/**
+     * The stiffness of the stretch constraints. Note that PBD stiffnesses will be internally clamped to [0,1].
+     * If a valid weight map is found with the given Weight Map name, then both Low and High values
+     * are interpolated with the per particle weight to make the final value used for the simulation.
+     * Otherwise all particles are considered to have a zero weight, and only the Low value is meaningful.
+     */
+    UPROPERTY(EditAnywhere, Category = "Stretch Properties", Meta = (UIMin = "0", UIMax = "10000", ClampMin = "0", ClampMax = "1000000000", EditCondition = "(SolverType == EChaosClothAssetConstraintSolverType::PBD || DistributionType == EChaosClothAssetConstraintDistributionType::Isotropic)", EditConditionHides))
+    FChaosClothAssetWeightedValue StretchStiffness = { true, 1.f, 1.f, TEXT("StretchStiffness") };
+	
 	/**
 	 * 
 	 * The damping of the stretch constraints, relative to critical damping.
@@ -74,7 +84,7 @@ public:
 	 * are interpolated with the per particle weight to make the final value used for the simulation.
 	 * Otherwise all particles are considered to have a zero weight, and only the Low value is meaningful.
 	 */
-	UPROPERTY(EditAnywhere, Category = "Stretch Properties", Meta = (UIMin = "0", UIMax = "10", ClampMin = "0", ClampMax = "1000", EditCondition = "bImportProperties == false && SolverType == EChaosClothAssetConstraintSolverType::XPBD", EditConditionHides))
+	UPROPERTY(EditAnywhere, Category = "Stretch Properties", Meta = (UIMin = "0", UIMax = "10", ClampMin = "0", ClampMax = "1000", EditCondition = "SolverType == EChaosClothAssetConstraintSolverType::XPBD && DistributionType == EChaosClothAssetConstraintDistributionType::Isotropic", EditConditionHides))
 	FChaosClothAssetWeightedValue StretchDamping = { true, 1.f, 1.f, TEXT("StretchDamping") };
 
 	/**
@@ -83,7 +93,7 @@ public:
 	 * are interpolated with the per particle weight to make the final value used for the simulation.
 	 * Otherwise all particles are considered to have a zero weight, and only the Low value is meaningful.
 	 */
-	UPROPERTY(EditAnywhere, Category = "Stretch Properties", Meta = (UIMin = "0", UIMax = "10", ClampMin = "0", ClampMax = "10000000", EditCondition = "bImportProperties == true || (SolverType == EChaosClothAssetConstraintSolverType::XPBD && DistributionType == EChaosClothAssetConstraintDistributionType::Anisotropic)", EditConditionHides))
+	UPROPERTY(EditAnywhere, Category = "Stretch Properties", Meta = (UIMin = "0", UIMax = "10", ClampMin = "0", ClampMax = "10000000", EditCondition = "(SolverType == EChaosClothAssetConstraintSolverType::XPBD && DistributionType == EChaosClothAssetConstraintDistributionType::Anisotropic)", EditConditionHides))
 	FChaosClothAssetWeightedValue StretchWarpScale = { true, 1.f, 1.f, TEXT("StretchWarpScale") };
 
 	/**
@@ -92,7 +102,7 @@ public:
 	 * are interpolated with the per particle weight to make the final value used for the simulation.
 	 * Otherwise all particles are considered to have a zero weight, and only the Low value is meaningful.
 	 */
-	UPROPERTY(EditAnywhere, Category = "Stretch Properties", Meta = (UIMin = "0", UIMax = "10", ClampMin = "0", ClampMax = "10000000", EditCondition = "bImportProperties == true || (SolverType == EChaosClothAssetConstraintSolverType::XPBD && DistributionType == EChaosClothAssetConstraintDistributionType::Anisotropic)", EditConditionHides))
+	UPROPERTY(EditAnywhere, Category = "Stretch Properties", Meta = (UIMin = "0", UIMax = "10", ClampMin = "0", ClampMax = "10000000", EditCondition = "(SolverType == EChaosClothAssetConstraintSolverType::XPBD && DistributionType == EChaosClothAssetConstraintDistributionType::Anisotropic)", EditConditionHides))
 	FChaosClothAssetWeightedValue StretchWeftScale = { true, 1.f, 1.f, TEXT("StretchWeftScale") };
 
 	/**
@@ -102,7 +112,7 @@ public:
 	 * are interpolated with the per particle weight to make the final value used for the simulation.
 	 * Otherwise all particles are considered to have a zero weight, and only the Low value is meaningful.
 	 */
-	UPROPERTY(EditAnywhere, Category = "Stretch Properties", Meta = (UIMin = "0", UIMax = "10000", ClampMin = "0", ClampMax = "10000000", EditCondition = "bImportProperties == false && bAddAreaConstraint == true && ((SolverType == EChaosClothAssetConstraintSolverType::XPBD && DistributionType == EChaosClothAssetConstraintDistributionType::Isotropic) || SolverType == EChaosClothAssetConstraintSolverType::PBD)", EditConditionHides))
+	UPROPERTY(EditAnywhere, Category = "Stretch Properties", Meta = (UIMin = "0", UIMax = "10000", ClampMin = "0", ClampMax = "10000000", EditCondition = "bAddAreaConstraint == true && ((SolverType == EChaosClothAssetConstraintSolverType::XPBD && DistributionType == EChaosClothAssetConstraintDistributionType::Isotropic) || SolverType == EChaosClothAssetConstraintSolverType::PBD)", EditConditionHides))
 	FChaosClothAssetWeightedValue AreaStiffness = { true, 1.f, 1.f, TEXT("AreaStiffness") };
 
 	FChaosClothAssetSimulationStretchConfigNode(const Dataflow::FNodeParameters& InParam, FGuid InGuid = FGuid::NewGuid());
diff --git a/Engine/Plugins/Experimental/ChaosClothAssetEditor/Source/ChaosClothAssetDataflowNodes/Public/ChaosClothAsset/WeightedValue.h b/Engine/Plugins/Experimental/ChaosClothAssetEditor/Source/ChaosClothAssetDataflowNodes/Public/ChaosClothAsset/WeightedValue.h
--- a/Engine/Plugins/Experimental/ChaosClothAssetEditor/Source/ChaosClothAssetDataflowNodes/Public/ChaosClothAsset/WeightedValue.h
+++ b/Engine/Plugins/Experimental/ChaosClothAssetEditor/Source/ChaosClothAssetDataflowNodes/Public/ChaosClothAsset/WeightedValue.h
@@ -50,8 +50,8 @@ struct FChaosClothAssetWeightedValue
 	 * A Weight of 0 always corresponds to the Low parameter value, and a Weight of 1 to the High parameter value.
 	 * The value for Low can be set to be bigger than for High in order to reverse the effect of the Weight Map.
 	 */
-	UPROPERTY(EditAnywhere, Category = "Weighted Value", Meta = (DisplayName = "Low Weight", ChaosClothAssetShortName = "Lo"))
-	float Low = 0.f;
+	UPROPERTY(EditAnywhere, Category = "Weighted Value", Meta = (DisplayName = "Low Weight", ChaosClothAssetShortName = "Lo", EditCondition = "!bImportFabricBounds"))
+	mutable float Low = 0.f;
 
 	/**
 	 * Property value corresponding to the upper bound of the Weight Map.
@@ -60,16 +60,34 @@ struct FChaosClothAssetWeightedValue
 	 * A Weight of 0 always corresponds to the Low parameter value, and a Weight of 1 to the High parameter value.
 	 * The value for Low can be set to be bigger than for High in order to reverse the effect of the Weight Map.
 	 */
-	UPROPERTY(EditAnywhere, Category = "Weighted Value", Meta = (DisplayName = "High Weight", ChaosClothAssetShortName = "Hi"))
-	float High = 1.f;
+	UPROPERTY(EditAnywhere, Category = "Weighted Value", Meta = (DisplayName = "High Weight", ChaosClothAssetShortName = "Hi", EditCondition = "!bImportFabricBounds"))
+	mutable float High = 1.f;
 
 	/** The name of the weight map for this property. */
 	UPROPERTY(EditAnywhere, Category = "Weighted Value", Meta = (DataflowInput))
 	mutable FString WeightMap = TEXT("WeightMap");  // Mutable so that it can be name checked in the evaluate function
 
+	/**
+	 * Whether the property could import fabrics datas or not
+	 */
+	UPROPERTY(VisibleAnywhere, Transient, Category = "Weighted Value")
+	bool bCouldUseFabrics = false;
+
 	/** The weight map override value for when the WeightMap has a connection that replaces the provided weight map value. */
 	UPROPERTY(VisibleAnywhere, Transient, Category = "Weighted Value")
 	mutable FString WeightMap_Override = UE::Chaos::ClothAsset::FWeightMapTools::NotOverridden;  // _Override has a special meaning to the property customization, mutable because this property is set while getting the original value
+	
+	/**
+	 * Whether the property can use the fabric bounds (low/high) values imported from USD
+	 */
+	UPROPERTY(EditAnywhere, Category = "Weighted Value", Meta = (EditCondition = "bCouldUseFabrics", EditConditionHides))
+	bool bImportFabricBounds = false;
+
+	/**
+	 * Whether the property can override the weight map based on the imported fabrics
+	 */
+	UPROPERTY(EditAnywhere, Category = "Weighted Value", Meta = (EditCondition = "bCouldUseFabrics", EditConditionHides))
+	bool bBuildFabricMaps = false;
 };
 
 USTRUCT()
@@ -86,8 +104,8 @@ struct FChaosClothAssetWeightedValueNonAnimatable
 	 * A Weight of 0 always corresponds to the Low parameter value, and a Weight of 1 to the High parameter value.
 	 * The value for Low can be set to be bigger than for High in order to reverse the effect of the Weight Map.
 	 */
-	UPROPERTY(EditAnywhere, Category = "Weighted Value", Meta = (DisplayName = "Low Weight", ChaosClothAssetShortName = "Lo"))
-	float Low = 0.f;
+	UPROPERTY(EditAnywhere, Category = "Weighted Value", Meta = (DisplayName = "Low Weight", ChaosClothAssetShortName = "Lo", EditCondition = "!bImportFabricBounds"))
+	mutable float Low = 0.f;
 
 	/**
 	 * Property value corresponding to the upper bound of the Weight Map.
@@ -96,18 +114,35 @@ struct FChaosClothAssetWeightedValueNonAnimatable
 	 * A Weight of 0 always corresponds to the Low parameter value, and a Weight of 1 to the High parameter value.
 	 * The value for Low can be set to be bigger than for High in order to reverse the effect of the Weight Map.
 	 */
-	UPROPERTY(EditAnywhere, Category = "Weighted Value", Meta = (DisplayName = "High Weight", ChaosClothAssetShortName = "Hi"))
-	float High = 1.f;
-
+	UPROPERTY(EditAnywhere, Category = "Weighted Value", Meta = (DisplayName = "High Weight", ChaosClothAssetShortName = "Hi", EditCondition = "!bImportFabricBounds"))
+	mutable float High = 1.f;
+	
 	/** The name of the weight map for this property. */
 	UPROPERTY(EditAnywhere, Category = "Weighted Value", Meta = (DataflowInput))
 	mutable FString WeightMap = TEXT("WeightMap");  // Mutable so that it can be name checked in the evaluate function
 
+	/**
+	 * Whether the property could import fabrics datas or not
+	 */
+	UPROPERTY(VisibleAnywhere, Transient, Category = "Weighted Value")
+	bool bCouldUseFabrics = false;
+	
 	/** The weight map override value for when the WeightMap has a connection that replaces the provided weight map value. */
 	UPROPERTY(VisibleAnywhere, Transient, Category = "Weighted Value")
 	mutable FString WeightMap_Override = UE::Chaos::ClothAsset::FWeightMapTools::NotOverridden;  // _Override has a special meaning to the property customization, mutable because this property is set while getting the original value
-};
+	
+	/**
+	 * Whether the property can use the fabric bounds (low/high) values imported from USD
+	 */
+	UPROPERTY(EditAnywhere, Category = "Weighted Value", Meta = (EditCondition = "bCouldUseFabrics", EditConditionHides))
+	bool bImportFabricBounds = false;
 
+	/**
+	 * Whether the property can override the weight map based on the imported fabrics
+	 */
+	UPROPERTY(EditAnywhere, Category = "Weighted Value", Meta = (EditCondition = "bCouldUseFabrics", EditConditionHides))
+	bool bBuildFabricMaps = false;
+};
 
 USTRUCT()
 struct FChaosClothAssetWeightedValueNonAnimatableNoLowHighRange
diff --git a/Engine/Plugins/Experimental/ChaosClothAssetEditor/Source/ChaosClothAssetTools/Private/ChaosClothAsset/ChaosClothAssetToolsModule.cpp b/Engine/Plugins/Experimental/ChaosClothAssetEditor/Source/ChaosClothAssetTools/Private/ChaosClothAsset/ChaosClothAssetToolsModule.cpp
--- a/Engine/Plugins/Experimental/ChaosClothAssetEditor/Source/ChaosClothAssetTools/Private/ChaosClothAsset/ChaosClothAssetToolsModule.cpp
+++ b/Engine/Plugins/Experimental/ChaosClothAssetEditor/Source/ChaosClothAssetTools/Private/ChaosClothAsset/ChaosClothAssetToolsModule.cpp
@@ -2,7 +2,7 @@
 
 #include "Features/IModularFeatures.h"
 #include "Modules/ModuleManager.h"
-#include "ClothAssetEditorStyle.h"
+#include "ChaosClothAsset/ClothAssetEditorStyle.h"
 #include "ClothComponentEditorStyle.h"
 #include "ClothingAssetToClothAssetExporter.h"
 
diff --git a/Engine/Plugins/Experimental/ChaosClothAssetEditor/Source/ChaosClothAssetTools/Private/ChaosClothAsset/ClothAssetEditorStyle.cpp b/Engine/Plugins/Experimental/ChaosClothAssetEditor/Source/ChaosClothAssetTools/Private/ChaosClothAsset/ClothAssetEditorStyle.cpp
--- a/Engine/Plugins/Experimental/ChaosClothAssetEditor/Source/ChaosClothAssetTools/Private/ChaosClothAsset/ClothAssetEditorStyle.cpp
+++ b/Engine/Plugins/Experimental/ChaosClothAssetEditor/Source/ChaosClothAssetTools/Private/ChaosClothAsset/ClothAssetEditorStyle.cpp
@@ -1,6 +1,6 @@
 // Copyright Epic Games, Inc. All Rights Reserved.
 
-#include "ClothAssetEditorStyle.h"
+#include "ChaosClothAsset/ClothAssetEditorStyle.h"
 #include "Styling/SlateStyleRegistry.h"
 #include "Interfaces/IPluginManager.h"
 
@@ -16,6 +16,9 @@ namespace UE::Chaos::ClothAsset
 		Set("ClassIcon.ChaosClothAsset", new FSlateVectorImageBrush(RootToContentDir(TEXT("ClothAsset_16.svg")), FVector2D(16)));
 		Set("ClassThumbnail.ChaosClothAsset", new FSlateVectorImageBrush(RootToContentDir(TEXT("ClothAsset_64.svg")), FVector2D(64)));
 
+		Set("ClassIcon.ChaosClothPreset", new FSlateVectorImageBrush(RootToContentDir(TEXT("ClothPreset_16.svg")), FVector2D(16)));
+		Set("ClassThumbnail.ChaosClothPreset", new FSlateVectorImageBrush(RootToContentDir(TEXT("ClothPreset_64.svg")), FVector2D(64)));
+
 		FSlateStyleRegistry::RegisterSlateStyle(*this);
 	}
 
diff --git a/Engine/Plugins/Experimental/ChaosClothAssetEditor/Source/ChaosClothAssetTools/Private/ChaosClothAsset/ClothAssetEditorStyle.h b/Engine/Plugins/Experimental/ChaosClothAssetEditor/Source/ChaosClothAssetTools/Private/ChaosClothAsset/ClothAssetEditorStyle.h
deleted file mode 100644
diff --git a/Engine/Plugins/Experimental/ChaosClothAssetEditor/Source/ChaosClothAssetTools/Public/ChaosClothAsset/ClothAssetEditorStyle.h b/Engine/Plugins/Experimental/ChaosClothAssetEditor/Source/ChaosClothAssetTools/Public/ChaosClothAsset/ClothAssetEditorStyle.h
new file mode 100644
--- /dev/null
+++ b/Engine/Plugins/Experimental/ChaosClothAssetEditor/Source/ChaosClothAssetTools/Public/ChaosClothAsset/ClothAssetEditorStyle.h
@@ -0,0 +1,36 @@
+// Copyright Epic Games, Inc. All Rights Reserved.
+
+#pragma once
+
+#include "Styling/SlateStyle.h"
+
+namespace UE::Chaos::ClothAsset
+{
+	/**
+	 * Editor style setting up the cloth asset icons in editor.
+	 */
+	class CHAOSCLOTHASSETTOOLS_API FClothAssetEditorStyle final : public FSlateStyleSet
+	{
+	public:
+		FClothAssetEditorStyle();
+		virtual ~FClothAssetEditorStyle() override;
+
+	public:
+		static FClothAssetEditorStyle& Get()
+		{
+			if (!Singleton.IsSet())
+			{
+				Singleton.Emplace();
+			}
+			return Singleton.GetValue();
+		}
+
+		static void Destroy()
+		{
+			Singleton.Reset();
+		}
+
+	private:
+		static TOptional<FClothAssetEditorStyle> Singleton;
+	};
+}  // End namespace UE::Chaos::ClothAsset
