Add grow, shrink, and flood buttons to the Selection tool

diff --git a/Engine/Plugins/Experimental/ChaosClothAssetEditor/Source/ChaosClothAssetEditorTools/Private/ChaosClothAsset/ClothMeshSelectionTool.cpp b/Engine/Plugins/Experimental/ChaosClothAssetEditor/Source/ChaosClothAssetEditorTools/Private/ChaosClothAsset/ClothMeshSelectionTool.cpp
--- a/Engine/Plugins/Experimental/ChaosClothAssetEditor/Source/ChaosClothAssetEditorTools/Private/ChaosClothAsset/ClothMeshSelectionTool.cpp
+++ b/Engine/Plugins/Experimental/ChaosClothAssetEditor/Source/ChaosClothAssetEditorTools/Private/ChaosClothAsset/ClothMeshSelectionTool.cpp
@@ -21,6 +21,7 @@
 #include "Dataflow/DataflowGraphEditor.h"
 #include "Dataflow/DataflowEdNode.h"
 #include "DynamicMesh/NonManifoldMappingSupport.h"
+#include "Selections/GeometrySelectionUtil.h"
 
 #include UE_INLINE_GENERATED_CPP_BY_NAME(ClothMeshSelectionTool)
 
@@ -535,6 +536,15 @@ void UClothMeshSelectionTool::ApplyAction(EClothMeshSelectionToolActions ActionT
 	case EClothMeshSelectionToolActions::TogglePrimarySecondary:
 		TogglePrimarySecondaryAction();
 		break;
+	case EClothMeshSelectionToolActions::GrowSelection:
+		GrowSelection();
+		break;
+	case EClothMeshSelectionToolActions::ShrinkSelection:
+		ShrinkSelection();
+		break;
+	case EClothMeshSelectionToolActions::FloodSelection:
+		FloodSelection();
+		break;
 	}
 }
 
@@ -694,4 +704,87 @@ void UClothMeshSelectionTool::UpdatePrimarySecondaryMessage()
 	GetToolManager()->DisplayMessage(Message, EToolMessageLevel::UserWarning);
 }
 
+void UClothMeshSelectionTool::GrowSelection()
+{
+	using namespace UE::Geometry;
+
+	PreviewMesh->ProcessMesh([this](const FDynamicMesh3& Mesh)
+		{
+			FGeometrySelection Selection;
+			Selection.ElementType = SelectionMechanic->Properties->bSelectFaces ? EGeometryElementType::Face : EGeometryElementType::Vertex;
+			SelectionMechanic->GetSelection_AsTriangleTopology(Selection);
+
+			FGeometrySelection BoundarySelection;
+			BoundarySelection.ElementType = SelectionMechanic->Properties->bSelectFaces ? EGeometryElementType::Face : EGeometryElementType::Vertex;
+			MakeBoundaryConnectedSelection(Mesh,
+				Topology.Get(),
+				Selection,
+				[](FGeoSelectionID) { return true; },
+				BoundarySelection
+			);
+			CombineSelectionInPlace(Selection, BoundarySelection, EGeometrySelectionCombineModes::Add);
+			SelectionMechanic->SetSelection_AsTriangleTopology(Selection);
+		});
+
+	PreviewMesh->FastNotifySecondaryTrianglesChanged();
+}
+
+
+void UClothMeshSelectionTool::ShrinkSelection()
+{
+	using namespace UE::Geometry;
+
+	PreviewMesh->ProcessMesh([this](const FDynamicMesh3& Mesh)
+		{
+			FGeometrySelection Selection;
+			Selection.ElementType = SelectionMechanic->Properties->bSelectFaces ? EGeometryElementType::Face : EGeometryElementType::Vertex;
+			SelectionMechanic->GetSelection_AsTriangleTopology(Selection);
+
+			FGeometrySelection BoundarySelection;
+			BoundarySelection.ElementType = SelectionMechanic->Properties->bSelectFaces ? EGeometryElementType::Face : EGeometryElementType::Vertex;
+			MakeBoundaryConnectedSelection(Mesh,
+				Topology.Get(),
+				Selection,
+				[](FGeoSelectionID) { return true; },
+				BoundarySelection
+			);
+
+			CombineSelectionInPlace(Selection, BoundarySelection, EGeometrySelectionCombineModes::Subtract);
+			
+			// TODO: SetSelection_AsTriangleTopology doesn't overwrite the selection, it only adds to it. Fix that.
+			FGroupTopologySelection ClearSelection;
+			SelectionMechanic->SetSelection(ClearSelection, false);
+
+			SelectionMechanic->SetSelection_AsTriangleTopology(Selection);
+		});
+
+	PreviewMesh->FastNotifySecondaryTrianglesChanged();
+}
+
+
+void UClothMeshSelectionTool::FloodSelection()
+{
+	using namespace UE::Geometry;
+
+	PreviewMesh->ProcessMesh([this](const FDynamicMesh3& Mesh)
+		{
+			FGeometrySelection Selection;
+			Selection.ElementType = SelectionMechanic->Properties->bSelectFaces ? EGeometryElementType::Face : EGeometryElementType::Vertex;
+			SelectionMechanic->GetSelection_AsTriangleTopology(Selection);
+
+			FGeometrySelection ConnectedSelection;
+			ConnectedSelection.ElementType = SelectionMechanic->Properties->bSelectFaces ? EGeometryElementType::Face : EGeometryElementType::Vertex;
+			MakeSelectAllConnectedSelection(Mesh,
+				Topology.Get(),
+				Selection,
+				[](FGeoSelectionID) { return true; },
+				[](FGeoSelectionID, FGeoSelectionID) { return true; },
+				ConnectedSelection
+			);
+			SelectionMechanic->SetSelection_AsTriangleTopology(ConnectedSelection);
+		});
+
+	PreviewMesh->FastNotifySecondaryTrianglesChanged();
+}
+
 #undef LOCTEXT_NAMESPACE
diff --git a/Engine/Plugins/Experimental/ChaosClothAssetEditor/Source/ChaosClothAssetEditorTools/Public/ChaosClothAsset/ClothMeshSelectionTool.h b/Engine/Plugins/Experimental/ChaosClothAssetEditor/Source/ChaosClothAssetEditorTools/Public/ChaosClothAsset/ClothMeshSelectionTool.h
--- a/Engine/Plugins/Experimental/ChaosClothAssetEditor/Source/ChaosClothAssetEditorTools/Public/ChaosClothAsset/ClothMeshSelectionTool.h
+++ b/Engine/Plugins/Experimental/ChaosClothAssetEditor/Source/ChaosClothAssetEditorTools/Public/ChaosClothAsset/ClothMeshSelectionTool.h
@@ -38,6 +38,10 @@ enum class EClothMeshSelectionToolActions
 	ImportFromCollection,
 	ImportSecondaryFromCollection,
 	TogglePrimarySecondary,
+
+	GrowSelection,
+	ShrinkSelection,
+	FloodSelection
 };
 
 UCLASS()
@@ -71,6 +75,24 @@ public:
 		PostAction(EClothMeshSelectionToolActions::TogglePrimarySecondary);
 	}
 
+	UFUNCTION(CallInEditor, Category = Selection)
+	void GrowSelection()
+	{
+		PostAction(EClothMeshSelectionToolActions::GrowSelection);
+	}
+
+	UFUNCTION(CallInEditor, Category = Selection)
+	void ShrinkSelection()
+	{
+		PostAction(EClothMeshSelectionToolActions::ShrinkSelection);
+	}
+
+	UFUNCTION(CallInEditor, Category = Selection)
+	void FloodSelection()
+	{
+		PostAction(EClothMeshSelectionToolActions::FloodSelection);
+	}
+
 };
 
 UCLASS()
@@ -163,7 +185,10 @@ private:
 
 	void ImportFromCollection(bool bImportFromSecondarySet);
 	void TogglePrimarySecondaryAction();
-	
 	void UpdatePrimarySecondaryMessage();
+
+	void GrowSelection();
+	void ShrinkSelection();
+	void FloodSelection();
 };
 
