Merge cloth collection weight map name fix

diff --git a/Engine/Plugins/Experimental/ChaosClothAssetEditor/Source/ChaosClothAssetDataflowNodes/Private/ChaosClothAsset/MergeClothCollectionsNode.cpp b/Engine/Plugins/Experimental/ChaosClothAssetEditor/Source/ChaosClothAssetDataflowNodes/Private/ChaosClothAsset/MergeClothCollectionsNode.cpp
--- a/Engine/Plugins/Experimental/ChaosClothAssetEditor/Source/ChaosClothAssetDataflowNodes/Private/ChaosClothAsset/MergeClothCollectionsNode.cpp
+++ b/Engine/Plugins/Experimental/ChaosClothAssetEditor/Source/ChaosClothAssetDataflowNodes/Private/ChaosClothAsset/MergeClothCollectionsNode.cpp
@@ -31,27 +31,40 @@ namespace UE::Chaos::ClothAsset::Private
 		FClothDataflowTools::LogAndToastWarning(DataflowNode, Headline, Details);
 	}
 
+	struct FMergedProperty
+	{
+		FString WeightMapName;
+		FVector4f PropertyBounds;
+	};
+
 	/** Build weight maps for each properties if necessary */
 	static FString BuildWeightMaps(const FDataflowNode& DataflowNode,
 		const FCollectionClothConstFacade& InClothFacade, FCollectionClothFacade& OutClothFacade,
 		const FVector2f& InPropertyBounds, const FVector2f& OutPropertyBounds,
 		const FVector2f& PropertyBounds, const FString& PropertyName,
-		const FString& InWeightMapName, const FString& OutWeightMapName)
+		const FString& InWeightMapName, const FString& OutWeightMapName, TMap<FString,FMergedProperty>& MergedPropertyMaps)
 	{
-		FString WeightMapName = (OutWeightMapName.IsEmpty() && !InWeightMapName.IsEmpty()) ? InWeightMapName :
-								(OutWeightMapName.IsEmpty() && InWeightMapName.IsEmpty()) ? PropertyName : OutWeightMapName;
+		const FMergedProperty MergedProperty = {OutWeightMapName + FString(TEXT("_")) + InWeightMapName,
+			FVector4f(InPropertyBounds[0], InPropertyBounds[1], OutPropertyBounds[0], OutPropertyBounds[1])};
 
-		if(WeightMapName != OutWeightMapName)
+		for(const TPair<FString, FMergedProperty>& MergedPropertyMap : MergedPropertyMaps)
 		{
-			FString IncrWeightMapName = WeightMapName;
-			int32 IncrCount = 0;
-			
-			while(OutClothFacade.GetWeightMap(FName(IncrWeightMapName)).Num() > 0)
+			if((MergedPropertyMap.Value.WeightMapName == MergedProperty.WeightMapName) &&
+			   (MergedPropertyMap.Value.PropertyBounds == MergedProperty.PropertyBounds))
 			{
-				IncrWeightMapName = WeightMapName;
-				IncrWeightMapName.AppendInt(++IncrCount);
+				return MergedPropertyMap.Key;
 			}
-			WeightMapName = IncrWeightMapName;
+		}
+		FString WeightMapName = PropertyName;
+		int32 WeightMapCount = 0;
+		
+		// the weight map could already been stored on the out collection and linked to different bounds
+		// coming from the out collection itself or from previous merge with in collection
+		// Since we don't want to break them we need to create a new one on the first available slot
+		while(OutClothFacade.GetWeightMap(FName(WeightMapName)).Num() > 0)
+		{
+			WeightMapName = PropertyName;
+			WeightMapName.AppendInt(++WeightMapCount);
 		}
 		
 		// If the low high values of the merged property are the same we don't need to build a weight map
@@ -62,6 +75,8 @@ namespace UE::Chaos::ClothAsset::Private
 			{
 				Private::LogAndToastDifferentWeightMapNames(DataflowNode, PropertyName, InWeightMapName, OutWeightMapName, WeightMapName);	
 			}
+			MergedPropertyMaps.Add(WeightMapName,MergedProperty);
+			
 			// Create if necessary a new weight map
 			OutClothFacade.AddWeightMap(FName(WeightMapName));
 			TArrayView<float> WeightMap = OutClothFacade.GetWeightMap(FName(WeightMapName));
@@ -177,6 +192,7 @@ namespace UE::Chaos::ClothAsset::Private
 			  ::Chaos::Softs::FCollectionPropertyMutableFacade& OutPropertyFacade)
 	{
 		const int32 InNumInKeys = InPropertyFacade.Num();
+		TMap<FString,FMergedProperty> MergedPropertyMaps;
 		for (int32 InKeyIndex = 0; InKeyIndex < InNumInKeys; ++InKeyIndex)
 		{
 			PRAGMA_DISABLE_DEPRECATION_WARNINGS
@@ -208,7 +224,7 @@ namespace UE::Chaos::ClothAsset::Private
                     // We keep the string value to be the one in the output if defined
                     const FString WeightMapName = BuildWeightMaps(DataflowNode, InClothFacade, OutClothFacade,
                     	InPropertyBounds, OutPropertyBounds, PropertyBounds, InPropertyKey,
-							InPropertyFacade.GetStringValue(InKeyIndex), OutPropertyFacade.GetStringValue(OutKeyIndex));
+							InPropertyFacade.GetStringValue(InKeyIndex), OutPropertyFacade.GetStringValue(OutKeyIndex), MergedPropertyMaps);
 
 					OutPropertyFacade.SetStringValue(OutKeyIndex, WeightMapName);
 					bOverrideProperty = false;
