Cloth Editor: focus on tool-defined selection box when user hits 'F', not the entire cloth object

diff --git a/Engine/Plugins/Experimental/ChaosClothAssetEditor/Source/ChaosClothAssetEditor/Private/ChaosClothAsset/ClothEditorMode.cpp b/Engine/Plugins/Experimental/ChaosClothAssetEditor/Source/ChaosClothAssetEditor/Private/ChaosClothAsset/ClothEditorMode.cpp
--- a/Engine/Plugins/Experimental/ChaosClothAssetEditor/Source/ChaosClothAssetEditor/Private/ChaosClothAsset/ClothEditorMode.cpp
+++ b/Engine/Plugins/Experimental/ChaosClothAssetEditor/Source/ChaosClothAssetEditor/Private/ChaosClothAsset/ClothEditorMode.cpp
@@ -1155,6 +1155,17 @@ FBox UChaosClothAssetEditorMode::SceneBoundingBox() const
 
 FBox UChaosClothAssetEditorMode::SelectionBoundingBox() const
 {
+	// if Tool supports custom Focus box, use that first
+	if (GetToolManager()->HasAnyActiveTool())
+	{
+		UInteractiveTool* const Tool = GetToolManager()->GetActiveTool(EToolSide::Mouse);
+		IInteractiveToolCameraFocusAPI* const FocusAPI = Cast<IInteractiveToolCameraFocusAPI>(Tool);
+		if (FocusAPI && FocusAPI->SupportsWorldSpaceFocusBox())
+		{
+			return FocusAPI->GetWorldSpaceFocusBox();
+		}
+	}
+
 	const USelection* const SelectedComponents = GetModeManager()->GetSelectedComponents();
 
 	if (DynamicMeshComponent && SelectedComponents->IsSelected(DynamicMeshComponent))
diff --git a/Engine/Plugins/Experimental/ChaosClothAssetEditor/Source/ChaosClothAssetEditorTools/Private/ChaosClothAsset/ClothMeshSelectionTool.cpp b/Engine/Plugins/Experimental/ChaosClothAssetEditor/Source/ChaosClothAssetEditorTools/Private/ChaosClothAsset/ClothMeshSelectionTool.cpp
--- a/Engine/Plugins/Experimental/ChaosClothAssetEditor/Source/ChaosClothAssetEditorTools/Private/ChaosClothAsset/ClothMeshSelectionTool.cpp
+++ b/Engine/Plugins/Experimental/ChaosClothAssetEditor/Source/ChaosClothAssetEditorTools/Private/ChaosClothAsset/ClothMeshSelectionTool.cpp
@@ -348,6 +348,13 @@ bool UClothMeshSelectionTool::CanAccept() const
 	return bAnyChangeMade;
 }
 
+FBox UClothMeshSelectionTool::GetWorldSpaceFocusBox()
+{
+	static constexpr bool bWorld = true;
+	return FBox(SelectionMechanic->GetSelectionBounds(bWorld));
+}
+
+
 void UClothMeshSelectionTool::SetClothEditorContextObject(TObjectPtr<UClothEditorContextObject> InClothEditorContextObject)
 {
 	ClothEditorContextObject = InClothEditorContextObject;
diff --git a/Engine/Plugins/Experimental/ChaosClothAssetEditor/Source/ChaosClothAssetEditorTools/Public/ChaosClothAsset/ClothMeshSelectionTool.h b/Engine/Plugins/Experimental/ChaosClothAssetEditor/Source/ChaosClothAssetEditorTools/Public/ChaosClothAsset/ClothMeshSelectionTool.h
--- a/Engine/Plugins/Experimental/ChaosClothAssetEditor/Source/ChaosClothAssetEditorTools/Public/ChaosClothAsset/ClothMeshSelectionTool.h
+++ b/Engine/Plugins/Experimental/ChaosClothAssetEditor/Source/ChaosClothAssetEditorTools/Public/ChaosClothAsset/ClothMeshSelectionTool.h
@@ -144,6 +144,9 @@ private:
 	virtual bool HasAccept() const override { return true; }
 	virtual bool CanAccept() const override;
 
+	// IInteractiveToolCameraFocusAPI implementation
+	virtual FBox GetWorldSpaceFocusBox() override;
+
 	void SetClothEditorContextObject(TObjectPtr<UClothEditorContextObject> InClothEditorContextObject);
 	bool GetSelectedNodeInfo(FString& OutMapName, UE::Geometry::FGroupTopologySelection& OutSelection);
 	void UpdateSelectedNode();
