Chaos Cloth Editor: Bend and Stretch Override nodes.

This is a first pass of these nodes. Updates to improve them in a future CL, but this is what's necessary to get the merged asset working.

diff --git a/Engine/Plugins/Experimental/ChaosClothAssetEditor/Source/ChaosClothAssetDataflowNodes/Private/ChaosClothAsset/ChaosClothAssetDataflowNodesModule.cpp b/Engine/Plugins/Experimental/ChaosClothAssetEditor/Source/ChaosClothAssetDataflowNodes/Private/ChaosClothAsset/ChaosClothAssetDataflowNodesModule.cpp
--- a/Engine/Plugins/Experimental/ChaosClothAssetEditor/Source/ChaosClothAssetDataflowNodes/Private/ChaosClothAsset/ChaosClothAssetDataflowNodesModule.cpp
+++ b/Engine/Plugins/Experimental/ChaosClothAssetEditor/Source/ChaosClothAssetDataflowNodes/Private/ChaosClothAsset/ChaosClothAssetDataflowNodesModule.cpp
@@ -22,6 +22,7 @@
 #include "ChaosClothAsset/SimulationAnimDriveConfigNode.h"
 #include "ChaosClothAsset/SimulationBackstopConfigNode.h"
 #include "ChaosClothAsset/SimulationBendingConfigNode.h"
+#include "ChaosClothAsset/SimulationBendingOverrideConfigNode.h"
 #include "ChaosClothAsset/SimulationCollisionConfigNode.h"
 #include "ChaosClothAsset/SimulationDampingConfigNode.h"
 #include "ChaosClothAsset/SimulationDefaultConfigNode.h"
@@ -38,6 +39,7 @@
 #include "ChaosClothAsset/SimulationSelfCollisionSpheresConfigNode.h"
 #include "ChaosClothAsset/SimulationSolverConfigNode.h"
 #include "ChaosClothAsset/SimulationStretchConfigNode.h"
+#include "ChaosClothAsset/SimulationStretchOverrideConfigNode.h"
 #include "ChaosClothAsset/SimulationVelocityScaleConfigNode.h"
 #include "ChaosClothAsset/SimulationXPBDAreaSpringConfigNode.h"
 #include "ChaosClothAsset/SimulationXPBDAnisoBendingConfigNode.h"
@@ -86,7 +88,9 @@ namespace UE::Chaos::ClothAsset
 			DATAFLOW_NODE_REGISTER_CREATION_FACTORY(FChaosClothAssetSimulationBackstopConfigNode);
 			DATAFLOW_NODE_REGISTER_CREATION_FACTORY(FChaosClothAssetSimulationCollisionConfigNode);
 			DATAFLOW_NODE_REGISTER_CREATION_FACTORY(FChaosClothAssetSimulationBendingConfigNode);
+			DATAFLOW_NODE_REGISTER_CREATION_FACTORY(FChaosClothAssetSimulationBendingOverrideConfigNode);
 			DATAFLOW_NODE_REGISTER_CREATION_FACTORY(FChaosClothAssetSimulationStretchConfigNode);
+			DATAFLOW_NODE_REGISTER_CREATION_FACTORY(FChaosClothAssetSimulationStretchOverrideConfigNode);
 			DATAFLOW_NODE_REGISTER_CREATION_FACTORY(FChaosClothAssetSimulationDampingConfigNode);
 			DATAFLOW_NODE_REGISTER_CREATION_FACTORY(FChaosClothAssetSimulationDefaultConfigNode);
 			DATAFLOW_NODE_REGISTER_CREATION_FACTORY(FChaosClothAssetSimulationGravityConfigNode);
@@ -137,6 +141,7 @@ namespace UE::Chaos::ClothAsset
 				PropertyModule->RegisterCustomPropertyTypeLayout(FChaosClothAssetWeightedValueNonAnimatable::StaticStruct()->GetFName(), FOnGetPropertyTypeCustomizationInstance::CreateStatic(&FWeightedValueCustomization::MakeInstance));
 				PropertyModule->RegisterCustomPropertyTypeLayout(FChaosClothAssetWeightedValueNonAnimatableNoLowHighRange::StaticStruct()->GetFName(), FOnGetPropertyTypeCustomizationInstance::CreateStatic(&FWeightedValueCustomization::MakeInstance));
 				PropertyModule->RegisterCustomPropertyTypeLayout(FChaosClothAssetConnectableStringValue::StaticStruct()->GetFName(), FOnGetPropertyTypeCustomizationInstance::CreateStatic(&FWeightedValueCustomization::MakeInstance));
+				PropertyModule->RegisterCustomPropertyTypeLayout(FChaosClothAssetWeightedValueOverride::StaticStruct()->GetFName(), FOnGetPropertyTypeCustomizationInstance::CreateStatic(&FWeightedValueCustomization::MakeInstance));
 				PropertyModule->RegisterCustomPropertyTypeLayout(FChaosClothAssetConnectableIStringValue::StaticStruct()->GetFName(), FOnGetPropertyTypeCustomizationInstance::CreateStatic(&FWeightedValueCustomization::MakeInstance));
 				PropertyModule->RegisterCustomPropertyTypeLayout(FChaosClothAssetConnectableIOStringValue::StaticStruct()->GetFName(), FOnGetPropertyTypeCustomizationInstance::CreateStatic(&FWeightedValueCustomization::MakeInstance));
 				PropertyModule->RegisterCustomPropertyTypeLayout(FChaosClothAssetNodeSelectionGroup::StaticStruct()->GetFName(), FOnGetPropertyTypeCustomizationInstance::CreateStatic(&FSelectionGroupCustomization::MakeInstance));
@@ -158,6 +163,7 @@ namespace UE::Chaos::ClothAsset
 					PropertyModule->UnregisterCustomPropertyTypeLayout(FChaosClothAssetWeightedValue::StaticStruct()->GetFName());
 					PropertyModule->UnregisterCustomPropertyTypeLayout(FChaosClothAssetWeightedValueNonAnimatable::StaticStruct()->GetFName());
 					PropertyModule->UnregisterCustomPropertyTypeLayout(FChaosClothAssetWeightedValueNonAnimatableNoLowHighRange::StaticStruct()->GetFName());
+					PropertyModule->UnregisterCustomPropertyTypeLayout(FChaosClothAssetWeightedValueOverride::StaticStruct()->GetFName());
 					PropertyModule->UnregisterCustomPropertyTypeLayout(FChaosClothAssetConnectableIStringValue::StaticStruct()->GetFName());
 					PropertyModule->UnregisterCustomPropertyTypeLayout(FChaosClothAssetConnectableIOStringValue::StaticStruct()->GetFName());
 					PropertyModule->UnregisterCustomPropertyTypeLayout(FChaosClothAssetNodeSelectionGroup::StaticStruct()->GetFName());
diff --git a/Engine/Plugins/Experimental/ChaosClothAssetEditor/Source/ChaosClothAssetDataflowNodes/Private/ChaosClothAsset/SimulationBaseConfigNode.cpp b/Engine/Plugins/Experimental/ChaosClothAssetEditor/Source/ChaosClothAssetDataflowNodes/Private/ChaosClothAsset/SimulationBaseConfigNode.cpp
--- a/Engine/Plugins/Experimental/ChaosClothAssetEditor/Source/ChaosClothAssetDataflowNodes/Private/ChaosClothAsset/SimulationBaseConfigNode.cpp
+++ b/Engine/Plugins/Experimental/ChaosClothAssetEditor/Source/ChaosClothAssetDataflowNodes/Private/ChaosClothAsset/SimulationBaseConfigNode.cpp
@@ -96,7 +96,7 @@ int32 FChaosClothAssetSimulationBaseConfigNode::AddPropertyHelper(
 	{
 		KeyIndex = Properties.AddProperty(PropertyName.ToString());
 	}
-	else if (!Properties.IsLegacy(KeyIndex))  // Only warns of duplicates when the property hasn't been set using a legacy Chaos config
+	else if (bWarnDuplicateProperty && !Properties.IsLegacy(KeyIndex))  // Only warns of duplicates when the property hasn't been set using a legacy Chaos config
 	{
 		Private::LogAndToastDuplicateProperty(*this, PropertyName);
 	}
@@ -211,6 +211,64 @@ int32 FChaosClothAssetSimulationBaseConfigNode::FPropertyHelper::SetPropertyWeig
 		1.0f, PropertyValue.WeightMap, PropertyValue.WeightMap_Override, SimilarPropertyNames, PropertyFlags);
 }
 
+void FChaosClothAssetSimulationBaseConfigNode::FPropertyHelper::OverridePropertiesBool(const TArray<FName>& PropertyNames, bool bPropertyValue)
+{
+	for (const FName& PropertyName : PropertyNames)
+	{
+		const int32 PropertyKeyIndex = Properties.GetKeyIndex(PropertyName.ToString());
+		if (PropertyKeyIndex != INDEX_NONE)
+		{
+			Properties.SetValue(PropertyKeyIndex, bPropertyValue);
+		}
+	}
+}
+
+void FChaosClothAssetSimulationBaseConfigNode::FPropertyHelper::OverridePropertiesFloat(const TArray<FName>& PropertyNames, const EChaosClothAssetConstraintOverrideType OverrideType, const float OverrideValue)
+{
+	if (OverrideType == EChaosClothAssetConstraintOverrideType::None)
+	{
+		return;
+	}
+	for (const FName& PropertyName : PropertyNames)
+	{
+		const int32 PropertyKeyIndex = Properties.GetKeyIndex(PropertyName.ToString());
+		if (PropertyKeyIndex != INDEX_NONE)
+		{
+			if (OverrideType == EChaosClothAssetConstraintOverrideType::Override)
+			{
+				Properties.SetValue(PropertyKeyIndex, OverrideValue);
+			}
+			else if (OverrideType == EChaosClothAssetConstraintOverrideType::Multiply)
+			{
+				Properties.SetValue(PropertyKeyIndex, Properties.GetValue<float>(PropertyKeyIndex) * OverrideValue);
+			}
+		}
+	}
+}
+
+void FChaosClothAssetSimulationBaseConfigNode::FPropertyHelper::OverridePropertiesWeighted(const TArray<FName>& PropertyNames, const EChaosClothAssetConstraintOverrideType OverrideType, const FChaosClothAssetWeightedValueOverride& OverrideValue)
+{
+	if (OverrideType == EChaosClothAssetConstraintOverrideType::None)
+	{
+		return;
+	}
+	for (const FName& PropertyName : PropertyNames)
+	{
+		const int32 PropertyKeyIndex = Properties.GetKeyIndex(PropertyName.ToString());
+		if (PropertyKeyIndex != INDEX_NONE)
+		{
+			if (OverrideType == EChaosClothAssetConstraintOverrideType::Override)
+			{
+				Properties.SetWeightedValue(PropertyKeyIndex, OverrideValue.Low, OverrideValue.High);
+			}
+			else if (OverrideType == EChaosClothAssetConstraintOverrideType::Multiply)
+			{
+				Properties.SetWeightedValue(PropertyKeyIndex, Properties.GetLowValue<float>(PropertyKeyIndex) * OverrideValue.Low, Properties.GetHighValue<float>(PropertyKeyIndex) * OverrideValue.High);
+			}
+		}
+	}
+}
+
 int32 FChaosClothAssetSimulationBaseConfigNode::FPropertyHelper::SetPropertyString(
 	const FName& PropertyName,
 	const FChaosClothAssetConnectableIStringValue& PropertyValue,
diff --git a/Engine/Plugins/Experimental/ChaosClothAssetEditor/Source/ChaosClothAssetDataflowNodes/Private/ChaosClothAsset/SimulationBendingOverrideConfigNode.cpp b/Engine/Plugins/Experimental/ChaosClothAssetEditor/Source/ChaosClothAssetDataflowNodes/Private/ChaosClothAsset/SimulationBendingOverrideConfigNode.cpp
new file mode 100644
--- /dev/null
+++ b/Engine/Plugins/Experimental/ChaosClothAssetEditor/Source/ChaosClothAssetDataflowNodes/Private/ChaosClothAsset/SimulationBendingOverrideConfigNode.cpp
@@ -0,0 +1,126 @@
+// Copyright Epic Games, Inc. All Rights Reserved.
+
+#include "ChaosClothAsset/SimulationBendingOverrideConfigNode.h"
+#include "Chaos/CollectionPropertyFacade.h"
+
+#include UE_INLINE_GENERATED_CPP_BY_NAME(SimulationBendingOverrideConfigNode)
+
+FChaosClothAssetSimulationBendingOverrideConfigNode::FChaosClothAssetSimulationBendingOverrideConfigNode(const Dataflow::FNodeParameters& InParam, FGuid InGuid)
+	: FChaosClothAssetSimulationBaseConfigNode(InParam, InGuid)
+{
+	RegisterCollectionConnections();
+}
+
+void FChaosClothAssetSimulationBendingOverrideConfigNode::AddProperties(FPropertyHelper& PropertyHelper) const
+{
+	if (OverrideFlatnessRatio != EChaosClothAssetConstraintOverrideType::None)
+	{
+		PropertyHelper.OverridePropertiesWeighted({
+				FName(TEXT("FlatnessRatio")),
+				FName(TEXT("XPBDFlatnessRatio")),
+				FName(TEXT("XPBDAnisoFlatnessRatio")) }, OverrideFlatnessRatio, FlatnessRatio);
+	}
+	if (OverrideBendingStiffness != EChaosClothAssetConstraintOverrideType::None)
+	{
+		TArray<FName> BendingPropertyNames =
+		{
+			FName(TEXT("BendingSpringStiffness")),
+			FName(TEXT("BendingElementStiffness")),
+			FName(TEXT("XPBDBendingSpringStiffness")),
+			FName(TEXT("XPBDBendingElementStiffness"))
+		};
+		if (bApplyUniformBendingStiffnessOverride)
+		{
+			BendingPropertyNames.Append(
+				{
+					FName(TEXT("XPBDAnisoBendingStiffnessWarp")),
+					FName(TEXT("XPBDAnisoBendingStiffnessWeft")),
+					FName(TEXT("XPBDAnisoBendingStiffnessBias")),
+				});
+		}
+		if (bApplyBendingStiffnessOverrideToBuckling)
+		{
+			BendingPropertyNames.Append(
+				{
+					FName(TEXT("BucklingStiffness")),
+					FName(TEXT("XPBDBucklingStiffness"))
+				});
+
+			if (bApplyUniformBendingStiffnessOverride)
+			{
+				BendingPropertyNames.Append(
+					{
+						FName(TEXT("XPBDAnisoBucklingStiffnessWarp")),
+						FName(TEXT("XPBDAnisoBucklingStiffnessWeft")),
+						FName(TEXT("XPBDAnisoBucklingStiffnessBias")),
+					});
+			}
+		}
+		PropertyHelper.OverridePropertiesWeighted(BendingPropertyNames, OverrideBendingStiffness, BendingStiffness);
+	}
+
+	if (OverrideBucklingRatio != EChaosClothAssetConstraintOverrideType::None)
+	{
+		PropertyHelper.OverridePropertiesFloat({
+				FName(TEXT("BucklingRatio")),
+				FName(TEXT("XPBDBucklingRatio")),
+				FName(TEXT("XPBDAnisoBucklingRatio")) }, OverrideBucklingRatio, BucklingRatio);
+	}
+
+	if ((OverrideBendingStiffness == EChaosClothAssetConstraintOverrideType::None || !bApplyBendingStiffnessOverrideToBuckling) && OverrideBucklingStiffness != EChaosClothAssetConstraintOverrideType::None)
+	{
+		TArray<FName> BucklingPropertyNames =
+		{
+			FName(TEXT("BucklingStiffness")),
+			FName(TEXT("XPBDBucklingStiffness"))
+		};
+		if (bApplyUniformBucklingStiffnessOverride)
+		{
+			BucklingPropertyNames.Append(
+				{
+						FName(TEXT("XPBDAnisoBucklingStiffnessWarp")),
+						FName(TEXT("XPBDAnisoBucklingStiffnessWeft")),
+						FName(TEXT("XPBDAnisoBucklingStiffnessBias")),
+				});
+		}
+		PropertyHelper.OverridePropertiesWeighted(BucklingPropertyNames, OverrideBucklingStiffness, BucklingStiffness);
+	}
+
+	if (OverrideBendingStiffness == EChaosClothAssetConstraintOverrideType::None || !bApplyUniformBendingStiffnessOverride)
+	{
+		if (OverrideBendingStiffnessWarp != EChaosClothAssetConstraintOverrideType::None)
+		{
+			PropertyHelper.OverridePropertiesWeighted({ FName(TEXT("XPBDAnisoBendingStiffnessWarp"))}, OverrideBendingStiffnessWarp, BendingStiffnessWarp);
+		}
+		if (OverrideBendingStiffnessWeft != EChaosClothAssetConstraintOverrideType::None)
+		{
+			PropertyHelper.OverridePropertiesWeighted({ FName(TEXT("XPBDAnisoBendingStiffnessWeft")) }, OverrideBendingStiffnessWeft, BendingStiffnessWeft);
+		}
+		if (OverrideBendingStiffnessBias != EChaosClothAssetConstraintOverrideType::None)
+		{
+			PropertyHelper.OverridePropertiesWeighted({ FName(TEXT("XPBDAnisoBendingStiffnessBias")) }, OverrideBendingStiffnessBias, BendingStiffnessBias);
+		}
+	}
+	if ((OverrideBendingStiffness == EChaosClothAssetConstraintOverrideType::None || !bApplyUniformBendingStiffnessOverride || !bApplyBendingStiffnessOverrideToBuckling)
+		&& (OverrideBucklingStiffness == EChaosClothAssetConstraintOverrideType::None || !bApplyUniformBucklingStiffnessOverride))
+	{
+		if (OverrideBucklingStiffnessWarp != EChaosClothAssetConstraintOverrideType::None)
+		{
+			PropertyHelper.OverridePropertiesWeighted({ FName(TEXT("XPBDAnisoBucklingStiffnessWarp")) }, OverrideBucklingStiffnessWarp, BucklingStiffnessWarp);
+		}
+		if (OverrideBucklingStiffnessWeft != EChaosClothAssetConstraintOverrideType::None)
+		{
+			PropertyHelper.OverridePropertiesWeighted({ FName(TEXT("XPBDAnisoBucklingStiffnessWeft")) }, OverrideBucklingStiffnessWeft, BucklingStiffnessWeft);
+		}
+		if (OverrideBucklingStiffnessBias != EChaosClothAssetConstraintOverrideType::None)
+		{
+			PropertyHelper.OverridePropertiesWeighted({ FName(TEXT("XPBDAnisoBucklingStiffnessBias")) }, OverrideBucklingStiffnessBias, BucklingStiffnessBias);
+		}
+	}
+	if (OverrideBendingDamping != EChaosClothAssetConstraintOverrideType::None)
+	{
+		PropertyHelper.OverridePropertiesWeighted({
+				FName(TEXT("XPBDBendingSpringDamping")),
+				FName(TEXT("XPBDBendingElementDamping"))}, OverrideBendingDamping, BendingDamping);
+	}
+}
diff --git a/Engine/Plugins/Experimental/ChaosClothAssetEditor/Source/ChaosClothAssetDataflowNodes/Private/ChaosClothAsset/SimulationStretchOverrideConfigNode.cpp b/Engine/Plugins/Experimental/ChaosClothAssetEditor/Source/ChaosClothAssetDataflowNodes/Private/ChaosClothAsset/SimulationStretchOverrideConfigNode.cpp
new file mode 100644
--- /dev/null
+++ b/Engine/Plugins/Experimental/ChaosClothAssetEditor/Source/ChaosClothAssetDataflowNodes/Private/ChaosClothAsset/SimulationStretchOverrideConfigNode.cpp
@@ -0,0 +1,85 @@
+// Copyright Epic Games, Inc. All Rights Reserved.
+
+#include "ChaosClothAsset/SimulationStretchOverrideConfigNode.h"
+#include "Chaos/CollectionPropertyFacade.h"
+
+#include UE_INLINE_GENERATED_CPP_BY_NAME(SimulationStretchOverrideConfigNode)
+
+FChaosClothAssetSimulationStretchOverrideConfigNode::FChaosClothAssetSimulationStretchOverrideConfigNode(const Dataflow::FNodeParameters& InParam, FGuid InGuid)
+	: FChaosClothAssetSimulationBaseConfigNode(InParam, InGuid)
+{
+	RegisterCollectionConnections();
+}
+
+void FChaosClothAssetSimulationStretchOverrideConfigNode::AddProperties(FPropertyHelper& PropertyHelper) const
+{
+	if (bOverrideStretchUse3dRestLengths)
+	{
+		PropertyHelper.OverridePropertiesBool({
+			FName(TEXT("XPBDAnisoStretchUse3dRestLengths")),
+			FName(TEXT("XPBDAnisoSpringUse3dRestLengths")) }, bStretchUse3dRestLengths);
+	}
+	if (OverrideStretchStiffness != EChaosClothAssetConstraintOverrideType::None)
+	{
+		TArray<FName> StretchPropertyNames =
+		{
+			FName(TEXT("EdgeSpringStiffness")),
+			FName(TEXT("XPBDEdgeSpringStiffness")),
+			FName(TEXT("AreaSpringStiffness")),
+			FName(TEXT("XPBDAreaSpringStiffness"))
+		};
+		if (bApplyUniformStretchStiffnessOverride)
+		{
+			StretchPropertyNames.Append(
+				{
+					FName(TEXT("XPBDAnisoSpringStiffnessWarp")),
+					FName(TEXT("XPBDAnisoStretchStiffnessWarp")),
+					FName(TEXT("XPBDAnisoSpringStiffnessWeft")),
+					FName(TEXT("XPBDAnisoStretchStiffnessWeft")),
+					FName(TEXT("XPBDAnisoSpringStiffnessBias")),
+					FName(TEXT("XPBDAnisoStretchStiffnessBias"))
+				});
+		}
+		PropertyHelper.OverridePropertiesWeighted(StretchPropertyNames, OverrideStretchStiffness, StretchStiffness);
+	}
+	if (OverrideStretchStiffness == EChaosClothAssetConstraintOverrideType::None || !bApplyUniformStretchStiffnessOverride)
+	{
+		if (OverrideStretchStiffnessWarp != EChaosClothAssetConstraintOverrideType::None)
+		{
+			PropertyHelper.OverridePropertiesWeighted({
+					FName(TEXT("XPBDAnisoSpringStiffnessWarp")),
+					FName(TEXT("XPBDAnisoStretchStiffnessWarp"))}, OverrideStretchStiffnessWarp, StretchStiffnessWarp);
+		}
+		if (OverrideStretchStiffnessWeft != EChaosClothAssetConstraintOverrideType::None)
+		{
+			PropertyHelper.OverridePropertiesWeighted({
+					FName(TEXT("XPBDAnisoSpringStiffnessWeft")),
+					FName(TEXT("XPBDAnisoStretchStiffnessWeft")) }, OverrideStretchStiffnessWeft, StretchStiffnessWeft);
+		}
+		if (OverrideStretchStiffnessBias != EChaosClothAssetConstraintOverrideType::None)
+		{
+			PropertyHelper.OverridePropertiesWeighted({
+					FName(TEXT("XPBDAnisoSpringStiffnessBias")),
+					FName(TEXT("XPBDAnisoStretchStiffnessBias")) }, OverrideStretchStiffnessBias, StretchStiffnessBias);
+		}
+	}
+	if (OverrideStretchDamping != EChaosClothAssetConstraintOverrideType::None)
+	{
+		PropertyHelper.OverridePropertiesWeighted({
+				FName(TEXT("XPBDAnisoSpringDamping")),
+				FName(TEXT("XPBDAnisoStretchDamping")),
+				FName(TEXT("XPBDEdgeSpringDamping"))}, OverrideStretchDamping, StretchDamping);
+	}
+	if (OverrideWarpScale != EChaosClothAssetConstraintOverrideType::None)
+	{
+		PropertyHelper.OverridePropertiesWeighted({
+				FName(TEXT("XPBDAnisoSpringWarpScale")),
+				FName(TEXT("XPBDAnisoStretchWarpScale"))}, OverrideWarpScale, WarpScale);
+	}
+	if (OverrideWeftScale != EChaosClothAssetConstraintOverrideType::None)
+	{
+		PropertyHelper.OverridePropertiesWeighted({
+				FName(TEXT("XPBDAnisoSpringWeftScale")),
+				FName(TEXT("XPBDAnisoStretchWeftScale")) }, OverrideWeftScale, WeftScale);
+	}
+}
diff --git a/Engine/Plugins/Experimental/ChaosClothAssetEditor/Source/ChaosClothAssetDataflowNodes/Public/ChaosClothAsset/SimulationBaseConfigNode.h b/Engine/Plugins/Experimental/ChaosClothAssetEditor/Source/ChaosClothAssetDataflowNodes/Public/ChaosClothAsset/SimulationBaseConfigNode.h
--- a/Engine/Plugins/Experimental/ChaosClothAssetEditor/Source/ChaosClothAssetDataflowNodes/Public/ChaosClothAsset/SimulationBaseConfigNode.h
+++ b/Engine/Plugins/Experimental/ChaosClothAssetEditor/Source/ChaosClothAssetDataflowNodes/Public/ChaosClothAsset/SimulationBaseConfigNode.h
@@ -5,6 +5,7 @@
 #include "Chaos/CollectionPropertyFacade.h"
 #include "ChaosClothAsset/ConnectableValue.h"
 #include "ChaosClothAsset/WeightedValue.h"
+#include "ChaosClothAsset/SimulationConfigNodePropertyTypes.h"
 #include "Dataflow/DataflowNode.h"
 #include "GeometryCollection/ManagedArrayCollection.h"
 #include "SimulationBaseConfigNode.generated.h"
@@ -22,6 +23,10 @@ public:
 	UPROPERTY(Meta = (Dataflowinput, DataflowOutput, DataflowPassthrough = "Collection"))
 	FManagedArrayCollection Collection;
 
+	/** Warn when overriding an existing property in the collection.*/
+	UPROPERTY(EditAnywhere, Category = "Simulation Config")
+	bool bWarnDuplicateProperty = true;
+
 	FChaosClothAssetSimulationBaseConfigNode() = default;
 
 	FChaosClothAssetSimulationBaseConfigNode(const Dataflow::FNodeParameters& InParam, FGuid InGuid = FGuid::NewGuid());
@@ -63,7 +68,11 @@ protected:
 		int32 SetPropertyWeighted(const FName& PropertyName, const FChaosClothAssetWeightedValue& PropertyValue, const TArray<FName>& SimilarPropertyNames = {}, ECollectionPropertyFlags PropertyFlags = ECollectionPropertyFlags::None);
 		int32 SetPropertyWeighted(const FName& PropertyName, const FChaosClothAssetWeightedValueNonAnimatable& PropertyValue, const TArray<FName>& SimilarPropertyNames = {}, ECollectionPropertyFlags PropertyFlags = ECollectionPropertyFlags::None);
 		int32 SetPropertyWeighted(const FName& PropertyName, const FChaosClothAssetWeightedValueNonAnimatableNoLowHighRange& PropertyValue, const TArray<FName>& SimilarPropertyNames = {}, ECollectionPropertyFlags PropertyFlags = ECollectionPropertyFlags::None);
-		
+
+		void OverridePropertiesBool(const TArray<FName>& PropertyNames, bool bPropertyValue);
+		void OverridePropertiesFloat(const TArray<FName>& PropertyNames, const EChaosClothAssetConstraintOverrideType OverrideType, const float OverrideValue);
+		void OverridePropertiesWeighted(const TArray<FName>& PropertyNames, const EChaosClothAssetConstraintOverrideType OverrideType, const FChaosClothAssetWeightedValueOverride& OverrideValue);
+
 		template<typename T, typename WeightedValueType, TEMPLATE_REQUIRES(TIsDerivedFrom<T, FChaosClothAssetSimulationBaseConfigNode>::Value)>
 		inline int32 SetPropertyWeighted(const T* ConfigStruct, const WeightedValueType* PropertyValue, const TArray<FName>& SimilarPropertyNames = {}, ECollectionPropertyFlags PropertyFlags = ECollectionPropertyFlags::None);
 
diff --git a/Engine/Plugins/Experimental/ChaosClothAssetEditor/Source/ChaosClothAssetDataflowNodes/Public/ChaosClothAsset/SimulationBendingOverrideConfigNode.h b/Engine/Plugins/Experimental/ChaosClothAssetEditor/Source/ChaosClothAssetDataflowNodes/Public/ChaosClothAsset/SimulationBendingOverrideConfigNode.h
new file mode 100644
--- /dev/null
+++ b/Engine/Plugins/Experimental/ChaosClothAssetEditor/Source/ChaosClothAssetDataflowNodes/Public/ChaosClothAsset/SimulationBendingOverrideConfigNode.h
@@ -0,0 +1,123 @@
+// Copyright Epic Games, Inc. All Rights Reserved.
+
+#pragma once 
+
+#include "ChaosClothAsset/SimulationBaseConfigNode.h"
+#include "ChaosClothAsset/WeightedValue.h"
+#include "ChaosClothAsset/SimulationConfigNodePropertyTypes.h"
+#include "SimulationBendingOverrideConfigNode.generated.h"
+
+/** Bending constraint property override configuration node. */
+USTRUCT(Meta = (DataflowCloth))
+struct FChaosClothAssetSimulationBendingOverrideConfigNode : public FChaosClothAssetSimulationBaseConfigNode
+{
+	GENERATED_USTRUCT_BODY()
+	DATAFLOW_NODE_DEFINE_INTERNAL(FChaosClothAssetSimulationBendingOverrideConfigNode, "SimulationBendingOverrideConfig", "Cloth", "Cloth Simulation Bending Override Config")
+
+public:
+	
+	/** Flatness override type.*/
+	UPROPERTY(EditAnywhere, Category = "Bending Override")
+	EChaosClothAssetConstraintOverrideType OverrideFlatnessRatio = EChaosClothAssetConstraintOverrideType::None;
+
+	/** Flatness override value.*/
+	UPROPERTY(EditAnywhere, Category = "Bending Override", meta = (ClampMin = "0", EditCondition = "OverrideFlatnessRatio != EChaosClothAssetConstraintOverrideType::None", EditConditionHides))
+	FChaosClothAssetWeightedValueOverride FlatnessRatio;
+
+	/** Bending stiffness override type.*/
+	UPROPERTY(EditAnywhere, Category = "Bending Override")
+	EChaosClothAssetConstraintOverrideType OverrideBendingStiffness = EChaosClothAssetConstraintOverrideType::None;
+
+	/** Bending stiffness override value.*/
+	UPROPERTY(EditAnywhere, Category = "Bending Override", meta = (ClampMin = "0", EditCondition = "OverrideBendingStiffness != EChaosClothAssetConstraintOverrideType::None", EditConditionHides))
+	FChaosClothAssetWeightedValueOverride BendingStiffness;
+
+	/** Whether or not to apply the Bending Stiffness Override to warp, weft, and bias stiffnesses of anisotropic bending elements.*/
+	UPROPERTY(EditAnywhere, Category = "Bending Override", meta = (EditCondition = "OverrideBendingStiffness != EChaosClothAssetConstraintOverrideType::None", EditConditionHides))
+	bool bApplyUniformBendingStiffnessOverride = true;
+
+	/** Whether or not to apply the Bending Stiffness Override to buckling stiffnesses.*/
+	UPROPERTY(EditAnywhere, Category = "Bending Override", meta = (EditCondition = "OverrideBendingStiffness != EChaosClothAssetConstraintOverrideType::None", EditConditionHides))
+	bool bApplyBendingStiffnessOverrideToBuckling = true;
+
+	/** Buckling ratio override type.*/
+	UPROPERTY(EditAnywhere, Category = "Bending Override")
+	EChaosClothAssetConstraintOverrideType OverrideBucklingRatio = EChaosClothAssetConstraintOverrideType::None;
+
+	/** Buckling ratio override value.*/
+	UPROPERTY(EditAnywhere, Category = "Bending Override", meta = (ClampMin = "0", EditCondition = "OverrideBucklingRatio != EChaosClothAssetConstraintOverrideType::None", EditConditionHides))
+	float BucklingRatio = 1.0;
+
+	/** Buckling stiffness override type.*/
+	UPROPERTY(EditAnywhere, Category = "Bending Override", meta = (EditCondition = "OverrideBendingStiffness == EChaosClothAssetConstraintOverrideType::None || !bApplyBendingStiffnessOverrideToBuckling", EditConditionHides))
+	EChaosClothAssetConstraintOverrideType OverrideBucklingStiffness = EChaosClothAssetConstraintOverrideType::None;
+
+	/** Buckling stiffness override value.*/
+	UPROPERTY(EditAnywhere, Category = "Bending Override", meta = (ClampMin = "0", EditCondition = "(OverrideBendingStiffness == EChaosClothAssetConstraintOverrideType::None || !bApplyBendingStiffnessOverrideToBuckling) && OverrideBucklingStiffness != EChaosClothAssetConstraintOverrideType::None", EditConditionHides))
+	FChaosClothAssetWeightedValueOverride BucklingStiffness;
+
+	/** Whether or not to apply the Buckling Stiffness Override to warp, weft, and bias stiffnesses of anisotropic bending elements.*/
+	UPROPERTY(EditAnywhere, Category = "Bending Override", meta = (EditCondition = "(OverrideBendingStiffness == EChaosClothAssetConstraintOverrideType::None || !bApplyBendingStiffnessOverrideToBuckling) && OverrideBucklingStiffness != EChaosClothAssetConstraintOverrideType::None", EditConditionHides))
+	bool bApplyUniformBucklingStiffnessOverride = true;
+
+	/** Warp stiffness override type.*/
+	UPROPERTY(EditAnywhere, Category = "Bending Override", meta = (EditCondition = "OverrideBendingStiffness == EChaosClothAssetConstraintOverrideType::None || !bApplyUniformBendingStiffnessOverride", EditConditionHides))
+	EChaosClothAssetConstraintOverrideType OverrideBendingStiffnessWarp = EChaosClothAssetConstraintOverrideType::None;
+
+	/** Bending stiffness warp override value.*/
+	UPROPERTY(EditAnywhere, Category = "Bending Override", meta = (ClampMin = "0", EditCondition = "(OverrideBendingStiffness == EChaosClothAssetConstraintOverrideType::None || !bApplyUniformBendingStiffnessOverride) && OverrideBendingStiffnessWarp != EChaosClothAssetConstraintOverrideType::None", EditConditionHides))
+	FChaosClothAssetWeightedValueOverride BendingStiffnessWarp;
+
+	/** Weft stiffness override type.*/
+	UPROPERTY(EditAnywhere, Category = "Bending Override", meta = (EditCondition = "OverrideBendingStiffness == EChaosClothAssetConstraintOverrideType::None || !bApplyUniformBendingStiffnessOverride", EditConditionHides))
+	EChaosClothAssetConstraintOverrideType OverrideBendingStiffnessWeft = EChaosClothAssetConstraintOverrideType::None;
+
+	/** Bending stiffness weft override value.*/
+	UPROPERTY(EditAnywhere, Category = "Bending Override", meta = (ClampMin = "0", EditCondition = "(OverrideBendingStiffness == EChaosClothAssetConstraintOverrideType::None || !bApplyUniformBendingStiffnessOverride) && OverrideBendingStiffnessWeft != EChaosClothAssetConstraintOverrideType::None", EditConditionHides))
+	FChaosClothAssetWeightedValueOverride BendingStiffnessWeft;
+
+	/** Bias stiffness override type.*/
+	UPROPERTY(EditAnywhere, Category = "Bending Override", meta = (EditCondition = "OverrideBendingStiffness == EChaosClothAssetConstraintOverrideType::None || !bApplyUniformBendingStiffnessOverride", EditConditionHides))
+	EChaosClothAssetConstraintOverrideType OverrideBendingStiffnessBias = EChaosClothAssetConstraintOverrideType::None;
+
+	/** Bending stiffness bias override value.*/
+	UPROPERTY(EditAnywhere, Category = "Bending Override", meta = (ClampMin = "0", EditCondition = "(OverrideBendingStiffness == EChaosClothAssetConstraintOverrideType::None || !bApplyUniformBendingStiffnessOverride) && OverrideBendingStiffnessBias != EChaosClothAssetConstraintOverrideType::None", EditConditionHides))
+	FChaosClothAssetWeightedValueOverride BendingStiffnessBias;
+
+	/** Warp buckling stiffness override type.*/
+	UPROPERTY(EditAnywhere, Category = "Bending Override", meta = (EditCondition = "(OverrideBendingStiffness == EChaosClothAssetConstraintOverrideType::None || !bApplyUniformBendingStiffnessOverride || !bApplyBendingStiffnessOverrideToBuckling) && (OverrideBucklingStiffness == EChaosClothAssetConstraintOverrideType::None || !bApplyUniformBucklingStiffnessOverride)", EditConditionHides))
+	EChaosClothAssetConstraintOverrideType OverrideBucklingStiffnessWarp = EChaosClothAssetConstraintOverrideType::None;
+
+	/** Buckling stiffness warp override value.*/
+	UPROPERTY(EditAnywhere, Category = "Bending Override", meta = (ClampMin = "0", EditCondition = "(OverrideBendingStiffness == EChaosClothAssetConstraintOverrideType::None || !bApplyUniformBendingStiffnessOverride || !bApplyBendingStiffnessOverrideToBuckling) && (OverrideBucklingStiffness == EChaosClothAssetConstraintOverrideType::None || !bApplyUniformBucklingStiffnessOverride) && OverrideBucklingStiffnessWarp != EChaosClothAssetConstraintOverrideType::None", EditConditionHides))
+	FChaosClothAssetWeightedValueOverride BucklingStiffnessWarp;
+	
+	/** Weft buckling stiffness override type.*/
+	UPROPERTY(EditAnywhere, Category = "Bending Override", meta = (EditCondition = "(OverrideBendingStiffness == EChaosClothAssetConstraintOverrideType::None || !bApplyUniformBendingStiffnessOverride || !bApplyBendingStiffnessOverrideToBuckling) && (OverrideBucklingStiffness == EChaosClothAssetConstraintOverrideType::None || !bApplyUniformBucklingStiffnessOverride)", EditConditionHides))
+	EChaosClothAssetConstraintOverrideType OverrideBucklingStiffnessWeft = EChaosClothAssetConstraintOverrideType::None;
+
+	/** Buckling stiffness Weft override value.*/
+	UPROPERTY(EditAnywhere, Category = "Bending Override", meta = (ClampMin = "0", EditCondition = "(OverrideBendingStiffness == EChaosClothAssetConstraintOverrideType::None || !bApplyUniformBendingStiffnessOverride || !bApplyBendingStiffnessOverrideToBuckling) && (OverrideBucklingStiffness == EChaosClothAssetConstraintOverrideType::None || !bApplyUniformBucklingStiffnessOverride) && OverrideBucklingStiffnessWeft != EChaosClothAssetConstraintOverrideType::None", EditConditionHides))
+	FChaosClothAssetWeightedValueOverride BucklingStiffnessWeft;
+
+	/** Bias buckling stiffness override type.*/
+	UPROPERTY(EditAnywhere, Category = "Bending Override", meta = (EditCondition = "(OverrideBendingStiffness == EChaosClothAssetConstraintOverrideType::None || !bApplyUniformBendingStiffnessOverride || !bApplyBendingStiffnessOverrideToBuckling) && (OverrideBucklingStiffness == EChaosClothAssetConstraintOverrideType::None || !bApplyUniformBucklingStiffnessOverride)", EditConditionHides))
+	EChaosClothAssetConstraintOverrideType OverrideBucklingStiffnessBias = EChaosClothAssetConstraintOverrideType::None;
+
+	/** Buckling stiffness Bias override value.*/
+	UPROPERTY(EditAnywhere, Category = "Bending Override", meta = (ClampMin = "0", EditCondition = "(OverrideBendingStiffness == EChaosClothAssetConstraintOverrideType::None || !bApplyUniformBendingStiffnessOverride || !bApplyBendingStiffnessOverrideToBuckling) && (OverrideBucklingStiffness == EChaosClothAssetConstraintOverrideType::None || !bApplyUniformBucklingStiffnessOverride) && OverrideBucklingStiffnessBias != EChaosClothAssetConstraintOverrideType::None", EditConditionHides))
+	FChaosClothAssetWeightedValueOverride BucklingStiffnessBias;
+
+	/** Damping override type.*/
+	UPROPERTY(EditAnywhere, Category = "Bending Override")
+	EChaosClothAssetConstraintOverrideType OverrideBendingDamping = EChaosClothAssetConstraintOverrideType::None;
+
+	/** Bending damping override value.*/
+	UPROPERTY(EditAnywhere, Category = "Bending Override", meta = (ClampMin = "0", EditCondition = "OverrideBendingDamping != EChaosClothAssetConstraintOverrideType::None", EditConditionHides))
+	FChaosClothAssetWeightedValueOverride BendingDamping;
+
+	FChaosClothAssetSimulationBendingOverrideConfigNode(const Dataflow::FNodeParameters& InParam, FGuid InGuid = FGuid::NewGuid());
+
+private:
+	virtual void AddProperties(FPropertyHelper& PropertyHelper) const override;
+};
diff --git a/Engine/Plugins/Experimental/ChaosClothAssetEditor/Source/ChaosClothAssetDataflowNodes/Public/ChaosClothAsset/SimulationConfigNodePropertyTypes.h b/Engine/Plugins/Experimental/ChaosClothAssetEditor/Source/ChaosClothAssetDataflowNodes/Public/ChaosClothAsset/SimulationConfigNodePropertyTypes.h
--- a/Engine/Plugins/Experimental/ChaosClothAssetEditor/Source/ChaosClothAssetDataflowNodes/Public/ChaosClothAsset/SimulationConfigNodePropertyTypes.h
+++ b/Engine/Plugins/Experimental/ChaosClothAssetEditor/Source/ChaosClothAssetDataflowNodes/Public/ChaosClothAsset/SimulationConfigNodePropertyTypes.h
@@ -44,3 +44,14 @@ enum class EChaosClothAssetConstraintSolverType : uint8
 	/** PBD style with no damping/stiffness */
 	PBD,
 };
+
+UENUM()
+enum class EChaosClothAssetConstraintOverrideType : uint8
+{
+	/** Do not override */
+	None,
+	/** Override any existing values with this new value.*/
+	Override,
+	/** Multiply any existing values with this new value.*/
+	Multiply
+};
diff --git a/Engine/Plugins/Experimental/ChaosClothAssetEditor/Source/ChaosClothAssetDataflowNodes/Public/ChaosClothAsset/SimulationStretchOverrideConfigNode.h b/Engine/Plugins/Experimental/ChaosClothAssetEditor/Source/ChaosClothAssetDataflowNodes/Public/ChaosClothAsset/SimulationStretchOverrideConfigNode.h
new file mode 100644
--- /dev/null
+++ b/Engine/Plugins/Experimental/ChaosClothAssetEditor/Source/ChaosClothAssetDataflowNodes/Public/ChaosClothAsset/SimulationStretchOverrideConfigNode.h
@@ -0,0 +1,92 @@
+// Copyright Epic Games, Inc. All Rights Reserved.
+
+#pragma once 
+
+#include "ChaosClothAsset/SimulationBaseConfigNode.h"
+#include "ChaosClothAsset/WeightedValue.h"
+#include "ChaosClothAsset/SimulationConfigNodePropertyTypes.h"
+#include "SimulationStretchOverrideConfigNode.generated.h"
+
+/** Stretching constraint property override configuration node. */
+USTRUCT(Meta = (DataflowCloth))
+struct FChaosClothAssetSimulationStretchOverrideConfigNode : public FChaosClothAssetSimulationBaseConfigNode
+{
+	GENERATED_USTRUCT_BODY()
+	DATAFLOW_NODE_DEFINE_INTERNAL(FChaosClothAssetSimulationStretchOverrideConfigNode, "SimulationStretchOverrideConfig", "Cloth", "Cloth Simulation Stretching Override Config")
+
+public:
+
+	/** Enable overriding Stretch Use 3d Rest Lengths */
+	UPROPERTY(EditAnywhere, Category = "Stretch Override", meta = (InlineEditConditionToggle))
+	bool bOverrideStretchUse3dRestLengths = false;
+
+	/** Whether to use the 3D draped space as rest lengths, or use the 2D pattern space instead. */
+	UPROPERTY(EditAnywhere, Category = "Stretch Override", meta = (EditCondition = "bOverrideStretchUse3dRestLengths"))
+	bool bStretchUse3dRestLengths = true;
+
+	/** Stretch stiffness override type.*/
+	UPROPERTY(EditAnywhere, Category = "Stretch Override")
+	EChaosClothAssetConstraintOverrideType OverrideStretchStiffness = EChaosClothAssetConstraintOverrideType::None;
+
+	/** Stretch stiffness override value.*/
+	UPROPERTY(EditAnywhere, Category = "Stretch Override", meta = (ClampMin = "0", EditCondition = "OverrideStretchStiffness != EChaosClothAssetConstraintOverrideType::None", EditConditionHides))
+	FChaosClothAssetWeightedValueOverride StretchStiffness;
+
+	/** Whether or not to apply the Stretch Stiffness Override to warp, weft, and bias stiffnesses of anisotropic springs.*/
+	UPROPERTY(EditAnywhere, Category = "Stretch Override", meta = (EditCondition = "OverrideStretchStiffness != EChaosClothAssetConstraintOverrideType::None", EditConditionHides))
+	bool bApplyUniformStretchStiffnessOverride = true;
+
+	/** Warp stiffness override type.*/
+	UPROPERTY(EditAnywhere, Category = "Stretch Override", meta = (EditCondition = "OverrideStretchStiffness == EChaosClothAssetConstraintOverrideType::None || !bApplyUniformStretchStiffnessOverride", EditConditionHides))
+	EChaosClothAssetConstraintOverrideType OverrideStretchStiffnessWarp = EChaosClothAssetConstraintOverrideType::None;
+
+	/** Stretch stiffness warp override value.*/
+	UPROPERTY(EditAnywhere, Category = "Stretch Override", meta = (ClampMin = "0", EditCondition = "(OverrideStretchStiffness == EChaosClothAssetConstraintOverrideType::None || !bApplyUniformStretchStiffnessOverride) && OverrideStretchStiffnessWarp != EChaosClothAssetConstraintOverrideType::None", EditConditionHides))
+	FChaosClothAssetWeightedValueOverride StretchStiffnessWarp;
+
+	/** Weft stiffness override type.*/
+	UPROPERTY(EditAnywhere, Category = "Stretch Override", meta = (EditCondition = "OverrideStretchStiffness == EChaosClothAssetConstraintOverrideType::None || !bApplyUniformStretchStiffnessOverride", EditConditionHides))
+	EChaosClothAssetConstraintOverrideType OverrideStretchStiffnessWeft = EChaosClothAssetConstraintOverrideType::None;
+
+	/** Stretch stiffness weft override value.*/
+	UPROPERTY(EditAnywhere, Category = "Stretch Override", meta = (ClampMin = "0", EditCondition = "(OverrideStretchStiffness == EChaosClothAssetConstraintOverrideType::None || !bApplyUniformStretchStiffnessOverride) && OverrideStretchStiffnessWeft != EChaosClothAssetConstraintOverrideType::None", EditConditionHides))
+	FChaosClothAssetWeightedValueOverride StretchStiffnessWeft;
+
+	/** Bias stiffness override type.*/
+	UPROPERTY(EditAnywhere, Category = "Stretch Override", meta = (EditCondition = "OverrideStretchStiffness == EChaosClothAssetConstraintOverrideType::None || !bApplyUniformStretchStiffnessOverride", EditConditionHides))
+	EChaosClothAssetConstraintOverrideType OverrideStretchStiffnessBias = EChaosClothAssetConstraintOverrideType::None;
+
+	/** Stretch stiffness bias override value.*/
+	UPROPERTY(EditAnywhere, Category = "Stretch Override", meta = (ClampMin = "0", EditCondition = "(OverrideStretchStiffness == EChaosClothAssetConstraintOverrideType::None || !bApplyUniformStretchStiffnessOverride) && OverrideStretchStiffnessBias != EChaosClothAssetConstraintOverrideType::None", EditConditionHides))
+	FChaosClothAssetWeightedValueOverride StretchStiffnessBias;
+
+	/** Damping override type.*/
+	UPROPERTY(EditAnywhere, Category = "Stretch Override")
+	EChaosClothAssetConstraintOverrideType OverrideStretchDamping = EChaosClothAssetConstraintOverrideType::None;
+
+	/** Stretch damping override value.*/
+	UPROPERTY(EditAnywhere, Category = "Stretch Override", meta = (ClampMin = "0", EditCondition = "OverrideStretchDamping != EChaosClothAssetConstraintOverrideType::None", EditConditionHides))
+	FChaosClothAssetWeightedValueOverride StretchDamping;
+
+	/** Warp scale override type.*/
+	UPROPERTY(EditAnywhere, Category = "Stretch Override")
+	EChaosClothAssetConstraintOverrideType OverrideWarpScale = EChaosClothAssetConstraintOverrideType::None;
+
+	/** Stretch damping override value.*/
+	UPROPERTY(EditAnywhere, Category = "Stretch Override", meta = (ClampMin = "0", EditCondition = "OverrideWarpScale != EChaosClothAssetConstraintOverrideType::None", EditConditionHides))
+	FChaosClothAssetWeightedValueOverride WarpScale;
+
+	/** Weft scale override type.*/
+	UPROPERTY(EditAnywhere, Category = "Stretch Override")
+	EChaosClothAssetConstraintOverrideType OverrideWeftScale = EChaosClothAssetConstraintOverrideType::None;
+
+	/** Stretch damping override value.*/
+	UPROPERTY(EditAnywhere, Category = "Stretch Override", meta = (ClampMin = "0", EditCondition = "OverrideWeftScale != EChaosClothAssetConstraintOverrideType::None", EditConditionHides))
+	FChaosClothAssetWeightedValueOverride WeftScale;
+
+
+	FChaosClothAssetSimulationStretchOverrideConfigNode(const Dataflow::FNodeParameters& InParam, FGuid InGuid = FGuid::NewGuid());
+
+private:
+	virtual void AddProperties(FPropertyHelper& PropertyHelper) const override;
+};
diff --git a/Engine/Plugins/Experimental/ChaosClothAssetEditor/Source/ChaosClothAssetDataflowNodes/Public/ChaosClothAsset/WeightedValue.h b/Engine/Plugins/Experimental/ChaosClothAssetEditor/Source/ChaosClothAssetDataflowNodes/Public/ChaosClothAsset/WeightedValue.h
--- a/Engine/Plugins/Experimental/ChaosClothAssetEditor/Source/ChaosClothAssetDataflowNodes/Public/ChaosClothAsset/WeightedValue.h
+++ b/Engine/Plugins/Experimental/ChaosClothAssetEditor/Source/ChaosClothAssetDataflowNodes/Public/ChaosClothAsset/WeightedValue.h
@@ -156,4 +156,22 @@ struct FChaosClothAssetWeightedValueNonAnimatableNoLowHighRange
 	/** The weight map override value for when the WeightMap has a connection that replaces the provided weight map value. */
 	UPROPERTY(VisibleAnywhere, Transient, Category = "Weighted Value")
 	mutable FString WeightMap_Override = UE::Chaos::ClothAsset::FWeightMapTools::NotOverridden;  // _Override has a special meaning to the property customization, mutable because this property is set while getting the original value
-};
\ No newline at end of file
+};
+
+USTRUCT()
+struct FChaosClothAssetWeightedValueOverride
+{
+	GENERATED_BODY()
+
+	/**
+	 * Low value override.
+	 */
+	UPROPERTY(EditAnywhere, Category = "Weighted Value", Meta = (DisplayName = "Low Weight", ChaosClothAssetShortName = "Lo"))
+	float Low = 1.f;
+
+	/**
+	 * High value override
+	 */
+	UPROPERTY(EditAnywhere, Category = "Weighted Value", Meta = (DisplayName = "High Weight", ChaosClothAssetShortName = "Hi"))
+	float High = 1.f;
+};
