Cloth selection tool: fix back face material checkerboard

diff --git a/Engine/Plugins/Experimental/ChaosClothAssetEditor/Source/ChaosClothAssetEditorTools/Private/ChaosClothAsset/ClothMeshSelectionTool.cpp b/Engine/Plugins/Experimental/ChaosClothAssetEditor/Source/ChaosClothAssetEditorTools/Private/ChaosClothAsset/ClothMeshSelectionTool.cpp
--- a/Engine/Plugins/Experimental/ChaosClothAssetEditor/Source/ChaosClothAssetEditorTools/Private/ChaosClothAsset/ClothMeshSelectionTool.cpp
+++ b/Engine/Plugins/Experimental/ChaosClothAssetEditor/Source/ChaosClothAssetEditorTools/Private/ChaosClothAsset/ClothMeshSelectionTool.cpp
@@ -22,6 +22,7 @@
 #include "Dataflow/DataflowEdNode.h"
 #include "DynamicMesh/NonManifoldMappingSupport.h"
 #include "Selections/GeometrySelectionUtil.h"
+#include "Materials/Material.h"
 
 #include UE_INLINE_GENERATED_CPP_BY_NAME(ClothMeshSelectionTool)
 
@@ -109,10 +110,15 @@ void UClothMeshSelectionTool::Setup()
 	PreviewMesh->SetMaterials(MaterialSet.Materials);
 
 	// configure secondary render material for selected triangles
-	UMaterialInterface* SelectionMaterial = ToolSetupUtil::GetSelectionMaterial(FLinearColor::Yellow, GetToolManager());
-	if (SelectionMaterial != nullptr)
+	// NOTE: the material returned by ToolSetupUtil::GetSelectionMaterial has a checkerboard pattern on back faces which makes it hard to use
+	UMaterialInterface* Material = LoadObject<UMaterial>(nullptr, TEXT("/MeshModelingToolsetExp/Materials/SculptMaterial"));
+	if (Material != nullptr)
 	{
-		PreviewMesh->SetSecondaryRenderMaterial(SelectionMaterial);
+		if (UMaterialInstanceDynamic* MatInstance = UMaterialInstanceDynamic::Create(Material, GetToolManager()))
+		{
+			MatInstance->SetVectorParameterValue(TEXT("Color"), FLinearColor::Yellow);
+			PreviewMesh->SetSecondaryRenderMaterial(MatInstance);
+		}
 	}
 
 	PreviewMesh->SetTangentsMode(EDynamicMeshComponentTangentsMode::AutoCalculated);
