Fix crash when closing cloth editor while selection tool is running

diff --git a/Engine/Plugins/Experimental/ChaosClothAssetEditor/Source/ChaosClothAssetEditor/Private/ChaosClothAsset/ClothEditorToolkit.cpp b/Engine/Plugins/Experimental/ChaosClothAssetEditor/Source/ChaosClothAssetEditor/Private/ChaosClothAsset/ClothEditorToolkit.cpp
--- a/Engine/Plugins/Experimental/ChaosClothAssetEditor/Source/ChaosClothAssetEditor/Private/ChaosClothAsset/ClothEditorToolkit.cpp
+++ b/Engine/Plugins/Experimental/ChaosClothAssetEditor/Source/ChaosClothAssetEditor/Private/ChaosClothAsset/ClothEditorToolkit.cpp
@@ -207,6 +207,7 @@ FChaosClothAssetEditorToolkit::FChaosClothAssetEditorToolkit(UAssetEditor* InOwn
 
 FChaosClothAssetEditorToolkit::~FChaosClothAssetEditorToolkit()
 {
+	// This code is also called in OnRequestClose. Leaving this here as well in case the toolkit gets destroyed without having OnRequestClose called
 	if (SelectedDataflowNode && OnNodeInvalidatedDelegateHandle.IsValid())
 	{
 		SelectedDataflowNode->GetOnNodeInvalidatedDelegate().Remove(OnNodeInvalidatedDelegateHandle);
@@ -380,6 +381,13 @@ bool FChaosClothAssetEditorToolkit::OnRequestClose(EAssetEditorCloseReason InClo
 		return true;
 	}
 
+	if (SelectedDataflowNode && OnNodeInvalidatedDelegateHandle.IsValid())
+	{
+		SelectedDataflowNode->GetOnNodeInvalidatedDelegate().Remove(OnNodeInvalidatedDelegateHandle);
+		SelectedDataflowNode->OnDeselected();
+	}
+	SelectedDataflowNode.Reset();
+
 	return FAssetEditorToolkit::OnRequestClose(InCloseReason);
 }
 
