Chaos Cloth Editor: Cleanup mesh after deleting elements.

diff --git a/Engine/Plugins/Experimental/ChaosClothAssetEditor/Source/ChaosClothAssetDataflowNodes/Private/ChaosClothAsset/DeleteElementNode.cpp b/Engine/Plugins/Experimental/ChaosClothAssetEditor/Source/ChaosClothAssetDataflowNodes/Private/ChaosClothAsset/DeleteElementNode.cpp
--- a/Engine/Plugins/Experimental/ChaosClothAssetEditor/Source/ChaosClothAssetDataflowNodes/Private/ChaosClothAsset/DeleteElementNode.cpp
+++ b/Engine/Plugins/Experimental/ChaosClothAssetEditor/Source/ChaosClothAssetDataflowNodes/Private/ChaosClothAsset/DeleteElementNode.cpp
@@ -366,6 +366,8 @@ void FChaosClothAssetDeleteElementNode::Evaluate(Dataflow::FContext& Context, co
 			{
 				FClothGeometryTools::DeleteRenderMesh(ClothCollection);
 			}
+
+			FClothGeometryTools::CleanupAndCompactMesh(ClothCollection);
 		}
 		SetValue(Context, MoveTemp(*ClothCollection), &Collection);
 	}
