Chaos Cloth: add option to only collide against kinematic colliders.

diff --git a/Engine/Plugins/Experimental/ChaosClothAssetEditor/Source/ChaosClothAssetDataflowNodes/Private/ChaosClothAsset/SimulationSelfCollisionConfigNode.cpp b/Engine/Plugins/Experimental/ChaosClothAssetEditor/Source/ChaosClothAssetDataflowNodes/Private/ChaosClothAsset/SimulationSelfCollisionConfigNode.cpp
--- a/Engine/Plugins/Experimental/ChaosClothAssetEditor/Source/ChaosClothAssetDataflowNodes/Private/ChaosClothAsset/SimulationSelfCollisionConfigNode.cpp
+++ b/Engine/Plugins/Experimental/ChaosClothAssetEditor/Source/ChaosClothAssetDataflowNodes/Private/ChaosClothAsset/SimulationSelfCollisionConfigNode.cpp
@@ -25,6 +25,7 @@ void FChaosClothAssetSimulationSelfCollisionConfigNode::AddProperties(FPropertyH
 	PropertyHelper.SetProperty(this, &SelfCollisionDisableNeighborDistance, {}, ECollectionPropertyFlags::None); // Non animatable
 	PropertyHelper.SetPropertyString(this, &SelfCollisionLayers);
 	PropertyHelper.SetPropertyString(this, &SelfCollisionDisabledFaces);
+	PropertyHelper.SetPropertyBool(this, &bSelfCollideAgainstKinematicCollidersOnly);
 	PropertyHelper.SetPropertyBool(this, &bSelfCollideAgainstAllKinematicVertices);
 	PropertyHelper.SetPropertyString(this, &SelfCollisionEnabledKinematicFaces);
 	PropertyHelper.SetProperty(this, &SelfCollisionKinematicColliderThickness);
diff --git a/Engine/Plugins/Experimental/ChaosClothAssetEditor/Source/ChaosClothAssetDataflowNodes/Public/ChaosClothAsset/SimulationSelfCollisionConfigNode.h b/Engine/Plugins/Experimental/ChaosClothAssetEditor/Source/ChaosClothAssetDataflowNodes/Public/ChaosClothAsset/SimulationSelfCollisionConfigNode.h
--- a/Engine/Plugins/Experimental/ChaosClothAssetEditor/Source/ChaosClothAssetDataflowNodes/Public/ChaosClothAsset/SimulationSelfCollisionConfigNode.h
+++ b/Engine/Plugins/Experimental/ChaosClothAssetEditor/Source/ChaosClothAssetDataflowNodes/Public/ChaosClothAsset/SimulationSelfCollisionConfigNode.h
@@ -19,31 +19,31 @@ public:
 	float SelfCollisionThickness = 0.5f;
 
 	/** The stiffness of the springs used to control self collision (PBD Solver). */
-	UPROPERTY(EditAnywhere, Category = "Self-Collision Properties", meta = (UIMin = "0", UIMax = "1", ClampMin = "0", ClampMax = "1"))
+	UPROPERTY(EditAnywhere, Category = "Self-Collision Properties", meta = (UIMin = "0", UIMax = "1", ClampMin = "0", ClampMax = "1", EditCondition = "!bSelfCollideAgainstKinematicCollidersOnly"))
 	float SelfCollisionStiffness = 0.5f;
 
 	/** Friction coefficient for cloth - cloth interaction. */
-	UPROPERTY(EditAnywhere, Category = "Self-Collision Properties", meta = (UIMin = "0", UIMax = "1", ClampMin = "0", ClampMax = "1"))
+	UPROPERTY(EditAnywhere, Category = "Self-Collision Properties", meta = (UIMin = "0", UIMax = "1", ClampMin = "0", ClampMax = "1", EditCondition = "!bSelfCollideAgainstKinematicCollidersOnly"))
 	float SelfCollisionFriction = 0.0f;
 
 	/** Disabled neighbor collision ring. Collisions are disabled between vertices within this N-ring connectivity distance.*/
-	UPROPERTY(EditAnywhere, Category = "Self-Collision Properties", meta = (UIMin = "1", UIMax = "5", ClampMin = "1"))
+	UPROPERTY(EditAnywhere, Category = "Self-Collision Properties", meta = (UIMin = "1", UIMax = "5", ClampMin = "1", EditCondition = "!bSelfCollideAgainstKinematicCollidersOnly"))
 	int32 SelfCollisionDisableNeighborDistance = 5;
 
 	/** Self collision layers face int map. Generate this map using the SelectionsToIntMap node with SimFace Selections.
 	* Faces labeled with -1 will collide normally without any layering behavior.
 	* Faces labeled with any other number will keep higher layer numbers outside lower layer numbers (outside = front facing normal direction).
 	*/
-	UPROPERTY(EditAnywhere, Category = "Self-Collision Properties")
+	UPROPERTY(EditAnywhere, Category = "Self-Collision Properties", meta = (EditCondition = "!bSelfCollideAgainstKinematicCollidersOnly"))
 	FChaosClothAssetConnectableIStringValue SelfCollisionLayers = { TEXT("SelfCollisionLayers") };
 
 	/** Sim face selection set of faces which should not self collide */
 	UPROPERTY(EditAnywhere, Category = "Self-Collision Properties")
 	FChaosClothAssetConnectableIStringValue SelfCollisionDisabledFaces = { TEXT("SelfCollisionDisabledFaces") };
-
-	/** Self collide against all kinematic vertices. Kinematic colliders do not do Self Intersections. They always collide against the front-face.*/
+	
+	/** Collide only against kinematic colliders (no dynamic self collisions). Kinematic colliders do not do Self Intersections. They always collide against the front-face.*/
 	UPROPERTY(EditAnywhere, Category = "Self-Collision Properties - Kinematic Colliders")
-	bool bSelfCollideAgainstAllKinematicVertices = false;
+	bool bSelfCollideAgainstKinematicCollidersOnly = false;
 
 	/** Sim face selection set of kinematic faces which should self collide. Kinematic colliders do not do Self Intersections. They always collide against the front-face. */
 	UPROPERTY(EditAnywhere, Category = "Self-Collision Properties - Kinematic Colliders", meta = (EditCondition = "!bSelfCollideAgainstAllKinematicVertices"))
@@ -61,28 +61,32 @@ public:
 	UPROPERTY(EditAnywhere, Category = "Self-Collision Properties - Kinematic Colliders", meta = (UIMin = "0", UIMax = "1", ClampMin = "0", ClampMax = "1"))
 	FChaosClothAssetWeightedValue SelfCollisionKinematicColliderFrictionWeighted = {true, 0.0f, 0.f, TEXT("SelfCollisionKinematicColliderFriction")};
 
+	/** Self collide against all kinematic vertices. Kinematic colliders do not do Self Intersections. They always collide against the front-face.*/
+	UPROPERTY(EditAnywhere, Category = "Self-Collision Properties - Kinematic Colliders")
+	bool bSelfCollideAgainstAllKinematicVertices = false;
+
 	/** Enable self intersection resolution. This will try to fix any cloth intersections that are not handled by collision repulsions. */
-	UPROPERTY(EditAnywhere, Category = "Experimental")
+	UPROPERTY(EditAnywhere, Category = "Experimental", meta = (EditCondition = "!bSelfCollideAgainstKinematicCollidersOnly"))
 	bool bUseSelfIntersections = false;
 
 	/** Do global intersection analysis to determine the correct normals for the collision springs */
-	UPROPERTY(EditAnywhere, Category = "Experimental", meta = (EditCondition = "bUseSelfIntersections"))
+	UPROPERTY(EditAnywhere, Category = "Experimental", meta = (EditCondition = "bUseSelfIntersections && !bSelfCollideAgainstKinematicCollidersOnly"))
 	bool bUseGlobalIntersectionAnalysis = true;
 
 	/** Do a step of contour minimization at the beginning of the timestep. */
-	UPROPERTY(EditAnywhere, Category = "Experimental", meta = (EditCondition = "bUseSelfIntersections"))
+	UPROPERTY(EditAnywhere, Category = "Experimental", meta = (EditCondition = "bUseSelfIntersections && !bSelfCollideAgainstKinematicCollidersOnly"))
 	bool bUseContourMinimization = true;
 
 	/** Number of post timestep contour minimization steps to do. (Expensive!)*/
-	UPROPERTY(EditAnywhere, Category = "Experimental", meta = (ClampMin = "0", EditCondition = "bUseSelfIntersections"))
+	UPROPERTY(EditAnywhere, Category = "Experimental", meta = (ClampMin = "0", EditCondition = "bUseSelfIntersections && !bSelfCollideAgainstKinematicCollidersOnly"))
 	int32 NumContourMinimizationPostSteps = 0;
 
 	/** Use global contour gradients when doing post timestep contour minimization */
-	UPROPERTY(EditAnywhere, Category = "Experimental", meta = (EditCondition = "bUseSelfIntersections && NumContourMinimizationPostSteps > 0"))
+	UPROPERTY(EditAnywhere, Category = "Experimental", meta = (EditCondition = "bUseSelfIntersections && NumContourMinimizationPostSteps > 0 && !bSelfCollideAgainstKinematicCollidersOnly"))
 	bool bUseGlobalPostStepContours = true;
 
 	/** The stiffness of the proximity repulsions used to control self collision (Force-based Solver). Units = kg cm/ s^2 (same as XPBD springs) */
-	UPROPERTY(EditAnywhere, Category = "Experimental", meta = (UIMin = "0", UIMax = "10000", ClampMin = "0", ClampMax = "10000000"))
+	UPROPERTY(EditAnywhere, Category = "Experimental", meta = (UIMin = "0", UIMax = "10000", ClampMin = "0", ClampMax = "10000000", EditCondition = "!bSelfCollideAgainstAllKinematicVertices"))
 	float SelfCollisionProximityStiffness = 1.f;
 
 	FChaosClothAssetSimulationSelfCollisionConfigNode(const Dataflow::FNodeParameters& InParam, FGuid InGuid = FGuid::NewGuid());
diff --git a/Engine/Source/Runtime/Experimental/Chaos/Private/Chaos/PBDCollisionSpringConstraintsBase.cpp b/Engine/Source/Runtime/Experimental/Chaos/Private/Chaos/PBDCollisionSpringConstraintsBase.cpp
--- a/Engine/Source/Runtime/Experimental/Chaos/Private/Chaos/PBDCollisionSpringConstraintsBase.cpp
+++ b/Engine/Source/Runtime/Experimental/Chaos/Private/Chaos/PBDCollisionSpringConstraintsBase.cpp
@@ -70,6 +70,7 @@ FPBDCollisionSpringConstraintsBase::FPBDCollisionSpringConstraintsBase(
 	const FSolverReal InThickness,
 	const FSolverReal InStiffness,
 	const FSolverReal InFrictionCoefficient,
+	const bool bInOnlyCollideKinematics,
 	const FSolverReal InKinematicColliderThickness,
 	const FSolverReal InKinematicColliderStiffness,
 	const FSolverVec2 InKinematicColliderFrictionCoefficient,
@@ -77,6 +78,7 @@ FPBDCollisionSpringConstraintsBase::FPBDCollisionSpringConstraintsBase(
 	: Thickness(InThickness)
 	, Stiffness(InStiffness)
 	, FrictionCoefficient(InFrictionCoefficient)
+	, bOnlyCollideKinematics(bInOnlyCollideKinematics)
 	, KinematicColliderThickness(InKinematicColliderThickness)
 	, KinematicColliderStiffness(InKinematicColliderStiffness)
 	, KinematicColliderFrictionCoefficient(InKinematicColliderFrictionCoefficient, InKinematicColliderFrictionMultipliers, InNumParticles)
@@ -150,7 +152,7 @@ void FPBDCollisionSpringConstraintsBase::Init(const SolverParticlesOrRange& Part
 	PRAGMA_DISABLE_DEPRECATION_WARNINGS
 	const int32 NumDynamicElements = CollidableSubMesh.GetDynamicSubMesh().GetNumElements();
 
-	if (!NumDynamicElements)
+	if (!NumDynamicElements && !bOnlyCollideKinematics)
 	{
 		Constraints.Reset();
 		Barys.Reset();
@@ -202,7 +204,7 @@ void FPBDCollisionSpringConstraintsBase::Init(const SolverParticlesOrRange& Part
 
 				// Dynamic collisions
 				TArray< TTriangleCollisionPoint<FSolverReal> > DynamicResult;
-				if (CollidableSubMesh.GetDynamicSubMesh().PointProximityQuery(DynamicSpatial, static_cast<const TArrayView<const FSolverVec3>&>(Particles.XArray()), Index, Particles.X(Index), Thickness * ExtraThicknessMult, Thickness * ExtraThicknessMult,
+				if (!bOnlyCollideKinematics && CollidableSubMesh.GetDynamicSubMesh().PointProximityQuery(DynamicSpatial, static_cast<const TArrayView<const FSolverVec3>&>(Particles.XArray()), Index, Particles.X(Index), Thickness * ExtraThicknessMult, Thickness * ExtraThicknessMult,
 					[this, bVertexHasCollisionLayers, &Particles, &CollidableSubMesh, &VertexGIAColors, &TriangleGIAColors](const int32 PointIndex, const int32 SubMeshTriangleIndex)->bool
 					{
 						const TVector<int32, 3>& Elem = CollidableSubMesh.GetDynamicSubMesh().GetElements()[SubMeshTriangleIndex];
diff --git a/Engine/Source/Runtime/Experimental/Chaos/Private/Chaos/PBDTriangleMeshCollisions.cpp b/Engine/Source/Runtime/Experimental/Chaos/Private/Chaos/PBDTriangleMeshCollisions.cpp
--- a/Engine/Source/Runtime/Experimental/Chaos/Private/Chaos/PBDTriangleMeshCollisions.cpp
+++ b/Engine/Source/Runtime/Experimental/Chaos/Private/Chaos/PBDTriangleMeshCollisions.cpp
@@ -1282,7 +1282,7 @@ struct FPBDTriangleMeshCollisions::FScratchBuffers
 };
 
 template<typename SolverParticlesOrRange>
-void FPBDTriangleMeshCollisions::FTriangleSubMesh::Init(const SolverParticlesOrRange& Particles, const TSet<int32>& InDisabledFaces, bool bCollideAgainstAllKinematicVertices, const TSet<int32>& InEnabledKinematicFaces)
+void FPBDTriangleMeshCollisions::FTriangleSubMesh::Init(const SolverParticlesOrRange& Particles, const TSet<int32>& InDisabledFaces, bool bCollideAgainstAllKinematicVertices, const TSet<int32>& InEnabledKinematicFaces, const bool bOnlyCollideKinematics)
 {
 	FullMeshToSubMeshIndices.Reset();
 	DynamicSubMeshToFullMeshIndices.Reset();
@@ -1291,7 +1291,10 @@ void FPBDTriangleMeshCollisions::FTriangleSubMesh::Init(const SolverParticlesOrR
 
 	TArray<TVec3<int32>> DynamicMeshElements;
 	TArray<TVec3<int32>> KinematicMeshElements;
-	DynamicMeshElements.Reserve(FullMesh.GetNumElements());
+	if (!bOnlyCollideKinematics)
+	{
+		DynamicMeshElements.Reserve(FullMesh.GetNumElements());
+	}
 	KinematicMeshElements.Reserve(InEnabledKinematicFaces.Num());
 
 	FullMeshToSubMeshIndices.SetNumZeroed(FullMesh.GetNumElements());
@@ -1325,7 +1328,7 @@ void FPBDTriangleMeshCollisions::FTriangleSubMesh::Init(const SolverParticlesOrR
 				FullMeshToSubMeshIndices[FullElementIndex].SubMeshType = ESubMeshType::Invalid;
 			}
 		}
-		else
+		else if (!bOnlyCollideKinematics)
 		{
 			const int32 SubMeshIndex = DynamicMeshElements.Add(FullElements[FullElementIndex]);
 			DynamicSubMeshToFullMeshIndices.Add(FullElementIndex);
@@ -1380,13 +1383,13 @@ void FPBDTriangleMeshCollisions::Init(const SolverParticlesOrRange& Particles, c
 	const bool bDoSelfIntersections = bGlobalIntersectionAnalysis || bContourMinimization;
 	if (bCollidableSubMeshDirty)
 	{		
-		CollidableSubMesh.Init(Particles, DisabledFaces, bSelfCollideAgainstAllKinematicVertices, EnabledKinematicFaces);
+		CollidableSubMesh.Init(Particles, DisabledFaces, bSelfCollideAgainstAllKinematicVertices, EnabledKinematicFaces, bOnlyCollideWithKinematics);
 		bCollidableSubMeshDirty = false;
 	}
 
 	const FTriangleMesh& DynamicSubMesh = CollidableSubMesh.GetDynamicSubMesh();
 
-	if (DynamicSubMesh.GetNumElements() == 0)
+	if (DynamicSubMesh.GetNumElements() == 0 && !bOnlyCollideWithKinematics)
 	{
 		return;
 	}
@@ -1408,7 +1411,7 @@ void FPBDTriangleMeshCollisions::Init(const SolverParticlesOrRange& Particles, c
 	IntersectionContourPoints.Reset();
 	IntersectionContourTypes.Reset();
 
-	if (!bDoSelfIntersections)
+	if (!bDoSelfIntersections || bOnlyCollideWithKinematics)
 	{
 		return;
 	}
diff --git a/Engine/Source/Runtime/Experimental/Chaos/Public/Chaos/PBDCollisionSpringConstraints.h b/Engine/Source/Runtime/Experimental/Chaos/Public/Chaos/PBDCollisionSpringConstraints.h
--- a/Engine/Source/Runtime/Experimental/Chaos/Public/Chaos/PBDCollisionSpringConstraints.h
+++ b/Engine/Source/Runtime/Experimental/Chaos/Public/Chaos/PBDCollisionSpringConstraints.h
@@ -40,6 +40,7 @@ public:
 			(FSolverReal)FMath::Max(GetSelfCollisionThickness(PropertyCollection, Base::BackCompatThickness), 0.f),
 			(FSolverReal)FMath::Clamp(GetSelfCollisionStiffness(PropertyCollection, Base::BackCompatStiffness), 0.f, 1.f),
 			FMath::Clamp((FSolverReal)GetSelfCollisionFriction(PropertyCollection, Base::BackCompatFrictionCoefficient), MinFrictionCoefficient, MaxFrictionCoefficient),
+			GetSelfCollideAgainstKinematicCollidersOnly(PropertyCollection, false),
 			(FSolverReal)FMath::Max(GetSelfCollisionKinematicColliderThickness(PropertyCollection, Base::DefaultKinematicColliderThickness), 0.f),
 			(FSolverReal)FMath::Clamp((FSolverReal)GetSelfCollisionKinematicColliderStiffness(PropertyCollection, Base::DefaultKinematicColliderStiffness), 0.f, 1.f),
 			FSolverVec2(GetWeightedFloatSelfCollisionKinematicColliderFriction(PropertyCollection, Base::DefaultKinematicColliderFrictionCoefficient)).ClampAxes(MinFrictionCoefficient, MaxFrictionCoefficient),
@@ -49,6 +50,7 @@ public:
 		, SelfCollisionFrictionIndex(PropertyCollection)
 		, SelfCollisionProximityStiffnessIndex(PropertyCollection)
 		, SelfCollisionLayersIndex(PropertyCollection)
+		, SelfCollideAgainstKinematicCollidersOnlyIndex(PropertyCollection)
 		, SelfCollisionKinematicColliderThicknessIndex(PropertyCollection)
 		, SelfCollisionKinematicColliderStiffnessIndex(PropertyCollection)
 		, SelfCollisionKinematicColliderFrictionIndex(PropertyCollection)
@@ -98,6 +100,7 @@ public:
 		, SelfCollisionFrictionIndex(ForceInit)
 		, SelfCollisionProximityStiffnessIndex(ForceInit)
 		, SelfCollisionLayersIndex(ForceInit)
+		, SelfCollideAgainstKinematicCollidersOnlyIndex(ForceInit)
 		, SelfCollisionKinematicColliderThicknessIndex(ForceInit)
 		, SelfCollisionKinematicColliderStiffnessIndex(ForceInit)
 		, SelfCollisionKinematicColliderFrictionIndex(ForceInit)
@@ -135,6 +138,10 @@ public:
 				Base::UpdateCollisionLayers(FaceIntMaps.FindRef(FaceIntMapName));
 			}
 		}
+		if (IsSelfCollideAgainstKinematicCollidersOnlyMutable(PropertyCollection))
+		{
+			bOnlyCollideKinematics = GetSelfCollideAgainstKinematicCollidersOnly(PropertyCollection);
+		}
 		if (IsSelfCollisionKinematicColliderThicknessMutable(PropertyCollection))
 		{
 			KinematicColliderThickness = FMath::Max(GetSelfCollisionKinematicColliderThickness(PropertyCollection), 0.f);
@@ -197,6 +204,7 @@ private:
 	UE_CHAOS_DECLARE_PROPERTYCOLLECTION_NAME(SelfCollisionProximityStiffness, float);
 	UE_CHAOS_DECLARE_INDEXLESS_PROPERTYCOLLECTION_NAME(SelfCollisionDisableNeighborDistance, int32);
 	UE_CHAOS_DECLARE_PROPERTYCOLLECTION_NAME(SelfCollisionLayers, float);
+	UE_CHAOS_DECLARE_PROPERTYCOLLECTION_NAME(SelfCollideAgainstKinematicCollidersOnly, bool);
 	UE_CHAOS_DECLARE_PROPERTYCOLLECTION_NAME(SelfCollisionKinematicColliderThickness, float);
 	UE_CHAOS_DECLARE_PROPERTYCOLLECTION_NAME(SelfCollisionKinematicColliderStiffness, float);
 	UE_CHAOS_DECLARE_PROPERTYCOLLECTION_NAME(SelfCollisionKinematicColliderFriction, float);
diff --git a/Engine/Source/Runtime/Experimental/Chaos/Public/Chaos/PBDCollisionSpringConstraintsBase.h b/Engine/Source/Runtime/Experimental/Chaos/Public/Chaos/PBDCollisionSpringConstraintsBase.h
--- a/Engine/Source/Runtime/Experimental/Chaos/Public/Chaos/PBDCollisionSpringConstraintsBase.h
+++ b/Engine/Source/Runtime/Experimental/Chaos/Public/Chaos/PBDCollisionSpringConstraintsBase.h
@@ -41,6 +41,7 @@ public:
 		const FSolverReal InThickness = BackCompatThickness,
 		const FSolverReal InStiffness = BackCompatStiffness,
 		const FSolverReal InFrictionCoefficient = BackCompatFrictionCoefficient,
+		const bool bInOnlyCollideKinematics = false,
 		const FSolverReal InKinematicColliderThickness = DefaultKinematicColliderThickness,
 		const FSolverReal InKinematicColliderStiffness = DefaultKinematicColliderStiffness,
 		const FSolverVec2 InKinematicColliderFrictionCoefficient = FSolverVec2(DefaultKinematicColliderFrictionCoefficient),
@@ -186,7 +187,8 @@ protected:
 
 	FSolverReal Thickness;
 	FSolverReal Stiffness; // (0-1 compliance for PBD)
-	FSolverReal FrictionCoefficient;
+	FSolverReal FrictionCoefficient; 
+	bool bOnlyCollideKinematics;
 	FSolverReal KinematicColliderThickness;
 	FSolverReal KinematicColliderStiffness;
 	FPBDFlatWeightMap KinematicColliderFrictionCoefficient;
diff --git a/Engine/Source/Runtime/Experimental/Chaos/Public/Chaos/PBDTriangleMeshCollisions.h b/Engine/Source/Runtime/Experimental/Chaos/Public/Chaos/PBDTriangleMeshCollisions.h
--- a/Engine/Source/Runtime/Experimental/Chaos/Public/Chaos/PBDTriangleMeshCollisions.h
+++ b/Engine/Source/Runtime/Experimental/Chaos/Public/Chaos/PBDTriangleMeshCollisions.h
@@ -117,7 +117,7 @@ public:
 		{}
 
 		template<typename SolverParticlesOrRange>
-		void Init(const SolverParticlesOrRange& Particles, const TSet<int32>& DisabledFaces, bool bCollideAgainstAllKinematicVertices, const TSet<int32>& EnabledKinematicFaces);
+		void Init(const SolverParticlesOrRange& Particles, const TSet<int32>& DisabledFaces, bool bCollideAgainstAllKinematicVertices, const TSet<int32>& EnabledKinematicFaces, const bool bOnlyCollideKinematics = false);
 
 		void InitAllDynamic();
 
@@ -164,6 +164,7 @@ public:
 
 		const TArray<int32>& GetDynamicVertices() const { return DynamicVertices; }
 
+
 	private:
 		const FTriangleMesh& FullMesh;
 		FTriangleMesh DynamicSubMesh;
@@ -206,6 +207,7 @@ public:
 		, bContourMinimization(GetUseSelfIntersections(PropertyCollection, false) && GetUseContourMinimization(PropertyCollection, true))
 		, NumContourMinimizationPostSteps(GetUseSelfIntersections(PropertyCollection, false) ? GetNumContourMinimizationPostSteps(PropertyCollection, 0) : 0)
 		, bUseGlobalPostStepContours(GetUseGlobalPostStepContours(PropertyCollection, true))
+		, bOnlyCollideWithKinematics(GetSelfCollideAgainstKinematicCollidersOnly(PropertyCollection, false))
 		, bSelfCollideAgainstAllKinematicVertices(GetSelfCollideAgainstAllKinematicVertices(PropertyCollection, false))
 		, bCollidableSubMeshDirty(true)
 		, UseSelfIntersectionsIndex(PropertyCollection)
@@ -213,6 +215,7 @@ public:
 		, UseContourMinimizationIndex(PropertyCollection)
 		, NumContourMinimizationPostStepsIndex(PropertyCollection)
 		, UseGlobalPostStepContoursIndex(PropertyCollection)
+		, SelfCollideAgainstKinematicCollidersOnlyIndex(PropertyCollection)
 		, SelfCollideAgainstAllKinematicVerticesIndex(PropertyCollection)
 		, SelfCollisionDisabledFacesIndex(PropertyCollection)
 		, SelfCollisionEnabledKinematicFacesIndex(PropertyCollection)
@@ -245,6 +248,7 @@ public:
 		, bUseSelfIntersections(bInGlobalIntersectionAnalysis || bInContourMinimization)
 		, bGlobalIntersectionAnalysis(bInGlobalIntersectionAnalysis)
 		, bContourMinimization(bInContourMinimization)
+		, bOnlyCollideWithKinematics(false)
 		, bSelfCollideAgainstAllKinematicVertices(false)
 		, bCollidableSubMeshDirty(true)
 		, UseSelfIntersectionsIndex(ForceInit)
@@ -252,6 +256,7 @@ public:
 		, UseContourMinimizationIndex(ForceInit)
 		, NumContourMinimizationPostStepsIndex(ForceInit)
 		, UseGlobalPostStepContoursIndex(ForceInit)
+		, SelfCollideAgainstKinematicCollidersOnlyIndex(ForceInit)
 		, SelfCollideAgainstAllKinematicVerticesIndex(ForceInit)
 		, SelfCollisionDisabledFacesIndex(ForceInit)
 		, SelfCollisionEnabledKinematicFacesIndex(ForceInit)
@@ -292,6 +297,16 @@ public:
 			NumContourMinimizationPostSteps = 0;
 		}
 
+		if (IsSelfCollideAgainstKinematicCollidersOnlyMutable(PropertyCollection))
+		{
+			const bool bNewValue = GetSelfCollideAgainstKinematicCollidersOnly(PropertyCollection);
+			if (bNewValue != bOnlyCollideWithKinematics)
+			{
+				bOnlyCollideWithKinematics = bNewValue;
+				bCollidableSubMeshDirty = true;
+			}
+		}
+
 		if (IsSelfCollideAgainstAllKinematicVerticesMutable(PropertyCollection))
 		{
 			const bool bNewValue = GetSelfCollideAgainstAllKinematicVertices(PropertyCollection);
@@ -380,11 +395,13 @@ private:
 	bool bContourMinimization;	
 	int32 NumContourMinimizationPostSteps = 0;
 	bool bUseGlobalPostStepContours = true;
+	bool bOnlyCollideWithKinematics;
 	bool bSelfCollideAgainstAllKinematicVertices;
 	TSet<int32> DisabledFaces;
 	TSet<int32> EnabledKinematicFaces;
 
 	bool bCollidableSubMeshDirty = true;
+
 	
 	FTriangleMesh::TSpatialHashType<FSolverReal> DynamicSubMeshSpatialHash;
 	FTriangleMesh::TSpatialHashType<FSolverReal> KinematicSubMeshSpatialHash;
@@ -409,6 +426,7 @@ private:
 	UE_CHAOS_DECLARE_PROPERTYCOLLECTION_NAME(UseContourMinimization, bool);
 	UE_CHAOS_DECLARE_PROPERTYCOLLECTION_NAME(NumContourMinimizationPostSteps, int32);
 	UE_CHAOS_DECLARE_PROPERTYCOLLECTION_NAME(UseGlobalPostStepContours, bool);
+	UE_CHAOS_DECLARE_PROPERTYCOLLECTION_NAME(SelfCollideAgainstKinematicCollidersOnly, bool);
 	UE_CHAOS_DECLARE_PROPERTYCOLLECTION_NAME(SelfCollideAgainstAllKinematicVertices, bool);
 	UE_CHAOS_DECLARE_PROPERTYCOLLECTION_NAME(SelfCollisionDisabledFaces, bool);
 	UE_CHAOS_DECLARE_PROPERTYCOLLECTION_NAME(SelfCollisionEnabledKinematicFaces, bool);
