Chaos Cloth Editor: fix node traversal issue where we were getting stuck in cycles inside FChaosClothAssetEditorToolkit::EvaluateNode

diff --git a/Engine/Plugins/Experimental/ChaosClothAssetEditor/Source/ChaosClothAssetEditor/Private/ChaosClothAsset/ClothEditorToolkit.cpp b/Engine/Plugins/Experimental/ChaosClothAssetEditor/Source/ChaosClothAssetEditor/Private/ChaosClothAsset/ClothEditorToolkit.cpp
--- a/Engine/Plugins/Experimental/ChaosClothAssetEditor/Source/ChaosClothAssetEditor/Private/ChaosClothAsset/ClothEditorToolkit.cpp
+++ b/Engine/Plugins/Experimental/ChaosClothAssetEditor/Source/ChaosClothAssetEditor/Private/ChaosClothAsset/ClothEditorToolkit.cpp
@@ -857,11 +857,13 @@ void FChaosClothAssetEditorToolkit::EvaluateNode(FDataflowNode* Node, bool bForc
 		TQueue<FDataflowNode*> Queue;
 		Queue.Enqueue(Node);
 
+		TSet<FDataflowNode*> VisitedNodes;
+
 		do
 		{
 			FDataflowNode* VisitedNode;
 			Queue.Dequeue(VisitedNode);
-
+			VisitedNodes.Add(VisitedNode);
 			if (VisitedNode->GetTimestamp() >= LastDataflowNodeTimestamp)
 			{
 				UE_LOG(LogChaosClothAssetEditor, VeryVerbose, TEXT("EvaluateNode - Visiting node %s"), *VisitedNode->GetName().ToString());
@@ -872,7 +874,10 @@ void FChaosClothAssetEditorToolkit::EvaluateNode(FDataflowNode* Node, bool bForc
 					{
 						if (FDataflowNode* const OwningNode = ConnectedOutput->GetOwningNode())
 						{
-							Queue.Enqueue(OwningNode);
+							if (!VisitedNodes.Contains(OwningNode))
+							{
+								Queue.Enqueue(OwningNode);
+							}
 
 							if (FDataflowTerminalNode* const TerminalNode = OwningNode->AsType<FDataflowTerminalNode>())
 							{
