Cloth Editor Remesh Node: if the Cloth Collection input has proxy deformer data, remove that data from the output and warn the user that we are doing so

diff --git a/Engine/Plugins/Experimental/ChaosClothAssetEditor/Source/ChaosClothAssetDataflowNodes/Private/ChaosClothAsset/RemeshNode.cpp b/Engine/Plugins/Experimental/ChaosClothAssetEditor/Source/ChaosClothAssetDataflowNodes/Private/ChaosClothAsset/RemeshNode.cpp
--- a/Engine/Plugins/Experimental/ChaosClothAssetEditor/Source/ChaosClothAssetDataflowNodes/Private/ChaosClothAsset/RemeshNode.cpp
+++ b/Engine/Plugins/Experimental/ChaosClothAssetEditor/Source/ChaosClothAssetDataflowNodes/Private/ChaosClothAsset/RemeshNode.cpp
@@ -6,8 +6,10 @@
 #include "ChaosClothAsset/CollectionClothFacade.h"
 #include "ChaosClothAsset/CollectionClothSelectionFacade.h"
 #include "ChaosClothAsset/ClothCollectionGroup.h"
+#include "ChaosClothAsset/ClothCollectionAttribute.h"
 #include "ChaosClothAsset/ClothGeometryTools.h"
 #include "ChaosClothAsset/ClothEngineTools.h"
+#include "ChaosClothAsset/ClothDataflowTools.h"
 #include "Dataflow/DataflowInputOutput.h"
 #include "DynamicMesh/MeshTangents.h"
 #include "DynamicMesh/MeshNormals.h"
@@ -1555,28 +1557,54 @@ void FChaosClothAssetRemeshNode::Evaluate(Dataflow::FContext& Context, const FDa
 		// Evaluate in collection
 		FManagedArrayCollection InCollection = GetValue<FManagedArrayCollection>(Context, &Collection);
 		const TSharedRef<FManagedArrayCollection> ClothCollection = MakeShared<FManagedArrayCollection>(MoveTemp(InCollection));
+		FCollectionClothFacade ClothFacade(ClothCollection);
 
 		// Copy collection to output
-		TSharedRef<FManagedArrayCollection> OutputClothCollection = MakeShared<FManagedArrayCollection>(*ClothCollection);
+		TSharedPtr<FManagedArrayCollection> OutputClothCollection = MakeShared<FManagedArrayCollection>();
+
+		if ((bRemeshSim || bRemeshRender) && ClothFacade.IsValid(EClothCollectionOptionalSchemas::RenderDeformer))
+		{
+			FClothDataflowTools::LogAndToastWarning(*this, 
+				LOCTEXT("InputHasDeformerDataHeadline", "Proxy Deformer Data Found"),
+				LOCTEXT("InputHasDeformerDataDetails", "The input Cloth Collection has Proxy Deformer data that will be removed by the Remesh node. Default deformer bindings will be computed in the final asset. Consider placing ProxyDeformer Node after the Remesh Node."));
+
+			// Don't copy proxy deformer data
+			const TArray<FName> GroupsToSkip = TArray<FName>();
+			TArray<TTuple<FName, FName>> AttributesToSkip;
+			AttributesToSkip.Add({ ClothCollectionAttribute::RenderDeformerNumInfluences, ClothCollectionGroup::RenderPatterns });
+			AttributesToSkip.Add({ ClothCollectionAttribute::RenderDeformerPositionBaryCoordsAndDist, ClothCollectionGroup::RenderVertices});
+			AttributesToSkip.Add({ ClothCollectionAttribute::RenderDeformerNormalBaryCoordsAndDist, ClothCollectionGroup::RenderVertices });
+			AttributesToSkip.Add({ ClothCollectionAttribute::RenderDeformerTangentBaryCoordsAndDist, ClothCollectionGroup::RenderVertices });
+			AttributesToSkip.Add({ ClothCollectionAttribute::RenderDeformerSimIndices3D, ClothCollectionGroup::RenderVertices });
+			AttributesToSkip.Add({ ClothCollectionAttribute::RenderDeformerWeight, ClothCollectionGroup::RenderVertices });
+			AttributesToSkip.Add({ ClothCollectionAttribute::RenderDeformerSkinningBlend, ClothCollectionGroup::RenderVertices });
+			
+			ClothCollection->CopyTo(OutputClothCollection.Get(), GroupsToSkip, AttributesToSkip);
+		}
+		else
+		{
+			ClothCollection->CopyTo(OutputClothCollection.Get());
+		}
+
+		TSharedRef<FManagedArrayCollection> OutputClothCollectionRef = OutputClothCollection.ToSharedRef();
 
-		FCollectionClothFacade ClothFacade(ClothCollection);
 		if (ClothFacade.IsValid())  // Can only act on the collection if it is a valid cloth collection
 		{
 			if (bRemeshSim)
 			{
-				EmptySimSelections(OutputClothCollection);
-				RemeshSimMesh(ClothCollection, OutputClothCollection);
-				RebuildTopologyDependentSimData(ClothCollection, OutputClothCollection);
+				EmptySimSelections(OutputClothCollectionRef);
+				RemeshSimMesh(ClothCollection, OutputClothCollectionRef);
+				RebuildTopologyDependentSimData(ClothCollection, OutputClothCollectionRef);
 			}
 
 			if (bRemeshRender)
 			{
-				EmptyRenderSelections(OutputClothCollection);
-				RemeshRenderMesh(ClothCollection, OutputClothCollection);
+				EmptyRenderSelections(OutputClothCollectionRef);
+				RemeshRenderMesh(ClothCollection, OutputClothCollectionRef);
 			}
 		}
 
-		SetValue(Context, MoveTemp(*OutputClothCollection), &Collection);
+		SetValue(Context, MoveTemp(*OutputClothCollectionRef), &Collection);
 	}
 }
 
