Chaos Weight Map Painting: Add button to multiply map by a value.

diff --git a/Engine/Plugins/Experimental/ChaosClothAssetEditor/Source/ChaosClothAssetEditorTools/Private/ChaosClothAsset/ClothWeightMapPaintTool.cpp b/Engine/Plugins/Experimental/ChaosClothAssetEditor/Source/ChaosClothAssetEditorTools/Private/ChaosClothAsset/ClothWeightMapPaintTool.cpp
--- a/Engine/Plugins/Experimental/ChaosClothAssetEditor/Source/ChaosClothAssetEditorTools/Private/ChaosClothAsset/ClothWeightMapPaintTool.cpp
+++ b/Engine/Plugins/Experimental/ChaosClothAssetEditor/Source/ChaosClothAssetEditorTools/Private/ChaosClothAsset/ClothWeightMapPaintTool.cpp
@@ -1923,6 +1923,37 @@ void UClothEditorWeightMapPaintTool::InvertWeightsAction()
 	EndChange();
 }
 
+
+void UClothEditorWeightMapPaintTool::MultiplyWeightsAction()
+{
+	if (!ActiveWeightMap)
+	{
+		return;
+	}
+	BeginChange();
+
+	const float WeightMultiplierValue = FilterProperties->AttributeValue;
+
+	const FDynamicMesh3* Mesh = DynamicMeshComponent->GetMesh();
+	checkf(Mesh, TEXT("Paint Tool's DynamicMeshComponent has no FDynamicMesh"));
+
+	for (const int32 VertexID : Mesh->VertexIndicesItr())
+	{
+		ActiveWeightEditChangeTracker->SaveVertexOneRingTriangles(VertexID, true);
+
+		float WeightValue;
+		ActiveWeightMap->GetValue(VertexID, &WeightValue);
+		WeightValue = FMath::Clamp(WeightMultiplierValue * WeightValue, 0.f, 1.f);
+		ActiveWeightMap->SetValue(VertexID, &WeightValue);
+	}
+
+	// update colors
+	UpdateVertexColorOverlay();
+	DynamicMeshComponent->FastNotifyVertexAttributesUpdated(EMeshRenderAttributeFlags::VertexColors);
+	GetToolManager()->PostInvalidation();
+	EndChange();
+}
+
 void UClothEditorWeightMapPaintTool::ClearHiddenAction()
 {
 	HiddenTriangles.Reset();
@@ -2251,6 +2282,10 @@ void UClothEditorWeightMapPaintTool::ApplyAction(EClothEditorWeightMapPaintToolA
 		InvertWeightsAction();
 		break;
 
+	case EClothEditorWeightMapPaintToolActions::Multiply:
+		MultiplyWeightsAction();
+		break;
+
 	case EClothEditorWeightMapPaintToolActions::ClearHiddenTriangles:
 		ClearHiddenAction();
 		break;
diff --git a/Engine/Plugins/Experimental/ChaosClothAssetEditor/Source/ChaosClothAssetEditorTools/Public/ChaosClothAsset/ClothWeightMapPaintTool.h b/Engine/Plugins/Experimental/ChaosClothAssetEditor/Source/ChaosClothAssetEditorTools/Public/ChaosClothAsset/ClothWeightMapPaintTool.h
--- a/Engine/Plugins/Experimental/ChaosClothAssetEditor/Source/ChaosClothAssetEditorTools/Public/ChaosClothAsset/ClothWeightMapPaintTool.h
+++ b/Engine/Plugins/Experimental/ChaosClothAssetEditor/Source/ChaosClothAssetEditorTools/Public/ChaosClothAsset/ClothWeightMapPaintTool.h
@@ -216,7 +216,8 @@ enum class EClothEditorWeightMapPaintToolActions
 	FloodFillCurrent,
 	ClearAll,
 	Invert,
-	ClearHiddenTriangles
+	Multiply,
+	ClearHiddenTriangles,
 };
 
 
@@ -251,6 +252,12 @@ public:
 		PostAction(EClothEditorWeightMapPaintToolActions::Invert);
 	}
 
+	UFUNCTION(CallInEditor, Category = Operations, meta = (DisplayPriority = 14))
+	void Multiply()
+	{
+		PostAction(EClothEditorWeightMapPaintToolActions::Multiply);
+	}
+
 };
 
 UCLASS()
@@ -349,6 +356,7 @@ private:
 	void FloodFillCurrentWeightAction();
 	void ClearAllWeightsAction();
 	void InvertWeightsAction();
+	void MultiplyWeightsAction();
 	void ClearHiddenAction();
 
 public:
