Chaos Cloth Editor: Add refresh asset button on cloth asset import.

diff --git a/Engine/Plugins/Experimental/ChaosClothAssetEditor/Source/ChaosClothAssetDataflowNodes/Private/ChaosClothAsset/ChaosClothAssetDataflowNodesModule.cpp b/Engine/Plugins/Experimental/ChaosClothAssetEditor/Source/ChaosClothAssetDataflowNodes/Private/ChaosClothAsset/ChaosClothAssetDataflowNodesModule.cpp
--- a/Engine/Plugins/Experimental/ChaosClothAssetEditor/Source/ChaosClothAssetDataflowNodes/Private/ChaosClothAsset/ChaosClothAssetDataflowNodesModule.cpp
+++ b/Engine/Plugins/Experimental/ChaosClothAssetEditor/Source/ChaosClothAssetDataflowNodes/Private/ChaosClothAsset/ChaosClothAssetDataflowNodesModule.cpp
@@ -149,6 +149,7 @@ namespace UE::Chaos::ClothAsset
 				PropertyModule->RegisterCustomPropertyTypeLayout(FChaosClothAssetNodeSelectionGroup::StaticStruct()->GetFName(), FOnGetPropertyTypeCustomizationInstance::CreateStatic(&FSelectionGroupCustomization::MakeInstance));
 				PropertyModule->RegisterCustomPropertyTypeLayout(FChaosClothAssetImportFilePath::StaticStruct()->GetFName(), FOnGetPropertyTypeCustomizationInstance::CreateStatic(&FImportFilePathCustomization::MakeInstance));
 				PropertyModule->RegisterCustomPropertyTypeLayout(FChaosClothAssetTerminalNodeRefreshAsset::StaticStruct()->GetFName(), FOnGetPropertyTypeCustomizationInstance::CreateStatic(&FTerminalNodeRefreshAssetCustomization::MakeInstance));
+				PropertyModule->RegisterCustomPropertyTypeLayout(FChaosClothAssetImportNodeRefreshAsset::StaticStruct()->GetFName(), FOnGetPropertyTypeCustomizationInstance::CreateStatic(&FTerminalNodeRefreshAssetCustomization::MakeInstance));
 			}
 
 			// Register modular features
@@ -171,6 +172,7 @@ namespace UE::Chaos::ClothAsset
 					PropertyModule->UnregisterCustomPropertyTypeLayout(FChaosClothAssetNodeSelectionGroup::StaticStruct()->GetFName());
 					PropertyModule->UnregisterCustomPropertyTypeLayout(FChaosClothAssetImportFilePath::StaticStruct()->GetFName());
 					PropertyModule->UnregisterCustomPropertyTypeLayout(FChaosClothAssetTerminalNodeRefreshAsset::StaticStruct()->GetFName());
+					PropertyModule->UnregisterCustomPropertyTypeLayout(FChaosClothAssetImportNodeRefreshAsset::StaticStruct()->GetFName());
 				}
 			}
 
diff --git a/Engine/Plugins/Experimental/ChaosClothAssetEditor/Source/ChaosClothAssetDataflowNodes/Private/ChaosClothAsset/ImportNode.cpp b/Engine/Plugins/Experimental/ChaosClothAssetEditor/Source/ChaosClothAssetDataflowNodes/Private/ChaosClothAsset/ImportNode.cpp
--- a/Engine/Plugins/Experimental/ChaosClothAssetEditor/Source/ChaosClothAssetDataflowNodes/Private/ChaosClothAsset/ImportNode.cpp
+++ b/Engine/Plugins/Experimental/ChaosClothAssetEditor/Source/ChaosClothAssetDataflowNodes/Private/ChaosClothAsset/ImportNode.cpp
@@ -29,6 +29,11 @@ void FChaosClothAssetImportNode::Evaluate(Dataflow::FContext& Context, const FDa
 
 	if (Out->IsA<FManagedArrayCollection>(&Collection))
 	{
+		if (ReimportAsset.bRefreshAsset)
+		{
+			ReimportAsset.bRefreshAsset = false;
+		}
+
 		bool bSetValue = false;
 		if (ClothAsset)
 		{
diff --git a/Engine/Plugins/Experimental/ChaosClothAssetEditor/Source/ChaosClothAssetDataflowNodes/Public/ChaosClothAsset/ImportNode.h b/Engine/Plugins/Experimental/ChaosClothAssetEditor/Source/ChaosClothAssetDataflowNodes/Public/ChaosClothAsset/ImportNode.h
--- a/Engine/Plugins/Experimental/ChaosClothAssetEditor/Source/ChaosClothAssetDataflowNodes/Public/ChaosClothAsset/ImportNode.h
+++ b/Engine/Plugins/Experimental/ChaosClothAssetEditor/Source/ChaosClothAssetDataflowNodes/Public/ChaosClothAsset/ImportNode.h
@@ -8,6 +8,16 @@
 
 class UChaosClothAsset;
 
+/** Refresh structure for push button customization. */
+USTRUCT()
+struct FChaosClothAssetImportNodeRefreshAsset
+{
+	GENERATED_USTRUCT_BODY()
+
+	UPROPERTY(EditAnywhere, Category = "Import Node Refresh Asset")
+	bool bRefreshAsset = false;
+};
+
 /** Import an existing Cloth Asset into the graph. */
 USTRUCT(Meta = (DataflowCloth))
 struct FChaosClothAssetImportNode : public FDataflowNode
@@ -27,6 +37,12 @@ public:
 	UPROPERTY(EditAnywhere, Category = "Cloth Asset Import", Meta = (DisplayName = "Import LOD", ClampMin = "0"))
 	int32 ImportLod = 0;
 
+	/**
+	 * Reimport the imported asset. 
+	 */
+	UPROPERTY(EditAnywhere, Category = "Cloth Asset Import")
+	mutable FChaosClothAssetImportNodeRefreshAsset ReimportAsset;
+
 	FChaosClothAssetImportNode(const Dataflow::FNodeParameters& InParam, FGuid InGuid = FGuid::NewGuid());
 
 private:
