Read inpaint weight map stored on the render mesh and set it in the transfer weights operator.

diff --git a/Engine/Plugins/Experimental/ChaosClothAssetEditor/Source/ChaosClothAssetDataflowNodes/Private/ChaosClothAsset/TransferSkinWeightsNode.cpp b/Engine/Plugins/Experimental/ChaosClothAssetEditor/Source/ChaosClothAssetDataflowNodes/Private/ChaosClothAsset/TransferSkinWeightsNode.cpp
--- a/Engine/Plugins/Experimental/ChaosClothAssetEditor/Source/ChaosClothAssetDataflowNodes/Private/ChaosClothAsset/TransferSkinWeightsNode.cpp
+++ b/Engine/Plugins/Experimental/ChaosClothAssetEditor/Source/ChaosClothAssetDataflowNodes/Private/ChaosClothAsset/TransferSkinWeightsNode.cpp
@@ -21,6 +21,7 @@
 #include "SkeletalMeshAttributes.h"
 #include "Selections/MeshConnectedComponents.h"
 #include "BoneWeights.h"
+#include "ChaosClothAsset/ClothCollectionGroup.h"
 
 #include UE_INLINE_GENERATED_CPP_BY_NAME(TransferSkinWeightsNode)
 
@@ -336,9 +337,6 @@ namespace UE::Chaos::ClothAsset::Private
 		TransferBoneWeights.SmoothingStrength = TransferSettings.SmoothingStrength;
 		TransferBoneWeights.LayeredMeshSupport = TransferSettings.LayeredMeshSupport; // multilayerd clothing
 
-		// TODO: allow to specify an inpaint mask for the render mesh
-		// TransferBoneWeights.ForceInpaintWeightMapName = FName(TransferSettings.InpaintMaskWeightMapName);
-
 		if (TransferBoneWeights.Validate() != EOperationValidationResult::Ok)
 		{
 			UE_LOG(LogChaosClothAssetDataflowNodes, Warning, TEXT("TransferSkinWeightsNode: Transfer method parameters are invalid."));
@@ -349,6 +347,21 @@ namespace UE::Chaos::ClothAsset::Private
 		FMeshConnectedComponents ConnectedComponents(&TargetDynamicMesh);
 		ConnectedComponents.FindConnectedTriangles();
 
+		// Pointer to the weight layer containing force inpaint mask (if one exists)
+		const FDynamicMeshWeightAttribute* ForceInpaintWeightLayer = nullptr; 
+		if (!TransferSettings.InpaintMaskWeightMapName.IsEmpty())
+		{
+			for (int32 Idx = 0; Idx < TargetDynamicMesh.Attributes()->NumWeightLayers(); ++Idx)
+			{	
+				const FDynamicMeshWeightAttribute* WeightLayer = TargetDynamicMesh.Attributes()->GetWeightLayer(Idx);
+				if (WeightLayer && WeightLayer->GetName() == FName(TransferSettings.InpaintMaskWeightMapName))
+				{
+					ForceInpaintWeightLayer = WeightLayer;
+					break;
+				}
+			}
+		}
+
 		// Iterate over each component and perform per-component skin weight transfer
 		const int32 NumComponents = ConnectedComponents.Num();
 		for (int32 ComponentIdx = 0; ComponentIdx < NumComponents; ++ComponentIdx)
@@ -396,6 +409,23 @@ namespace UE::Chaos::ClothAsset::Private
 				SubmeshToBaseT.Add(TID);
 			}
 
+			// Copy over the force inpaint values
+			if (ForceInpaintWeightLayer != nullptr)
+			{
+				TArray<float> ForceInpaint;
+				ForceInpaint.SetNum(Submesh.MaxVertexID());
+				
+				for (int32 VID = 0; VID < Submesh.MaxVertexID(); ++VID)
+				{	
+					float Value = 0;
+					ForceInpaintWeightLayer->GetValue(SubmeshToBaseV[VID], &Value);
+					ForceInpaint[VID] = Value;
+				}
+
+				// Set the mask in the transfer operator
+				TransferBoneWeights.ForceInpaint = MoveTemp(ForceInpaint);
+			}
+
 			// Transfer weights to the current submesh only. If transfer using the inpaint method fails, fall back to  
 			// the closest point method which should always succeed. Common reason for the failure is if we didn't find  
 			// any matches on the source at all with the current transfer settings.
@@ -451,8 +481,6 @@ namespace UE::Chaos::ClothAsset::Private
 	{
 		using namespace UE::Geometry;
 
-		// We cannot handle it if there are orphaned vertices, so strip those now.
-		UE::Chaos::ClothAsset::FClothGeometryTools::CleanupAndCompactMesh(ClothCollection);
 		UE::Chaos::ClothAsset::FCollectionClothFacade ClothFacade(ClothCollection);
 
 		// Convert cloth sim mesh LOD to the welded dynamic sim mesh.
@@ -754,6 +782,9 @@ void FChaosClothAssetTransferSkinWeightsNode::Evaluate(Dataflow::FContext& Conte
 			FCollectionClothFacade ClothFacade(ClothCollection);
 			ClothFacade.SetSkeletalMeshPathName(SkeletalMesh->GetPathName());
 
+			// Clean up orphaned vertices
+			UE::Chaos::ClothAsset::FClothGeometryTools::CleanupAndCompactMesh(ClothCollection);
+
 			bool bTransferResult = false;
 			if (TransferMethod == EChaosClothAssetTransferSkinWeightsMethod::InpaintWeights)
 			{
