Cloth Selection Tool: add a flood fill deselect action: shift+ctrl+click will deselect all connected mesh elements

diff --git a/Engine/Plugins/Experimental/ChaosClothAssetEditor/Source/ChaosClothAssetEditorTools/Private/ChaosClothAsset/ClothMeshSelectionTool.cpp b/Engine/Plugins/Experimental/ChaosClothAssetEditor/Source/ChaosClothAssetEditorTools/Private/ChaosClothAsset/ClothMeshSelectionTool.cpp
--- a/Engine/Plugins/Experimental/ChaosClothAssetEditor/Source/ChaosClothAssetEditorTools/Private/ChaosClothAsset/ClothMeshSelectionTool.cpp
+++ b/Engine/Plugins/Experimental/ChaosClothAssetEditor/Source/ChaosClothAssetEditorTools/Private/ChaosClothAsset/ClothMeshSelectionTool.cpp
@@ -23,6 +23,7 @@
 #include "DynamicMesh/NonManifoldMappingSupport.h"
 #include "Selections/GeometrySelectionUtil.h"
 #include "Materials/Material.h"
+#include "Selections/MeshConnectedComponents.h"
 
 #include UE_INLINE_GENERATED_CPP_BY_NAME(ClothMeshSelectionTool)
 
@@ -86,6 +87,101 @@ void UClothMeshSelectionToolProperties::PostEditChangeProperty(FPropertyChangedE
 }
 
 
+// ------------------- Selection Mechanic -------------------
+
+bool UClothMeshSelectionMechanic::UpdateSelection(const FRay& WorldRay, FVector3d& LocalHitPositionOut, FVector3d& LocalHitNormalOut)
+{
+	using namespace UE::Geometry;
+
+	if (bShiftToggle && bCtrlToggle)		// TODO: Shift + Ctrl means something different when marquee is active
+	{
+		FRay3d LocalRay(TargetTransform.InverseTransformPosition((FVector3d)WorldRay.Origin), TargetTransform.InverseTransformVector((FVector3d)WorldRay.Direction));
+		UE::Geometry::Normalize(LocalRay.Direction);
+
+		const FGroupTopologySelection PreviousSelection = PersistentSelection;
+
+		FVector3d LocalPosition, LocalNormal;
+		FGroupTopologySelection Selection;
+		const FGroupTopologySelector::FSelectionSettings TopoSelectorSettings = GetTopoSelectorSettings(CameraState.bIsOrthographic);
+		if (TopoSelector->FindSelectedElement(TopoSelectorSettings, LocalRay, Selection, LocalPosition, LocalNormal))
+		{
+			LocalHitPositionOut = LocalPosition;
+			LocalHitNormalOut = LocalNormal;
+
+			// Get seed selection
+
+			const FTopologyProvider* const TopologyProvider = TopoSelector->GetTopologyProvider();
+
+			if (Properties->bSelectFaces && Selection.SelectedGroupIDs.Num() > 0)
+			{
+				// TopologyProvider doesn't provide an interface to get triangle indices from GroupIDs, so we ray cast again
+				// TODO: Implement GroupID->Triangles function in FTopologyProvider, similar to GetCornerVertexID()
+
+				FDynamicMeshAABBTree3* const Spatial = GetSpatialFunc();
+				check(Spatial);
+				const int32 TriangleID = Spatial->FindNearestHitTriangle(LocalRay);
+				if (TriangleID != IndexConstants::InvalidID)
+				{
+					TSet<int32> ConnectedTriangles;
+					FMeshConnectedComponents::GrowToConnectedTriangles(Mesh, { TriangleID }, ConnectedTriangles);
+
+					for (const int32 ConnectedTriangleID : ConnectedTriangles)
+					{
+						const int32 GroupID = TopologyProvider->GetGroupIDForTriangle(ConnectedTriangleID);
+						if (TopologyProvider->GetGroupIDAt(GroupID) != -1)
+						{
+							PersistentSelection.SelectedGroupIDs.Remove(GroupID);
+						}
+					}
+				}
+			}
+			else if (!Properties->bSelectFaces && Selection.SelectedCornerIDs.Num() > 0)
+			{
+				const int32 CornerID = *Selection.SelectedCornerIDs.CreateConstIterator();
+				const int32 VertexID = TopologyProvider->GetCornerVertexID(CornerID);
+
+				TSet<int32> ConnectedVertices;
+				FMeshConnectedComponents::GrowToConnectedVertices(*Mesh, { VertexID }, ConnectedVertices);
+
+				TArray<int32> CornersToRemove;
+
+				// FGroupTopology has GetCornerIDFromVertexID() but is inaccessible
+				// TODO: Expose Vertex -> CornerID map in FTopologyProvider
+
+				for (const int32 ConnectedVertex : ConnectedVertices)
+				{
+					for (const int32 SelectedCornerID : PersistentSelection.SelectedCornerIDs)
+					{
+						if (TopologyProvider->GetCornerVertexID(SelectedCornerID) == ConnectedVertex)
+						{
+							CornersToRemove.Add(SelectedCornerID);
+							break;
+						}
+					}
+				}
+				for (const int32 Remove : CornersToRemove)
+				{
+					PersistentSelection.SelectedCornerIDs.Remove(Remove);
+				}
+			}
+
+			if (PersistentSelection != PreviousSelection)
+			{
+				SelectionTimestamp++;
+				OnSelectionChanged.Broadcast();
+				return true;
+			}
+		}
+	}
+	else
+	{
+		return UPolygonSelectionMechanic::UpdateSelection(WorldRay, LocalHitPositionOut, LocalHitNormalOut);
+	}
+
+	return false;
+}
+
+
 // ------------------- Tool -------------------
 	
 void UClothMeshSelectionTool::Setup()
@@ -129,7 +225,7 @@ void UClothMeshSelectionTool::Setup()
 	// SelectionMechanic
 	//
 
-	SelectionMechanic = NewObject<UPolygonSelectionMechanic>(this);
+	SelectionMechanic = NewObject<UClothMeshSelectionMechanic>(this);
 	SelectionMechanic->bAddSelectionFilterPropertiesToParentTool = false;   // We'll do this ourselves later
 	SelectionMechanic->Setup(this);
 	SelectionMechanic->Properties->RestoreProperties(this);
diff --git a/Engine/Plugins/Experimental/ChaosClothAssetEditor/Source/ChaosClothAssetEditorTools/Public/ChaosClothAsset/ClothMeshSelectionTool.h b/Engine/Plugins/Experimental/ChaosClothAssetEditor/Source/ChaosClothAssetEditorTools/Public/ChaosClothAsset/ClothMeshSelectionTool.h
--- a/Engine/Plugins/Experimental/ChaosClothAssetEditor/Source/ChaosClothAssetEditorTools/Public/ChaosClothAsset/ClothMeshSelectionTool.h
+++ b/Engine/Plugins/Experimental/ChaosClothAssetEditor/Source/ChaosClothAssetEditorTools/Public/ChaosClothAsset/ClothMeshSelectionTool.h
@@ -4,6 +4,7 @@
 
 #include "BaseTools/SingleSelectionMeshEditingTool.h"
 #include "ChaosClothAsset/ClothEditorToolBuilder.h"
+#include "Selection/PolygonSelectionMechanic.h"
 #include "ClothMeshSelectionTool.generated.h"
 
 class UPolygonSelectionMechanic;
@@ -44,6 +45,18 @@ enum class EClothMeshSelectionToolActions
 	FloodSelection
 };
 
+
+UCLASS()
+class UClothMeshSelectionMechanic : public UPolygonSelectionMechanic
+{
+	GENERATED_BODY()
+
+private:
+
+	virtual bool UpdateSelection(const FRay& WorldRay, FVector3d& LocalHitPositionOut, FVector3d& LocalHitNormalOut) override;
+};
+
+
 UCLASS()
 class CHAOSCLOTHASSETEDITORTOOLS_API UClothMeshSelectionToolActions :  public UInteractiveToolPropertySet
 {
