Fix crash when deleting selection node while selection tool is active

diff --git a/Engine/Plugins/Experimental/ChaosClothAssetEditor/Source/ChaosClothAssetEditorTools/Public/ChaosClothAsset/ClothEditorContextObject.h b/Engine/Plugins/Experimental/ChaosClothAssetEditor/Source/ChaosClothAssetEditorTools/Public/ChaosClothAsset/ClothEditorContextObject.h
--- a/Engine/Plugins/Experimental/ChaosClothAssetEditor/Source/ChaosClothAssetEditorTools/Public/ChaosClothAsset/ClothEditorContextObject.h
+++ b/Engine/Plugins/Experimental/ChaosClothAssetEditor/Source/ChaosClothAssetEditorTools/Public/ChaosClothAsset/ClothEditorContextObject.h
@@ -37,10 +37,12 @@ public:
 		if (UEdGraphNode* const SingleSelectedNode = GraphEditor->GetSingleSelectedNode())
 		{
 			UDataflowEdNode* const SelectedDataflowEdNode = CastChecked<UDataflowEdNode>(SingleSelectedNode);
-			const TSharedPtr<FDataflowNode> DataflowNode = SelectedDataflowEdNode->GetDataflowNode();
-			if (NodeType* const NodeTypeNode = DataflowNode->AsType<NodeType>())
+			if (const TSharedPtr<FDataflowNode> DataflowNode = SelectedDataflowEdNode->GetDataflowNode())	// If the user deletes a node from the graph, the UDataflowEdNode might outlast the FDataflowNode
 			{
-				return NodeTypeNode;
+				if (NodeType* const NodeTypeNode = DataflowNode->AsType<NodeType>())
+				{
+					return NodeTypeNode;
+				}
 			}
 		}
 
