Fixed a bug upon saving the asset as another asset where the new asset's dataflow was still pointing to the old dataflow.

#rnx

diff --git a/Engine/Plugins/Experimental/ChaosClothAssetEditor/Source/ChaosClothAssetEditor/Private/ChaosClothAsset/ClothEditorToolkit.cpp b/Engine/Plugins/Experimental/ChaosClothAssetEditor/Source/ChaosClothAssetEditor/Private/ChaosClothAsset/ClothEditorToolkit.cpp
--- a/Engine/Plugins/Experimental/ChaosClothAssetEditor/Source/ChaosClothAssetEditor/Private/ChaosClothAsset/ClothEditorToolkit.cpp
+++ b/Engine/Plugins/Experimental/ChaosClothAssetEditor/Source/ChaosClothAssetEditor/Private/ChaosClothAsset/ClothEditorToolkit.cpp
@@ -235,7 +235,17 @@ TSharedPtr<Dataflow::FEngineContext> FChaosClothAssetEditorToolkit::GetDataflowC
 
 const UDataflow* FChaosClothAssetEditorToolkit::GetDataflow() const
 {
-	return Dataflow;
+	PRAGMA_DISABLE_DEPRECATION_WARNINGS  // TODO: Don't use public property, and have Getter/Setter API instead
+	return GetAsset()->DataflowAsset;
+	PRAGMA_ENABLE_DEPRECATION_WARNINGS
+}
+
+
+UDataflow* FChaosClothAssetEditorToolkit::GetDataflow()
+{
+	PRAGMA_DISABLE_DEPRECATION_WARNINGS  // TODO: Don't use public property, and have Getter/Setter API instead
+	return GetAsset()->DataflowAsset;
+	PRAGMA_ENABLE_DEPRECATION_WARNINGS
 }
 
 //~ Begin FTickableEditorObject overrides
@@ -311,7 +321,7 @@ void FChaosClothAssetEditorToolkit::CreateWidgets()
 
 	if (ClothAsset)
 	{
-		Dataflow = Private::GetDataflowFrom(ClothAsset);
+		UDataflow* const Dataflow = Private::GetDataflowFrom(ClothAsset);
 
 		// TODO: Figure out how to create the GraphEditor widgets when the ClothAsset doesn't have a Dataflow property set
 		if (Dataflow)
@@ -479,7 +489,7 @@ void FChaosClothAssetEditorToolkit::PostInitAssetEditor()
 
 	InitDetailsViewPanel();
 
-	ClothMode->DataflowGraph = Dataflow;
+	ClothMode->DataflowGraph = GetDataflow();
 	ClothMode->SetDataflowGraphEditor(GraphEditor);
 }
 
@@ -786,7 +796,7 @@ void FChaosClothAssetEditorToolkit::OnFinishedChangingAssetProperties(const FPro
 		const UChaosClothAsset* const ClothAsset = GetAsset();
 		if (ClothAsset)
 		{
-			Dataflow = Private::GetDataflowFrom(ClothAsset);
+			UDataflow* const Dataflow = Private::GetDataflowFrom(ClothAsset);
 
 			if (Dataflow)
 			{
@@ -812,6 +822,7 @@ void FChaosClothAssetEditorToolkit::OnFinishedChangingAssetProperties(const FPro
 void FChaosClothAssetEditorToolkit::EvaluateNode(FDataflowNode* Node, bool bForceOperation)
 {
 	UChaosClothAsset* const ClothAsset = GetAsset();
+	UDataflow* const Dataflow = GetDataflow(); 
 
 	if (Dataflow && ClothAsset)
 	{
@@ -933,6 +944,7 @@ void FChaosClothAssetEditorToolkit::EvaluateNode(FDataflowNode* Node, bool bForc
 
 TSharedRef<SDataflowGraphEditor> FChaosClothAssetEditorToolkit::CreateGraphEditorWidget()
 {
+	UDataflow* const Dataflow = GetDataflow();
 	ensure(Dataflow);
 	using namespace Dataflow;
 
@@ -961,6 +973,8 @@ TSharedRef<SDataflowGraphEditor> FChaosClothAssetEditorToolkit::CreateGraphEdito
 
 void FChaosClothAssetEditorToolkit::ReinitializeGraphEditorWidget()
 {
+	UDataflow* const Dataflow = GetDataflow();
+
 	ensure(Dataflow);
 
 	const auto EvalLambda = [this](FDataflowNode* Node, FDataflowOutput* /*Out*/)
@@ -1037,7 +1051,7 @@ TSharedPtr<IStructureDetailsView> FChaosClothAssetEditorToolkit::CreateNodeDetai
 
 void FChaosClothAssetEditorToolkit::OnPropertyValueChanged(const FPropertyChangedEvent& PropertyChangedEvent)
 {
-	FDataflowEditorCommands::OnPropertyValueChanged(Dataflow, DataflowContext, LastDataflowNodeTimestamp, PropertyChangedEvent, GraphEditor ? GraphEditor->GetSelectedNodes() : FGraphPanelSelectionSet());
+	FDataflowEditorCommands::OnPropertyValueChanged(GetDataflow(), DataflowContext, LastDataflowNodeTimestamp, PropertyChangedEvent, GraphEditor ? GraphEditor->GetSelectedNodes() : FGraphPanelSelectionSet());
 }
 
 bool FChaosClothAssetEditorToolkit::OnNodeVerifyTitleCommit(const FText& NewText, UEdGraphNode* GraphNode, FText& OutErrorMessage) const
@@ -1128,6 +1142,8 @@ void FChaosClothAssetEditorToolkit::OnNodeSelectionChanged(const TSet<UObject*>&
 	// Despite this function's name, we might not have actually changed which node is selected
 	bool bNodeSelectionChanged = false;
 
+	UDataflow* const Dataflow = GetDataflow();
+	
 	if (Dataflow)
 	{
 		Dataflow->RenderTargets.Reset();
@@ -1257,8 +1273,9 @@ void FChaosClothAssetEditorToolkit::OnNodeSingleClicked(UObject* ClickedNode) co
 }
 
 
-void FChaosClothAssetEditorToolkit::OnNodeDeleted(const TSet<UObject*>& DeletedNodes) const
+void FChaosClothAssetEditorToolkit::OnNodeDeleted(const TSet<UObject*>& DeletedNodes)
 {
+	UDataflow* const Dataflow = GetDataflow();
 	if (Dataflow)
 	{
 		Dataflow->RenderTargets.SetNum(Algo::RemoveIf(Dataflow->RenderTargets, [&DeletedNodes](const UDataflowEdNode* RenderTarget)
diff --git a/Engine/Plugins/Experimental/ChaosClothAssetEditor/Source/ChaosClothAssetEditor/Private/ChaosClothAsset/ClothEditorToolkit.h b/Engine/Plugins/Experimental/ChaosClothAssetEditor/Source/ChaosClothAssetEditor/Private/ChaosClothAsset/ClothEditorToolkit.h
--- a/Engine/Plugins/Experimental/ChaosClothAssetEditor/Source/ChaosClothAssetEditor/Private/ChaosClothAsset/ClothEditorToolkit.h
+++ b/Engine/Plugins/Experimental/ChaosClothAssetEditor/Source/ChaosClothAssetEditor/Private/ChaosClothAsset/ClothEditorToolkit.h
@@ -104,7 +104,8 @@ private:
 
 	// Return the cloth asset held by the Cloth Editor
 	UChaosClothAsset* GetAsset() const;
-
+	UDataflow* GetDataflow();
+	
 	TSharedRef<SDockTab> SpawnTab_ClothPreview(const FSpawnTabArgs& Args);
 	TSharedRef<SDockTab> SpawnTab_Outliner(const FSpawnTabArgs& Args);
 	TSharedRef<SDockTab> SpawnTab_PreviewSceneDetails(const FSpawnTabArgs& Args);
@@ -135,7 +136,7 @@ private:
 	bool OnNodeVerifyTitleCommit(const FText& NewText, UEdGraphNode* GraphNode, FText& OutErrorMessage) const;
 	void OnNodeTitleCommitted(const FText& InNewText, ETextCommit::Type InCommitType, UEdGraphNode* GraphNode) const;
 	void OnNodeSelectionChanged(const TSet<UObject*>& NewSelection);
-	void OnNodeDeleted(const TSet<UObject*>& DeletedNodes) const;
+	void OnNodeDeleted(const TSet<UObject*>& DeletedNodes);
 	void OnNodeSingleClicked(UObject* ClickedNode) const;
 
 	/** Scene in which the 3D sim space preview meshes live. Ownership shared with AdvancedPreviewSettingsWidget*/
@@ -158,7 +159,6 @@ private:
 	TSharedPtr<SClothCollectionOutliner> Outliner;
 
 	// Dataflow
-	UDataflow* Dataflow = nullptr;
 	TSharedPtr<Dataflow::FEngineContext> DataflowContext;
 	Dataflow::FTimestamp LastDataflowNodeTimestamp = Dataflow::FTimestamp::Invalid;
 	FDelegateHandle OnNodeInvalidatedDelegateHandle;
