Cloth Selection Node: allow selection nodes to remove elements from an existing selection, rather than just add to it

diff --git a/Engine/Plugins/Experimental/ChaosClothAssetEditor/Source/ChaosClothAssetDataflowNodes/Private/ChaosClothAsset/SelectionNode.cpp b/Engine/Plugins/Experimental/ChaosClothAssetEditor/Source/ChaosClothAssetDataflowNodes/Private/ChaosClothAsset/SelectionNode.cpp
--- a/Engine/Plugins/Experimental/ChaosClothAssetEditor/Source/ChaosClothAssetDataflowNodes/Private/ChaosClothAsset/SelectionNode.cpp
+++ b/Engine/Plugins/Experimental/ChaosClothAssetEditor/Source/ChaosClothAssetDataflowNodes/Private/ChaosClothAsset/SelectionNode.cpp
@@ -30,6 +30,8 @@ void FChaosClothAssetSelectionNode::Evaluate(Dataflow::FContext& Context, const
 		const int32 NumElementsInGroup = SelectionCollection->NumElements(SelectionGroupName);
 		bool bFoundAnyInvalidIndex = false;
 
+		DestSelectionSet.Reset();
+
 		for (const int32 Index : SourceIndices)
 		{
 			if (Index < 0 || Index >= NumElementsInGroup)
