Chaos Cloth Asset - Added a RecomputeTangents optional attribute to the Cloth Asset RenderPatterns group. [Edigrate CL 32180277]

[preflight] 65fb0972286546fba6cdcbe7
#tests Successfully checked that adding a new RecomputeTangents attribute removes some blotchiness from the cloth.

diff --git a/Engine/Plugins/Experimental/ChaosClothAsset/Source/ChaosClothAsset/Private/ChaosClothAsset/CollectionClothFacade.cpp b/Engine/Plugins/Experimental/ChaosClothAsset/Source/ChaosClothAsset/Private/ChaosClothAsset/CollectionClothFacade.cpp
--- a/Engine/Plugins/Experimental/ChaosClothAsset/Source/ChaosClothAsset/Private/ChaosClothAsset/CollectionClothFacade.cpp
+++ b/Engine/Plugins/Experimental/ChaosClothAsset/Source/ChaosClothAsset/Private/ChaosClothAsset/CollectionClothFacade.cpp
@@ -555,7 +555,6 @@ namespace UE::Chaos::ClothAsset
 		const int32 StartNumSimVertices2D = GetNumSimVertices2D();
 		const int32 StartNumSimPatterns = GetNumSimPatterns();
 		const int32 OtherNumSimPatterns = Other.GetNumSimPatterns();
-		const int32 OtherNumSimFaces = Other.GetNumSimFaces();
 		SetNumSimPatterns(StartNumSimPatterns + OtherNumSimPatterns);
 		for (int32 PatternIndex = 0; PatternIndex < OtherNumSimPatterns; ++PatternIndex)
 		{
@@ -605,12 +604,35 @@ namespace UE::Chaos::ClothAsset
 			FClothCollection::CopyArrayViewData(GetWeightMap(WeightMapName).Right(OtherNumSimVertices3D), Other.GetWeightMap(WeightMapName));
 		}
 
-		// Face int maps (self collision layers)
-		const TArray<FName> FaceIntMapNames = Other.GetUserDefinedAttributeNames<int32>(ClothCollectionGroup::SimFaces);
-		for (const FName& IntMapName : FaceIntMapNames)
+		// Copy user defined attributes
+		static const TArray<FName> ClothCollectionGroups =
 		{
-			AddUserDefinedAttribute<int32>(IntMapName, ClothCollectionGroup::SimFaces);
-			FClothCollection::CopyArrayViewData(GetUserDefinedAttribute<int32>(IntMapName, ClothCollectionGroup::SimFaces).Right(OtherNumSimFaces), Other.GetUserDefinedAttribute<int32>(IntMapName, ClothCollectionGroup::SimFaces));
+			ClothCollectionGroup::SimFaces,  // Face int maps (self collision layers)
+			ClothCollectionGroup::RenderPatterns  // RecomputeTangents
+		};
+
+		for (const FName& Group : ClothCollectionGroups)
+		{
+			const int32 OtherNumElements = Other.ClothCollection->GetNumElements(Group);
+
+			const TArray<FName> IntAttributeNames = Other.GetUserDefinedAttributeNames<int32>(Group);
+			for (const FName& AttributeName : IntAttributeNames)
+			{
+				AddUserDefinedAttribute<int32>(AttributeName, Group);
+				FClothCollection::CopyArrayViewData(GetUserDefinedAttribute<int32>(AttributeName, Group).Right(OtherNumElements), Other.GetUserDefinedAttribute<int32>(AttributeName, Group));
+			}
+			const TArray<FName> FloatAttributeNames = Other.GetUserDefinedAttributeNames<float>(Group);
+			for (const FName& AttributeName : FloatAttributeNames)
+			{
+				AddUserDefinedAttribute<float>(AttributeName, Group);
+				FClothCollection::CopyArrayViewData(GetUserDefinedAttribute<float>(AttributeName, Group).Right(OtherNumElements), Other.GetUserDefinedAttribute<float>(AttributeName, Group));
+			}
+			const TArray<FName> VectorAttributeNames = Other.GetUserDefinedAttributeNames<FVector3f>(Group);
+			for (const FName& AttributeName : VectorAttributeNames)
+			{
+				AddUserDefinedAttribute<FVector3f>(AttributeName, Group);
+				FClothCollection::CopyArrayViewData(GetUserDefinedAttribute<FVector3f>(AttributeName, Group).Right(OtherNumElements), Other.GetUserDefinedAttribute<FVector3f>(AttributeName, Group));
+			}
 		}
 	}
 
diff --git a/Engine/Plugins/Experimental/ChaosClothAsset/Source/ChaosClothAsset/Public/ChaosClothAsset/CollectionClothFacade.h b/Engine/Plugins/Experimental/ChaosClothAsset/Source/ChaosClothAsset/Public/ChaosClothAsset/CollectionClothFacade.h
--- a/Engine/Plugins/Experimental/ChaosClothAsset/Source/ChaosClothAsset/Public/ChaosClothAsset/CollectionClothFacade.h
+++ b/Engine/Plugins/Experimental/ChaosClothAsset/Source/ChaosClothAsset/Public/ChaosClothAsset/CollectionClothFacade.h
@@ -157,6 +157,7 @@ namespace UE::Chaos::ClothAsset
 	protected:
 		TSharedRef<const class FClothCollection> ClothCollection;
 
+		friend class FCollectionClothFacade;  // To enable access from a different instance
 		friend class FCollectionClothSeamFacade;
 		friend class FCollectionClothSeamConstFacade;
 		friend class FCollectionClothFabricFacade;
diff --git a/Engine/Plugins/Experimental/ChaosClothAsset/Source/ChaosClothAssetEngine/Private/ChaosClothAsset/ClothAssetBuilder.cpp b/Engine/Plugins/Experimental/ChaosClothAsset/Source/ChaosClothAssetEngine/Private/ChaosClothAsset/ClothAssetBuilder.cpp
--- a/Engine/Plugins/Experimental/ChaosClothAsset/Source/ChaosClothAssetEngine/Private/ChaosClothAsset/ClothAssetBuilder.cpp
+++ b/Engine/Plugins/Experimental/ChaosClothAsset/Source/ChaosClothAssetEngine/Private/ChaosClothAsset/ClothAssetBuilder.cpp
@@ -5,6 +5,7 @@
 #include "ChaosClothAsset/ClothAssetBuilder.h"
 #include "ChaosClothAsset/ClothAsset.h"
 #include "ChaosClothAsset/ClothSimulationModel.h"
+#include "ChaosClothAsset/ClothCollectionGroup.h"
 #include "ChaosClothAsset/CollectionClothFacade.h"
 #include "ChaosClothAsset/ClothAssetPrivate.h"
 #include "Chaos/CollectionPropertyFacade.h"
@@ -174,6 +175,8 @@ void UChaosClothAsset::FBuilder::BuildLod(FSkeletalMeshLODModel& LODModel, const
 	FScopedSlowTask SlowTask((float)ClothFacade.GetNumRenderFaces(), LOCTEXT("ClothAssetBuildLOD", "Building Cloth Asset LOD sections..."));
 	SlowTask.MakeDialogDelayed(1.f);
 
+	const TConstArrayView<int32> RecomputeTangent = ClothFacade.GetUserDefinedAttribute<int32>(TEXT("RecomputeTangents"), ClothCollectionGroup::RenderPatterns);
+
 	int32 BaseIndex = 0;
 	for (int32 SectionIndex = 0; SectionIndex < NumSections; ++SectionIndex)
 	{
@@ -196,6 +199,9 @@ void UChaosClothAsset::FBuilder::BuildLod(FSkeletalMeshLODModel& LODModel, const
 		// Build the section face data (indices)
 		Section.MaterialIndex = (uint16)MaterialIndex;
 
+		Section.bRecomputeTangent = RecomputeTangent.IsValidIndex(MaterialIndex) && RecomputeTangent[MaterialIndex];
+		Section.RecomputeTangentsVertexMaskChannel = ESkinVertexColorChannel::None;
+
 		Section.BaseIndex = (uint32)BaseIndex;
 		BaseIndex += NumIndices;
 
diff --git a/Engine/Plugins/Experimental/ChaosClothAssetEditor/Source/ChaosClothAssetDataflowNodes/Private/ChaosClothAsset/AttributeNode.cpp b/Engine/Plugins/Experimental/ChaosClothAssetEditor/Source/ChaosClothAssetDataflowNodes/Private/ChaosClothAsset/AttributeNode.cpp
new file mode 100644
--- /dev/null
+++ b/Engine/Plugins/Experimental/ChaosClothAssetEditor/Source/ChaosClothAssetDataflowNodes/Private/ChaosClothAsset/AttributeNode.cpp
@@ -0,0 +1,88 @@
+// Copyright Epic Games, Inc. All Rights Reserved.
+
+#include "ChaosClothAsset/AttributeNode.h"
+#include "ChaosClothAsset/ClothCollectionGroup.h"
+#include "ChaosClothAsset/CollectionClothFacade.h"
+#include "ChaosClothAsset/ClothDataflowTools.h"
+#include "Dataflow/DataflowInputOutput.h"
+
+#include UE_INLINE_GENERATED_CPP_BY_NAME(AttributeNode)
+
+#define LOCTEXT_NAMESPACE "ChaosClothAssetAttributeNode"
+
+FChaosClothAssetAttributeNode::FChaosClothAssetAttributeNode(const Dataflow::FNodeParameters& InParam, FGuid InGuid)
+	: FDataflowNode(InParam, InGuid)
+{
+	RegisterInputConnection(&Collection);
+	RegisterOutputConnection(&Collection, &Collection);
+	RegisterOutputConnection(&Name);
+}
+
+void FChaosClothAssetAttributeNode::Evaluate(Dataflow::FContext& Context, const FDataflowOutput* Out) const
+{
+	if (Out->IsA<FManagedArrayCollection>(&Collection))
+	{
+		using namespace UE::Chaos::ClothAsset;
+
+		// Evaluate in collection
+		FManagedArrayCollection InCollection = GetValue<FManagedArrayCollection>(Context, &Collection);
+		const TSharedRef<FManagedArrayCollection> ClothCollection = MakeShared<FManagedArrayCollection>(MoveTemp(InCollection));
+
+		FCollectionClothFacade Cloth(ClothCollection);
+		const FName GroupName = *Group.Name;
+		
+		if (Cloth.IsValid() && !Name.IsEmpty())
+		{
+			if (ClothCollection->HasGroup(GroupName))
+			{
+				switch (Type)
+				{
+				case EChaosClothAssetNodeAttributeType::Integer:
+					Cloth.AddUserDefinedAttribute<int32>(*Name, GroupName);
+					{
+						TArrayView<int32> Values = Cloth.GetUserDefinedAttribute<int32>(*Name, GroupName);
+						for (int32& Value : Values)
+						{
+							Value = IntValue;
+						}
+					}
+					break;
+				case EChaosClothAssetNodeAttributeType::Float:
+					Cloth.AddUserDefinedAttribute<float>(*Name, GroupName);
+					{
+						TArrayView<float> Values = Cloth.GetUserDefinedAttribute<float>(*Name, GroupName);
+						for (float& Value : Values)
+						{
+							Value = FloatValue;
+						}
+					}
+					break;
+				case EChaosClothAssetNodeAttributeType::Vector:
+					Cloth.AddUserDefinedAttribute<FVector3f>(*Name, GroupName);
+					{
+						TArrayView<FVector3f> Values = Cloth.GetUserDefinedAttribute<FVector3f>(*Name, GroupName);
+						for (FVector3f& Value : Values)
+						{
+							Value = VectorValue;
+						}
+					}
+					break;
+				}
+			}
+			else if (!GroupName.IsNone())
+			{
+				FClothDataflowTools::LogAndToastWarning(
+					*this,
+					LOCTEXT("CreateAttributeHeadline", "Invalid Group"),
+					FText::Format(LOCTEXT("CreateAttributeDetail", "No group \"{0}\" currently exists on the input collection"), FText::FromName(GroupName)));
+			}
+		}
+		SetValue(Context, MoveTemp(*ClothCollection), &Collection);
+	}
+	else if (Out->IsA<FString>(&Name))
+	{
+		SetValue(Context, Name, &Name);
+	}
+}
+
+#undef LOCTEXT_NAMESPACE
diff --git a/Engine/Plugins/Experimental/ChaosClothAssetEditor/Source/ChaosClothAssetDataflowNodes/Private/ChaosClothAsset/ChaosClothAssetDataflowNodesModule.cpp b/Engine/Plugins/Experimental/ChaosClothAssetEditor/Source/ChaosClothAssetDataflowNodes/Private/ChaosClothAsset/ChaosClothAssetDataflowNodesModule.cpp
--- a/Engine/Plugins/Experimental/ChaosClothAssetEditor/Source/ChaosClothAssetDataflowNodes/Private/ChaosClothAsset/ChaosClothAssetDataflowNodesModule.cpp
+++ b/Engine/Plugins/Experimental/ChaosClothAssetEditor/Source/ChaosClothAssetDataflowNodes/Private/ChaosClothAsset/ChaosClothAssetDataflowNodesModule.cpp
@@ -2,6 +2,7 @@
 
 #include "ChaosClothAsset/AddStitchNode.h"
 #include "ChaosClothAsset/AddWeightMapNode.h"
+#include "ChaosClothAsset/AttributeNode.h"
 #include "ChaosClothAsset/BindToRootBoneNode.h"
 #include "ChaosClothAsset/ColorScheme.h"
 #include "ChaosClothAsset/CopySimulationToRenderMeshNode.h"
@@ -71,6 +72,7 @@ namespace UE::Chaos::ClothAsset
 		{
 			DATAFLOW_NODE_REGISTER_CREATION_FACTORY(FChaosClothAssetAddStitchNode);
 			DATAFLOW_NODE_REGISTER_CREATION_FACTORY(FChaosClothAssetAddWeightMapNode);
+			DATAFLOW_NODE_REGISTER_CREATION_FACTORY(FChaosClothAssetAttributeNode);
 			DATAFLOW_NODE_REGISTER_CREATION_FACTORY(FChaosClothAssetBindToRootBoneNode);
 			DATAFLOW_NODE_REGISTER_CREATION_FACTORY(FChaosClothAssetCopySimulationToRenderMeshNode);
 			DATAFLOW_NODE_REGISTER_CREATION_FACTORY(FChaosClothAssetDatasmithImportNode);
@@ -150,6 +152,7 @@ namespace UE::Chaos::ClothAsset
 				PropertyModule->RegisterCustomPropertyTypeLayout(FChaosClothAssetImportFilePath::StaticStruct()->GetFName(), FOnGetPropertyTypeCustomizationInstance::CreateStatic(&FImportFilePathCustomization::MakeInstance));
 				PropertyModule->RegisterCustomPropertyTypeLayout(FChaosClothAssetTerminalNodeRefreshAsset::StaticStruct()->GetFName(), FOnGetPropertyTypeCustomizationInstance::CreateStatic(&FTerminalNodeRefreshAssetCustomization::MakeInstance));
 				PropertyModule->RegisterCustomPropertyTypeLayout(FChaosClothAssetImportNodeRefreshAsset::StaticStruct()->GetFName(), FOnGetPropertyTypeCustomizationInstance::CreateStatic(&FTerminalNodeRefreshAssetCustomization::MakeInstance));
+				PropertyModule->RegisterCustomPropertyTypeLayout(FChaosClothAssetNodeAttributeGroup::StaticStruct()->GetFName(), FOnGetPropertyTypeCustomizationInstance::CreateStatic(&FSelectionGroupCustomization::MakeInstance));
 			}
 
 			// Register modular features
@@ -173,6 +176,7 @@ namespace UE::Chaos::ClothAsset
 					PropertyModule->UnregisterCustomPropertyTypeLayout(FChaosClothAssetImportFilePath::StaticStruct()->GetFName());
 					PropertyModule->UnregisterCustomPropertyTypeLayout(FChaosClothAssetTerminalNodeRefreshAsset::StaticStruct()->GetFName());
 					PropertyModule->UnregisterCustomPropertyTypeLayout(FChaosClothAssetImportNodeRefreshAsset::StaticStruct()->GetFName());
+					PropertyModule->UnregisterCustomPropertyTypeLayout(FChaosClothAssetNodeAttributeGroup::StaticStruct()->GetFName());
 				}
 			}
 
diff --git a/Engine/Plugins/Experimental/ChaosClothAssetEditor/Source/ChaosClothAssetDataflowNodes/Public/ChaosClothAsset/AttributeNode.h b/Engine/Plugins/Experimental/ChaosClothAssetEditor/Source/ChaosClothAssetDataflowNodes/Public/ChaosClothAsset/AttributeNode.h
new file mode 100644
--- /dev/null
+++ b/Engine/Plugins/Experimental/ChaosClothAssetEditor/Source/ChaosClothAssetDataflowNodes/Public/ChaosClothAsset/AttributeNode.h
@@ -0,0 +1,70 @@
+// Copyright Epic Games, Inc. All Rights Reserved.
+
+#pragma once 
+
+#include "Dataflow/DataflowNode.h"
+#include "GeometryCollection/ManagedArrayCollection.h"
+#include "ChaosClothAsset/ConnectableValue.h"
+#include "AttributeNode.generated.h"
+
+/**
+ * The managed array collection group used for the attribute creation.
+ * This separate structure is required to allow for customization of the UI.
+ */
+USTRUCT()
+struct FChaosClothAssetNodeAttributeGroup
+{
+	GENERATED_USTRUCT_BODY()
+
+	UPROPERTY(EditAnywhere, Category = "Attribute Group")
+	FString Name;
+};
+
+UENUM()
+enum class EChaosClothAssetNodeAttributeType : uint8
+{
+	Integer,
+	Float,
+	Vector
+};
+
+/** Create a new attribute for the specified group. */
+USTRUCT(Meta = (DataflowCloth))
+struct FChaosClothAssetAttributeNode : public FDataflowNode
+{
+	GENERATED_USTRUCT_BODY()
+	DATAFLOW_NODE_DEFINE_INTERNAL(FChaosClothAssetAttributeNode, "Attribute", "Cloth", "Cloth Attribute")
+
+public:
+
+	UPROPERTY(Meta = (Dataflowinput, DataflowOutput, DataflowPassthrough = "Collection"))
+	FManagedArrayCollection Collection;
+
+	/** The name of the attribute to create. */
+	UPROPERTY(EditAnywhere, Category = "Attribute", Meta = (DataflowOutput))
+	FString Name;
+
+	/** The attribute group. */
+	UPROPERTY(EditAnywhere, Category = "Attribute")
+	FChaosClothAssetNodeAttributeGroup Group;
+
+	/** The attribute type. */
+	UPROPERTY(EditAnywhere, Category = "Attribute")
+	EChaosClothAssetNodeAttributeType Type = EChaosClothAssetNodeAttributeType::Integer;
+
+	/** Default integer value. */
+	UPROPERTY(EditAnywhere, Category = "Attribute", Meta = (EditCondition = "Type == EChaosClothAssetNodeAttributeType::Integer", EditConditionHides))
+	int32 IntValue = 0;
+
+	/** Default float value. */
+	UPROPERTY(EditAnywhere, Category = "Attribute", Meta = (EditCondition = "Type == EChaosClothAssetNodeAttributeType::Float", EditConditionHides))
+	float FloatValue = 0.f;
+
+	/** Default vector value. */
+	UPROPERTY(EditAnywhere, Category = "Attribute", Meta = (EditCondition = "Type == EChaosClothAssetNodeAttributeType::Vector", EditConditionHides))
+	FVector3f VectorValue = FVector3f::ZeroVector;
+
+	FChaosClothAssetAttributeNode(const Dataflow::FNodeParameters& InParam, FGuid InGuid = FGuid::NewGuid());
+
+	virtual void Evaluate(Dataflow::FContext& Context, const FDataflowOutput* Out) const override;
+};
