Chaos Cloth Asset - Fixed Attribute node group customization.

#tests Tried the Attribute node Group combo box, and verified that the cloth collection groups can be selected.

diff --git a/Engine/Plugins/Experimental/ChaosClothAssetEditor/Source/ChaosClothAssetDataflowNodes/Private/ChaosClothAsset/AttributeNode.cpp b/Engine/Plugins/Experimental/ChaosClothAssetEditor/Source/ChaosClothAssetDataflowNodes/Private/ChaosClothAsset/AttributeNode.cpp
--- a/Engine/Plugins/Experimental/ChaosClothAssetEditor/Source/ChaosClothAssetDataflowNodes/Private/ChaosClothAsset/AttributeNode.cpp
+++ b/Engine/Plugins/Experimental/ChaosClothAssetEditor/Source/ChaosClothAssetDataflowNodes/Private/ChaosClothAsset/AttributeNode.cpp
@@ -85,4 +85,20 @@ void FChaosClothAssetAttributeNode::Evaluate(Dataflow::FContext& Context, const
 	}
 }
 
+void FChaosClothAssetAttributeNode::OnSelected(Dataflow::FContext& Context)
+{
+	// Re-evaluate the input collection
+	const FManagedArrayCollection& SelectionCollection = GetValue<FManagedArrayCollection>(Context, &Collection);
+
+	// Update the list of used group for the UI customization
+	CachedCollectionGroupNames = SelectionCollection.GroupNames();
+}
+
+void FChaosClothAssetAttributeNode::OnDeselected()
+{
+	// Clean up, to avoid another toolkit picking up the wrong context evaluation
+	CachedCollectionGroupNames.Reset();
+}
+
+
 #undef LOCTEXT_NAMESPACE
diff --git a/Engine/Plugins/Experimental/ChaosClothAssetEditor/Source/ChaosClothAssetDataflowNodes/Private/ChaosClothAsset/SelectionGroupCustomization.cpp b/Engine/Plugins/Experimental/ChaosClothAssetEditor/Source/ChaosClothAssetDataflowNodes/Private/ChaosClothAsset/SelectionGroupCustomization.cpp
--- a/Engine/Plugins/Experimental/ChaosClothAssetEditor/Source/ChaosClothAssetDataflowNodes/Private/ChaosClothAsset/SelectionGroupCustomization.cpp
+++ b/Engine/Plugins/Experimental/ChaosClothAssetEditor/Source/ChaosClothAssetDataflowNodes/Private/ChaosClothAsset/SelectionGroupCustomization.cpp
@@ -1,9 +1,10 @@
 // Copyright Epic Games, Inc. All Rights Reserved.
 
 #include "ChaosClothAsset/SelectionGroupCustomization.h"
-#include "ChaosClothAsset/SelectionNode.h"
-#include "ChaosClothAsset/DeleteElementNode.h"
+#include "ChaosClothAsset/AttributeNode.h"
 #include "ChaosClothAsset/ClothDataflowTools.h"  // For MakeCollectionName
+#include "ChaosClothAsset/DeleteElementNode.h"
+#include "ChaosClothAsset/SelectionNode.h"
 #include "Widgets/Input/SEditableTextBox.h"
 #include "Widgets/Input/SComboButton.h"
 #include "Widgets/Views/SListView.h"
@@ -119,6 +120,11 @@ namespace UE::Chaos::ClothAsset
 		{
 			CollectionGroupNames = DeleteNode->GetCachedCollectionGroupNames();
 		}
+		else if (const FChaosClothAssetAttributeNode* const AttributeNode =
+			GetOwnerStruct<FChaosClothAssetAttributeNode>())
+		{
+			CollectionGroupNames = AttributeNode->GetCachedCollectionGroupNames();
+		}
 
 		for (const FName& CollectionGroupName : CollectionGroupNames)
 		{
diff --git a/Engine/Plugins/Experimental/ChaosClothAssetEditor/Source/ChaosClothAssetDataflowNodes/Public/ChaosClothAsset/AttributeNode.h b/Engine/Plugins/Experimental/ChaosClothAssetEditor/Source/ChaosClothAssetDataflowNodes/Public/ChaosClothAsset/AttributeNode.h
--- a/Engine/Plugins/Experimental/ChaosClothAssetEditor/Source/ChaosClothAssetDataflowNodes/Public/ChaosClothAsset/AttributeNode.h
+++ b/Engine/Plugins/Experimental/ChaosClothAssetEditor/Source/ChaosClothAssetDataflowNodes/Public/ChaosClothAsset/AttributeNode.h
@@ -66,5 +66,13 @@ public:
 
 	FChaosClothAssetAttributeNode(const Dataflow::FNodeParameters& InParam, FGuid InGuid = FGuid::NewGuid());
 
+	/** Return a cached array of all the groups used by the input collection during at the time of the latest evaluation. */
+	const TArray<FName>& GetCachedCollectionGroupNames() const { return CachedCollectionGroupNames; }
+
+private:
 	virtual void Evaluate(Dataflow::FContext& Context, const FDataflowOutput* Out) const override;
+	virtual void OnSelected(Dataflow::FContext& Context) override;
+	virtual void OnDeselected() override;
+
+	TArray<FName> CachedCollectionGroupNames;
 };
