Solver/Pattern simulated properties from MD + Fix icons width

diff --git a/Engine/Plugins/Experimental/ChaosClothAsset/Source/ChaosClothAsset/Private/ChaosClothAsset/ClothCollection.cpp b/Engine/Plugins/Experimental/ChaosClothAsset/Source/ChaosClothAsset/Private/ChaosClothAsset/ClothCollection.cpp
--- a/Engine/Plugins/Experimental/ChaosClothAsset/Source/ChaosClothAsset/Private/ChaosClothAsset/ClothCollection.cpp
+++ b/Engine/Plugins/Experimental/ChaosClothAsset/Source/ChaosClothAsset/Private/ChaosClothAsset/ClothCollection.cpp
@@ -11,7 +11,16 @@ namespace UE::Chaos::ClothAsset::Private
 	static const TArray<FName> LodsGroupAttributes =
 	{
 		ClothCollectionAttribute::PhysicsAssetPathName,
-		ClothCollectionAttribute::SkeletalMeshPathName
+		ClothCollectionAttribute::SkeletalMeshPathName,
+	};
+
+	// Solvers Group
+	static const TArray<FName> SolverGroupAttributes =
+	{
+		ClothCollectionAttribute::SolverGravity,
+		ClothCollectionAttribute::SolverAirDamping,
+		ClothCollectionAttribute::SolverSubSteps,
+		ClothCollectionAttribute::SolverTimeStep,
 	};
 
 	// Fabrics Group
@@ -21,10 +30,12 @@ namespace UE::Chaos::ClothAsset::Private
 		ClothCollectionAttribute::FabricBucklingStiffness,
 		ClothCollectionAttribute::FabricStretchStiffness,
 		ClothCollectionAttribute::FabricBucklingRatio,
-		ClothCollectionAttribute::FabricClothDensity,
-		ClothCollectionAttribute::FabricClothFriction,
-		ClothCollectionAttribute::FabricClothThickness,
-		ClothCollectionAttribute::FabricClothDamping
+		ClothCollectionAttribute::FabricDensity,
+		ClothCollectionAttribute::FabricFriction,
+		ClothCollectionAttribute::FabricDamping,
+		ClothCollectionAttribute::FabricPressure,
+		ClothCollectionAttribute::FabricLayer,
+		ClothCollectionAttribute::FabricCollisionThickness,
 	};
 
 	// Seam Group
@@ -48,7 +59,8 @@ namespace UE::Chaos::ClothAsset::Private
 		ClothCollectionAttribute::SimVertices2DEnd,
 		ClothCollectionAttribute::SimFacesStart,
 		ClothCollectionAttribute::SimFacesEnd,
-		ClothCollectionAttribute::SimPatternFabric
+		ClothCollectionAttribute::SimPatternFabric,
+		
 	};
 
 	// Render Patterns Group
@@ -128,7 +140,8 @@ namespace UE::Chaos::ClothAsset::Private
 		{ ClothCollectionGroup::SimVertices3D, SimVertices3DGroupAttributes },
 		{ ClothCollectionGroup::RenderFaces, RenderFacesGroupAttributes },
 		{ ClothCollectionGroup::RenderVertices, RenderVerticesGroupAttributes },
-		{ ClothCollectionGroup::Fabrics, SimFabricGroupAttributes }
+		{ ClothCollectionGroup::Fabrics, SimFabricGroupAttributes },
+		{ ClothCollectionGroup::Solvers, SolverGroupAttributes },
 	};
 }  // End namespace UE::Chaos::ClothAsset::Private
 
@@ -143,6 +156,12 @@ namespace UE::Chaos::ClothAsset
 		PhysicsAssetPathName = ManagedArrayCollection->FindAttribute<FString>(ClothCollectionAttribute::PhysicsAssetPathName, ClothCollectionGroup::Lods);
 		SkeletalMeshPathName = ManagedArrayCollection->FindAttribute<FString>(ClothCollectionAttribute::SkeletalMeshPathName, ClothCollectionGroup::Lods);
 
+		// Solvers Group
+		SolverGravity = ManagedArrayCollection->FindAttribute<FVector3f>(ClothCollectionAttribute::SolverGravity, ClothCollectionGroup::Solvers);
+		SolverAirDamping = ManagedArrayCollection->FindAttribute<float>(ClothCollectionAttribute::SolverAirDamping, ClothCollectionGroup::Solvers);
+		SolverSubSteps = ManagedArrayCollection->FindAttribute<int32>(ClothCollectionAttribute::SolverSubSteps, ClothCollectionGroup::Solvers);
+		SolverTimeStep = ManagedArrayCollection->FindAttribute<float>(ClothCollectionAttribute::SolverTimeStep, ClothCollectionGroup::Solvers);
+
 		// Seam Group
 		SeamStitchStart = ManagedArrayCollection->FindAttribute<int32>(ClothCollectionAttribute::SeamStitchStart, ClothCollectionGroup::Seams);
 		SeamStitchEnd = ManagedArrayCollection->FindAttribute<int32>(ClothCollectionAttribute::SeamStitchEnd, ClothCollectionGroup::Seams);
@@ -171,10 +190,12 @@ namespace UE::Chaos::ClothAsset
 		FabricBucklingStiffness = ManagedArrayCollection->FindAttribute<FVector3f>(ClothCollectionAttribute::FabricBucklingStiffness, ClothCollectionGroup::Fabrics);
 		FabricStretchStiffness = ManagedArrayCollection->FindAttribute<FVector3f>(ClothCollectionAttribute::FabricStretchStiffness, ClothCollectionGroup::Fabrics);
 		FabricBucklingRatio = ManagedArrayCollection->FindAttribute<float>(ClothCollectionAttribute::FabricBucklingRatio, ClothCollectionGroup::Fabrics);
-		FabricClothDensity = ManagedArrayCollection->FindAttribute<float>(ClothCollectionAttribute::FabricClothDensity, ClothCollectionGroup::Fabrics);
-		FabricClothFriction = ManagedArrayCollection->FindAttribute<float>(ClothCollectionAttribute::FabricClothFriction, ClothCollectionGroup::Fabrics);
-		FabricClothThickness = ManagedArrayCollection->FindAttribute<float>(ClothCollectionAttribute::FabricClothThickness, ClothCollectionGroup::Fabrics);
-		FabricClothDamping = ManagedArrayCollection->FindAttribute<float>(ClothCollectionAttribute::FabricClothDamping, ClothCollectionGroup::Fabrics);
+		FabricDensity = ManagedArrayCollection->FindAttribute<float>(ClothCollectionAttribute::FabricDensity, ClothCollectionGroup::Fabrics);
+		FabricFriction = ManagedArrayCollection->FindAttribute<float>(ClothCollectionAttribute::FabricFriction, ClothCollectionGroup::Fabrics);
+		FabricDamping = ManagedArrayCollection->FindAttribute<float>(ClothCollectionAttribute::FabricDamping, ClothCollectionGroup::Fabrics);
+		FabricPressure = ManagedArrayCollection->FindAttribute<float>(ClothCollectionAttribute::FabricPressure, ClothCollectionGroup::Fabrics);
+		FabricLayer = ManagedArrayCollection->FindAttribute<int32>(ClothCollectionAttribute::FabricLayer, ClothCollectionGroup::Fabrics);
+		FabricCollisionThickness = ManagedArrayCollection->FindAttribute<float>(ClothCollectionAttribute::FabricCollisionThickness, ClothCollectionGroup::Fabrics);
 		
 		// Sim Faces Group
 		SimIndices2D = ManagedArrayCollection->FindAttribute<FIntVector3>(ClothCollectionAttribute::SimIndices2D, ClothCollectionGroup::SimFaces);
@@ -221,6 +242,13 @@ namespace UE::Chaos::ClothAsset
 			PhysicsAssetPathName &&
 			SkeletalMeshPathName &&
 
+			// Solvers Group
+			(!EnumHasAnyFlags(OptionalSchemas, EClothCollectionOptionalSchemas::Solvers) ||
+				(SolverGravity &&
+				 SolverAirDamping &&
+				 SolverSubSteps &&
+				 SolverTimeStep )) &&
+
 			// Seam Group
 			SeamStitchStart &&
 			SeamStitchEnd &&
@@ -235,6 +263,7 @@ namespace UE::Chaos::ClothAsset
 			SimFacesStart &&
 			SimFacesEnd &&
 			SimPatternFabric &&
+		
 
 			// Render Patterns Group
 			RenderVerticesStart &&
@@ -254,10 +283,12 @@ namespace UE::Chaos::ClothAsset
 			FabricBucklingStiffness &&
 			FabricStretchStiffness &&
 			FabricBucklingRatio &&
-			FabricClothDensity &&
-			FabricClothFriction &&
-			FabricClothThickness &&
-			FabricClothDamping &&
+			FabricDensity &&
+			FabricFriction &&
+			FabricDamping &&
+			FabricPressure &&
+			FabricLayer &&
+			FabricCollisionThickness &&
 
 			// Sim Vertices 2D Group
 			SimPosition2D &&
@@ -313,6 +344,15 @@ namespace UE::Chaos::ClothAsset
 		PhysicsAssetPathName = &ManagedArrayCollection->AddAttribute<FString>(ClothCollectionAttribute::PhysicsAssetPathName, ClothCollectionGroup::Lods);
 		SkeletalMeshPathName = &ManagedArrayCollection->AddAttribute<FString>(ClothCollectionAttribute::SkeletalMeshPathName, ClothCollectionGroup::Lods);
 
+		// Solvers Group
+		if (EnumHasAnyFlags(OptionalSchemas, EClothCollectionOptionalSchemas::Solvers))
+		{
+			SolverGravity = &ManagedArrayCollection->AddAttribute<FVector3f>(ClothCollectionAttribute::SolverGravity, ClothCollectionGroup::Solvers);
+			SolverAirDamping = &ManagedArrayCollection->AddAttribute<float>(ClothCollectionAttribute::SolverAirDamping, ClothCollectionGroup::Solvers);
+			SolverSubSteps = &ManagedArrayCollection->AddAttribute<int32>(ClothCollectionAttribute::SolverSubSteps, ClothCollectionGroup::Solvers);
+			SolverTimeStep = &ManagedArrayCollection->AddAttribute<float>(ClothCollectionAttribute::SolverTimeStep, ClothCollectionGroup::Solvers);
+		}
+		
 		// Seams Group
 		SeamStitchStart = &ManagedArrayCollection->AddAttribute<int32>(ClothCollectionAttribute::SeamStitchStart, ClothCollectionGroup::Seams, SeamStitchesDependency);
 		SeamStitchEnd = &ManagedArrayCollection->AddAttribute<int32>(ClothCollectionAttribute::SeamStitchEnd, ClothCollectionGroup::Seams, SeamStitchesDependency);
@@ -340,10 +380,12 @@ namespace UE::Chaos::ClothAsset
 		FabricBucklingStiffness = &ManagedArrayCollection->AddAttribute<FVector3f>(ClothCollectionAttribute::FabricBucklingStiffness, ClothCollectionGroup::Fabrics);
 		FabricStretchStiffness = &ManagedArrayCollection->AddAttribute<FVector3f>(ClothCollectionAttribute::FabricStretchStiffness, ClothCollectionGroup::Fabrics);
 		FabricBucklingRatio = &ManagedArrayCollection->AddAttribute<float>(ClothCollectionAttribute::FabricBucklingRatio, ClothCollectionGroup::Fabrics);
-		FabricClothDensity = &ManagedArrayCollection->AddAttribute<float>(ClothCollectionAttribute::FabricClothDensity, ClothCollectionGroup::Fabrics);
-		FabricClothFriction = &ManagedArrayCollection->AddAttribute<float>(ClothCollectionAttribute::FabricClothFriction, ClothCollectionGroup::Fabrics);
-		FabricClothThickness = &ManagedArrayCollection->AddAttribute<float>(ClothCollectionAttribute::FabricClothThickness, ClothCollectionGroup::Fabrics);
-		FabricClothDamping = &ManagedArrayCollection->AddAttribute<float>(ClothCollectionAttribute::FabricClothDamping, ClothCollectionGroup::Fabrics);
+		FabricDensity = &ManagedArrayCollection->AddAttribute<float>(ClothCollectionAttribute::FabricDensity, ClothCollectionGroup::Fabrics);
+		FabricFriction = &ManagedArrayCollection->AddAttribute<float>(ClothCollectionAttribute::FabricFriction, ClothCollectionGroup::Fabrics);
+		FabricDamping = &ManagedArrayCollection->AddAttribute<float>(ClothCollectionAttribute::FabricDamping, ClothCollectionGroup::Fabrics);
+		FabricPressure = &ManagedArrayCollection->AddAttribute<float>(ClothCollectionAttribute::FabricPressure, ClothCollectionGroup::Fabrics);
+		FabricLayer = &ManagedArrayCollection->AddAttribute<int32>(ClothCollectionAttribute::FabricLayer, ClothCollectionGroup::Fabrics);
+		FabricCollisionThickness = &ManagedArrayCollection->AddAttribute<float>(ClothCollectionAttribute::FabricCollisionThickness, ClothCollectionGroup::Fabrics);
 		
 		// Sim Faces Group
 		SimIndices2D = &ManagedArrayCollection->AddAttribute<FIntVector3>(ClothCollectionAttribute::SimIndices2D, ClothCollectionGroup::SimFaces, SimVertices2DDependency);
diff --git a/Engine/Plugins/Experimental/ChaosClothAsset/Source/ChaosClothAsset/Private/ChaosClothAsset/ClothCollection.h b/Engine/Plugins/Experimental/ChaosClothAsset/Source/ChaosClothAsset/Private/ChaosClothAsset/ClothCollection.h
--- a/Engine/Plugins/Experimental/ChaosClothAsset/Source/ChaosClothAsset/Private/ChaosClothAsset/ClothCollection.h
+++ b/Engine/Plugins/Experimental/ChaosClothAsset/Source/ChaosClothAsset/Private/ChaosClothAsset/ClothCollection.h
@@ -158,6 +158,12 @@ namespace UE::Chaos::ClothAsset
 		const TManagedArray<FString>* GetPhysicsAssetPathName() const { return PhysicsAssetPathName; }
 		const TManagedArray<FString>* GetSkeletalMeshPathName() const { return SkeletalMeshPathName; }
 
+		// ~ Solvers Group
+		const TManagedArray<FVector3f>* GetSolverGravity() const { return SolverGravity; }
+		const TManagedArray<float>* GetSolverAirDamping() const { return SolverAirDamping; }
+		const TManagedArray<int32>* GetSolverSubSteps() const { return SolverSubSteps; }
+		const TManagedArray<float>* GetSolverTimeStep() const { return SolverTimeStep; }
+
 		//~ Seam Group
 		const TManagedArray<int32>* GetSeamStitchStart() const { return SeamStitchStart; }
 		const TManagedArray<int32>* GetSeamStitchEnd() const { return SeamStitchEnd; }
@@ -171,10 +177,12 @@ namespace UE::Chaos::ClothAsset
 		const TManagedArray<FVector3f>* GetFabricBucklingStiffness() const { return FabricBucklingStiffness; }
 		const TManagedArray<FVector3f>* GetFabricStretchStiffness() const { return FabricStretchStiffness; }
 		const TManagedArray<float>* GetFabricBucklingRatio() const { return FabricBucklingRatio; }
-		const TManagedArray<float>* GetFabricDensity() const { return FabricClothDensity; }
-		const TManagedArray<float>* GetFabricFriction() const { return FabricClothFriction; }
-		const TManagedArray<float>* GetFabricThickness() const { return FabricClothThickness; }
-		const TManagedArray<float>* GetFabricDamping() const { return FabricClothDamping; }
+		const TManagedArray<float>* GetFabricDensity() const { return FabricDensity; }
+		const TManagedArray<float>* GetFabricFriction() const { return FabricFriction; }
+		const TManagedArray<float>* GetFabricDamping() const { return FabricDamping; }
+		const TManagedArray<float>* GetFabricPressure() const { return FabricPressure; }
+		const TManagedArray<int32>* GetFabricLayer() const { return FabricLayer; }
+		const TManagedArray<float>* GetFabricCollisionThickness() const { return FabricCollisionThickness; }
 		
 		//~ Sim Patterns Group
 		const TManagedArray<int32>* GetSimVertices2DStart() const { return SimVertices2DStart; }
@@ -182,7 +190,7 @@ namespace UE::Chaos::ClothAsset
 		const TManagedArray<int32>* GetSimFacesStart() const { return SimFacesStart; }
 		const TManagedArray<int32>* GetSimFacesEnd() const { return SimFacesEnd; }
 		const TManagedArray<int32>* GetSimPatternFabric() const { return SimPatternFabric; }
-
+		
 		//~ Render Patterns Group
 		const TManagedArray<int32>* GetRenderVerticesStart() const { return RenderVerticesStart; }
 		const TManagedArray<int32>* GetRenderVerticesEnd() const { return RenderVerticesEnd; }
@@ -231,16 +239,22 @@ namespace UE::Chaos::ClothAsset
 		//~ LODs Group (There should be only one LOD per ClothCollection)
 		TManagedArray<FString>* GetPhysicsAssetPathName(){ return PhysicsAssetPathName; }
 		TManagedArray<FString>* GetSkeletalMeshPathName() { return SkeletalMeshPathName; }
-
+		TManagedArray<FVector3f>* GetSolverGravity() { return SolverGravity; }
+		TManagedArray<float>* GetSolverAirDamping() { return SolverAirDamping; }
+		TManagedArray<int32>* GetSolverSubSteps() { return SolverSubSteps; }
+		TManagedArray<float>* GetSolverTimeStep() { return SolverTimeStep; }
+		
 		//~ Fabric Group
 		TManagedArray<FVector3f>* GetFabricBendingStiffness() { return FabricBendingStiffness; }
 		TManagedArray<FVector3f>* GetFabricBucklingStiffness() { return FabricBucklingStiffness; }
 		TManagedArray<FVector3f>* GetFabricStretchStiffness() { return FabricStretchStiffness; }
 		TManagedArray<float>* GetFabricBucklingRatio() { return FabricBucklingRatio; }
-		TManagedArray<float>* GetFabricDensity() { return FabricClothDensity; }
-		TManagedArray<float>* GetFabricFriction() { return FabricClothFriction; }
-		TManagedArray<float>* GetFabricThickness() { return FabricClothThickness; }
-		TManagedArray<float>* GetFabricDamping() { return FabricClothDamping; }
+		TManagedArray<float>* GetFabricDensity() { return FabricDensity; }
+		TManagedArray<float>* GetFabricFriction() { return FabricFriction; }
+		TManagedArray<float>* GetFabricDamping() { return FabricDamping; }
+		TManagedArray<float>* GetFabricPressure() { return FabricPressure; }
+		TManagedArray<int32>* GetFabricLayer() { return FabricLayer; }
+		TManagedArray<float>* GetFabricCollisionThickness() { return FabricCollisionThickness; }
 		
 		//~ Seam Group
 		TManagedArray<int32>* GetSeamStitchStart() { return SeamStitchStart; }
@@ -310,15 +324,23 @@ namespace UE::Chaos::ClothAsset
 		TManagedArray<FString>* PhysicsAssetPathName;
 		TManagedArray<FString>* SkeletalMeshPathName;
 
-		//! Fabrics Group
+		//~ Solvers group
+		TManagedArray<FVector3f>* SolverGravity;
+		TManagedArray<float>* SolverAirDamping;
+		TManagedArray<int32>* SolverSubSteps;
+		TManagedArray<float>* SolverTimeStep;
+
+		//~ Fabrics Group
 		TManagedArray<FVector3f>* FabricBendingStiffness;
 		TManagedArray<FVector3f>* FabricBucklingStiffness;
 		TManagedArray<FVector3f>* FabricStretchStiffness;
 		TManagedArray<float>* FabricBucklingRatio;
-		TManagedArray<float>* FabricClothDensity;
-		TManagedArray<float>* FabricClothFriction;
-		TManagedArray<float>* FabricClothThickness;
-		TManagedArray<float>* FabricClothDamping;
+		TManagedArray<float>* FabricDensity;
+		TManagedArray<float>* FabricFriction;
+		TManagedArray<float>* FabricDamping;
+		TManagedArray<float>* FabricPressure;
+		TManagedArray<int32>* FabricLayer;
+		TManagedArray<float>* FabricCollisionThickness;
 
 		//~ Seam Group
 		TManagedArray<int32>* SeamStitchStart;
diff --git a/Engine/Plugins/Experimental/ChaosClothAsset/Source/ChaosClothAsset/Private/ChaosClothAsset/CollectionClothFabricFacade.cpp b/Engine/Plugins/Experimental/ChaosClothAsset/Source/ChaosClothAsset/Private/ChaosClothAsset/CollectionClothFabricFacade.cpp
--- a/Engine/Plugins/Experimental/ChaosClothAsset/Source/ChaosClothAsset/Private/ChaosClothAsset/CollectionClothFabricFacade.cpp
+++ b/Engine/Plugins/Experimental/ChaosClothAsset/Source/ChaosClothAsset/Private/ChaosClothAsset/CollectionClothFabricFacade.cpp
@@ -24,47 +24,57 @@ namespace UE::Chaos::ClothAsset
 		return FVector3f(Weft, Warp, Bias);
 	}
 	
-	FCollectionClothFabricConstFacade::FAnisotropicData FCollectionClothFabricConstFacade::GetXPBDAnisoBendingStiffness() const
+	FCollectionClothFabricConstFacade::FAnisotropicData FCollectionClothFabricConstFacade::GetBendingStiffness() const
 	{
 		return FCollectionClothFabricConstFacade::FAnisotropicData(
 			ClothCollection->GetElements(ClothCollection->GetFabricBendingStiffness())[GetElementIndex()]);
 	}
 	
-	float FCollectionClothFabricConstFacade::GetXPBDAnisoBucklingRatio() const
+	float FCollectionClothFabricConstFacade::GetBucklingRatio() const
 	{
 		return ClothCollection->GetElements(ClothCollection->GetFabricBucklingRatio())[GetElementIndex()];
 	}
 	
-	FCollectionClothFabricConstFacade::FAnisotropicData FCollectionClothFabricConstFacade::GetXPBDAnisoBucklingStiffness() const
+	FCollectionClothFabricConstFacade::FAnisotropicData FCollectionClothFabricConstFacade::GetBucklingStiffness() const
 	{
 		return FCollectionClothFabricConstFacade::FAnisotropicData(
 			ClothCollection->GetElements(ClothCollection->GetFabricBucklingStiffness())[GetElementIndex()]);
 	}
 	
-	FCollectionClothFabricConstFacade::FAnisotropicData FCollectionClothFabricConstFacade::GetXPBDAnisoSpringStiffness() const
+	FCollectionClothFabricConstFacade::FAnisotropicData FCollectionClothFabricConstFacade::GetStretchStiffness() const
 	{
 		return FCollectionClothFabricConstFacade::FAnisotropicData(
 			ClothCollection->GetElements(ClothCollection->GetFabricStretchStiffness())[GetElementIndex()]);
 	}
 	
-	float FCollectionClothFabricConstFacade::GetDensityWeighted() const
+	float FCollectionClothFabricConstFacade::GetDensity() const
 	{
 		return ClothCollection->GetElements(ClothCollection->GetFabricDensity())[GetElementIndex()];
 	}
 	
-	float FCollectionClothFabricConstFacade::GetCollisionThickness() const
+	float FCollectionClothFabricConstFacade::GetDamping() const
 	{
-		return ClothCollection->GetElements(ClothCollection->GetFabricThickness())[GetElementIndex()];
+		return ClothCollection->GetElements(ClothCollection->GetFabricDamping())[GetElementIndex()];
 	}
 	
-	float FCollectionClothFabricConstFacade::GetXPBDAnisoDamping() const
+	float FCollectionClothFabricConstFacade::GetFriction() const
 	{
-		return ClothCollection->GetElements(ClothCollection->GetFabricDamping())[GetElementIndex()];
+		return ClothCollection->GetElements(ClothCollection->GetFabricFriction())[GetElementIndex()];
 	}
 	
-	float FCollectionClothFabricConstFacade::GetFrictionCoefficient() const
+	float FCollectionClothFabricConstFacade::GetPressure() const
 	{
-		return ClothCollection->GetElements(ClothCollection->GetFabricFriction())[GetElementIndex()];
+		return ClothCollection->GetElements(ClothCollection->GetFabricPressure())[GetElementIndex()];
+	}
+	
+	int32 FCollectionClothFabricConstFacade::GetLayer() const
+	{
+		return ClothCollection->GetElements(ClothCollection->GetFabricLayer())[GetElementIndex()];
+	}
+
+	float FCollectionClothFabricConstFacade::GetCollisionThickness() const
+	{
+		return ClothCollection->GetElements(ClothCollection->GetFabricCollisionThickness())[GetElementIndex()];
 	}
 	
 	FCollectionClothFabricConstFacade::FCollectionClothFabricConstFacade(const TSharedRef<const FClothCollection>& ClothCollection, int32 InFabricIndex)
@@ -77,7 +87,7 @@ namespace UE::Chaos::ClothAsset
 
 	void FCollectionClothFabricFacade::Initialize(const FAnisotropicData& BendingStiffness, const float BucklingRatio,
 			const FAnisotropicData& BucklingStiffness, const FAnisotropicData& StretchStiffness,
-			const float ClothDensity, const float ClothFriction, const float ClothDamping, const float ClothThickness)
+			const float Density, const float Friction, const float Damping, const float Pressure, const int32 Layer, const float CollisionThickness)
 	{
 		const int32 ElementIndex = GetElementIndex();
 		check((ElementIndex >= 0) && (ElementIndex < ClothCollection->GetNumElements(ClothCollectionGroup::Fabrics)));
@@ -88,17 +98,27 @@ namespace UE::Chaos::ClothAsset
 		ClothCol->GetElements(ClothCol->GetFabricBucklingRatio())[ElementIndex] = BucklingRatio;
 		ClothCol->GetElements(ClothCol->GetFabricStretchStiffness())[ElementIndex] = StretchStiffness.GetVectorDatas();
 		ClothCol->GetElements(ClothCol->GetFabricBucklingStiffness())[ElementIndex] = BucklingStiffness.GetVectorDatas();
-		ClothCol->GetElements(ClothCol->GetFabricFriction())[ElementIndex] = ClothFriction;
-		ClothCol->GetElements(ClothCol->GetFabricDensity())[ElementIndex] = ClothDensity;
-		ClothCol->GetElements(ClothCol->GetFabricThickness())[ElementIndex] = ClothThickness;
-		ClothCol->GetElements(ClothCol->GetFabricDamping())[ElementIndex] = ClothDamping;
+		ClothCol->GetElements(ClothCol->GetFabricFriction())[ElementIndex] = Friction;
+		ClothCol->GetElements(ClothCol->GetFabricDensity())[ElementIndex] = Density;
+		ClothCol->GetElements(ClothCol->GetFabricDamping())[ElementIndex] = Damping;
+		ClothCol->GetElements(ClothCol->GetFabricPressure())[ElementIndex] = Pressure;
+		ClothCol->GetElements(ClothCol->GetFabricLayer())[ElementIndex] = Layer;
+		ClothCol->GetElements(ClothCol->GetFabricCollisionThickness())[ElementIndex] = CollisionThickness;
+	}
+
+	void FCollectionClothFabricFacade::Initialize(const FCollectionClothFabricConstFacade& OtherFabricFacade,
+		const float Pressure, const int32 Layer, const float CollisionThickness)
+	{
+		Initialize( OtherFabricFacade.GetBendingStiffness(), OtherFabricFacade.GetBucklingRatio(), OtherFabricFacade.GetBucklingStiffness(),
+			OtherFabricFacade.GetStretchStiffness(), OtherFabricFacade.GetDensity(),
+			OtherFabricFacade.GetFriction(), OtherFabricFacade.GetDamping(), Pressure, Layer, CollisionThickness);
 	}
 	
 	void FCollectionClothFabricFacade::Initialize(const FCollectionClothFabricConstFacade& OtherFabricFacade)
 	{
-		Initialize( OtherFabricFacade.GetXPBDAnisoBendingStiffness(), OtherFabricFacade.GetXPBDAnisoBucklingRatio(), OtherFabricFacade.GetXPBDAnisoBucklingStiffness(),
-			OtherFabricFacade.GetXPBDAnisoSpringStiffness(), OtherFabricFacade.GetDensityWeighted(),
-			OtherFabricFacade.GetFrictionCoefficient(), OtherFabricFacade.GetXPBDAnisoDamping(), OtherFabricFacade.GetCollisionThickness());
+		Initialize( OtherFabricFacade.GetBendingStiffness(), OtherFabricFacade.GetBucklingRatio(), OtherFabricFacade.GetBucklingStiffness(),
+			OtherFabricFacade.GetStretchStiffness(), OtherFabricFacade.GetDensity(),
+			OtherFabricFacade.GetFriction(), OtherFabricFacade.GetDamping(), OtherFabricFacade.GetPressure(), OtherFabricFacade.GetLayer(), OtherFabricFacade.GetCollisionThickness());
 	}
 
 	FCollectionClothFabricFacade::FCollectionClothFabricFacade(const TSharedRef<FClothCollection>& ClothCollection, int32 InFabricIndex)
@@ -113,6 +133,8 @@ namespace UE::Chaos::ClothAsset
 
 	void FCollectionClothFabricFacade::SetDefaults()
 	{
-		Initialize(100.0f, 0.5f, 50.0f, 100.0f, 0.35f, 0.1f, 0.01f, 0.8f);
+		Initialize(FDefaultFabric::BendingStiffness, FDefaultFabric::BucklingRatio, FDefaultFabric::BucklingStiffness,
+			FDefaultFabric::StretchStiffness, FDefaultFabric::Density, FDefaultFabric::Friction, FDefaultFabric::Damping,
+			FDefaultFabric::Pressure, FDefaultFabric::Layer, FDefaultFabric::CollisionThickness);
 	}
 }  // End namespace UE::Chaos::ClothAsset
diff --git a/Engine/Plugins/Experimental/ChaosClothAsset/Source/ChaosClothAsset/Private/ChaosClothAsset/CollectionClothFacade.cpp b/Engine/Plugins/Experimental/ChaosClothAsset/Source/ChaosClothAsset/Private/ChaosClothAsset/CollectionClothFacade.cpp
--- a/Engine/Plugins/Experimental/ChaosClothAsset/Source/ChaosClothAsset/Private/ChaosClothAsset/CollectionClothFacade.cpp
+++ b/Engine/Plugins/Experimental/ChaosClothAsset/Source/ChaosClothAsset/Private/ChaosClothAsset/CollectionClothFacade.cpp
@@ -45,6 +45,12 @@ namespace UE::Chaos::ClothAsset
 		ResultHash = HashCombineFast(ResultHash, GetTypeHash(GetSkeletalMeshPathName()));
 		ResultHash = HashCombineFast(ResultHash, GetTypeHash(GetPhysicsAssetPathName()));
 
+		//~ Solvers Group
+		ResultHash = HashCombineFast(ResultHash, GetTypeHash(GetSolverGravity()));
+		ResultHash = HashCombineFast(ResultHash, GetTypeHash(GetSolverAirDamping()));
+		ResultHash = HashCombineFast(ResultHash, GetTypeHash(GetSolverSubSteps()));
+		ResultHash = HashCombineFast(ResultHash, GetTypeHash(GetSolverTimeStep()));
+
 		//~ Seam Stitches Group
 		ResultHash = HashCombineFast(ResultHash, ClothCollection->GetElementsTypeHash(ClothCollection->GetSeamStitchStart()));
 		ResultHash = HashCombineFast(ResultHash, ClothCollection->GetElementsTypeHash(ClothCollection->GetSeamStitchEnd()));
@@ -93,8 +99,10 @@ namespace UE::Chaos::ClothAsset
 		ResultHash = HashCombineFast(ResultHash, ClothCollection->GetElementsTypeHash(ClothCollection->GetFabricDamping()));
 		ResultHash = HashCombineFast(ResultHash, ClothCollection->GetElementsTypeHash(ClothCollection->GetFabricDensity()));
 		ResultHash = HashCombineFast(ResultHash, ClothCollection->GetElementsTypeHash(ClothCollection->GetFabricFriction()));
-		ResultHash = HashCombineFast(ResultHash, ClothCollection->GetElementsTypeHash(ClothCollection->GetFabricThickness()));
 		ResultHash = HashCombineFast(ResultHash, ClothCollection->GetElementsTypeHash(ClothCollection->GetFabricStretchStiffness()));
+		ResultHash = HashCombineFast(ResultHash, ClothCollection->GetElementsTypeHash(ClothCollection->GetFabricPressure()));
+		ResultHash = HashCombineFast(ResultHash, ClothCollection->GetElementsTypeHash(ClothCollection->GetFabricLayer()));
+		ResultHash = HashCombineFast(ResultHash, ClothCollection->GetElementsTypeHash(ClothCollection->GetFabricCollisionThickness()));
 
 		//~ Render Faces Group
 		ResultHash = HashCombineFast(ResultHash, ClothCollection->GetElementsTypeHash(ClothCollection->GetRenderIndices()));
@@ -162,6 +170,31 @@ namespace UE::Chaos::ClothAsset
 		return ClothCollection->GetSkeletalMeshPathName() && ClothCollection->GetNumElements(ClothCollectionGroup::Lods) > 0 ? (*ClothCollection->GetSkeletalMeshPathName())[0] : EmptyString;
 	}
 
+	bool FCollectionClothConstFacade::HasSolverElement() const
+	{
+		return ClothCollection->GetNumElements(ClothCollectionGroup::Solvers) == 1;
+	}
+
+	const FVector3f& FCollectionClothConstFacade::GetSolverGravity() const
+	{
+		return (ClothCollection->GetSolverGravity() && HasSolverElement()) ? (*ClothCollection->GetSolverGravity())[0] : FDefaultSolver::Gravity;
+	}
+
+	float FCollectionClothConstFacade::GetSolverAirDamping() const
+	{
+		return (ClothCollection->GetSolverAirDamping() && HasSolverElement()) ? (*ClothCollection->GetSolverAirDamping())[0] : FDefaultSolver::AirDamping;
+	}
+
+	int32 FCollectionClothConstFacade::GetSolverSubSteps() const
+	{
+		return (ClothCollection->GetSolverSubSteps() && HasSolverElement()) ? (*ClothCollection->GetSolverSubSteps())[0] : FDefaultSolver::SubSteps;
+	}
+
+	float FCollectionClothConstFacade::GetSolverTimeStep() const
+	{
+		return (ClothCollection->GetSolverTimeStep() && HasSolverElement()) ? (*ClothCollection->GetSolverTimeStep())[0] : FDefaultSolver::TimeStep;
+	}
+
 	int32 FCollectionClothConstFacade::GetNumSimVertices2D() const
 	{
 		return ClothCollection->GetNumElements(ClothCollectionGroup::SimVertices2D);
@@ -500,6 +533,12 @@ namespace UE::Chaos::ClothAsset
 		SetNumSimPatterns(0);
 		SetNumRenderPatterns(0);
 		SetNumSeams(0); // Do this after removing SimVertices3D and SimPatterns. Otherwise, Seams will do a bunch of unnecessary work to unseam stuff.
+		SetNumFabrics(0); 
+
+		if (IsValid(EClothCollectionOptionalSchemas::Solvers))
+		{
+			GetClothCollection()->SetNumElements(0, ClothCollectionGroup::Solvers);
+		}
 	}
 
 	void FCollectionClothFacade::Initialize(const FCollectionClothConstFacade& Other)
@@ -510,6 +549,15 @@ namespace UE::Chaos::ClothAsset
 		SetPhysicsAssetPathName(Other.GetPhysicsAssetPathName());
 		SetSkeletalMeshPathName(Other.GetSkeletalMeshPathName());
 		
+		// Solvers Group
+		if (Other.IsValid(EClothCollectionOptionalSchemas::Solvers) && Other.HasSolverElement())
+		{
+			SetSolverGravity(Other.GetSolverGravity());
+			SetSolverAirDamping(Other.GetSolverAirDamping());
+			SetSolverSubSteps(Other.GetSolverSubSteps());
+			SetSolverTimeStep(Other.GetSolverTimeStep());
+		}
+		
 		Append(Other);
 	}
 
@@ -642,8 +690,8 @@ namespace UE::Chaos::ClothAsset
 		{
 			(*GetClothCollection()->GetPhysicsAssetPathName())[0] = PathName;
 		}
-
 	}
+	
 	void FCollectionClothFacade::SetSkeletalMeshPathName(const FString& PathName)
 	{
 		if (ClothCollection->GetNumElements(ClothCollectionGroup::Lods))
@@ -651,6 +699,42 @@ namespace UE::Chaos::ClothAsset
 			(*GetClothCollection()->GetSkeletalMeshPathName())[0] = PathName;
 		}
 	}
+	
+	void FCollectionClothFacade::SetSolverGravity(const FVector3f& SolverGravity)
+	{
+		if (!IsValid(EClothCollectionOptionalSchemas::Solvers))
+		{
+			DefineSchema(EClothCollectionOptionalSchemas::Solvers);
+		}
+		(*GetClothCollection()->GetSolverGravity())[0] = SolverGravity;
+	}
+	
+	void FCollectionClothFacade::SetSolverAirDamping(const float SolverAirDamping)
+	{
+		if (!IsValid(EClothCollectionOptionalSchemas::Solvers))
+		{
+			DefineSchema(EClothCollectionOptionalSchemas::Solvers);
+		}
+		(*GetClothCollection()->GetSolverAirDamping())[0] = SolverAirDamping;
+	}
+	
+	void FCollectionClothFacade::SetSolverTimeStep(const float SolverTimeStep)
+	{
+		if (!IsValid(EClothCollectionOptionalSchemas::Solvers))
+		{
+			DefineSchema(EClothCollectionOptionalSchemas::Solvers);
+		}
+		(*GetClothCollection()->GetSolverTimeStep())[0] = SolverTimeStep;
+	}
+	
+	void FCollectionClothFacade::SetSolverSubSteps(const int32 SolverSubSteps)
+	{
+		if (!IsValid(EClothCollectionOptionalSchemas::Solvers))
+		{
+			DefineSchema(EClothCollectionOptionalSchemas::Solvers);
+		}
+		(*GetClothCollection()->GetSolverSubSteps())[0] = SolverSubSteps;
+	}
 
 	TArrayView<FVector2f> FCollectionClothFacade::GetSimPosition2D()
 	{
@@ -1048,6 +1132,16 @@ namespace UE::Chaos::ClothAsset
 	void FCollectionClothFacade::SetDefaults()
 	{
 		GetClothCollection()->SetNumElements(1, ClothCollectionGroup::Lods);
+
+		if (IsValid(EClothCollectionOptionalSchemas::Solvers))
+		{
+			GetClothCollection()->SetNumElements(1, ClothCollectionGroup::Solvers);
+			
+			SetSolverGravity(FDefaultSolver::Gravity);
+			SetSolverAirDamping(FDefaultSolver::AirDamping);
+			SetSolverSubSteps(FDefaultSolver::SubSteps);
+			SetSolverTimeStep(FDefaultSolver::TimeStep);
+		}
 	}
 
 } // End namespace UE::Chaos::ClothAsset
diff --git a/Engine/Plugins/Experimental/ChaosClothAsset/Source/ChaosClothAsset/Public/ChaosClothAsset/ClothCollectionAttribute.h b/Engine/Plugins/Experimental/ChaosClothAsset/Source/ChaosClothAsset/Public/ChaosClothAsset/ClothCollectionAttribute.h
--- a/Engine/Plugins/Experimental/ChaosClothAsset/Source/ChaosClothAsset/Public/ChaosClothAsset/ClothCollectionAttribute.h
+++ b/Engine/Plugins/Experimental/ChaosClothAsset/Source/ChaosClothAsset/Public/ChaosClothAsset/ClothCollectionAttribute.h
@@ -15,15 +15,23 @@ namespace UE::Chaos::ClothAsset
 		inline const FName PhysicsAssetPathName(TEXT("PhysicsAssetPathName"));
 		inline const FName SkeletalMeshPathName(TEXT("SkeletalMeshPathName"));
 
+		// Solvers Group
+		inline const FName SolverGravity(TEXT("SolverGravity"));
+		inline const FName SolverAirDamping(TEXT("SolverAirDamping"));
+		inline const FName SolverSubSteps(TEXT("SolverSubSteps"));
+		inline const FName SolverTimeStep(TEXT("SolverTimeStep"));
+
 		// Fabrics Group
 		inline const FName FabricBendingStiffness(TEXT("FabricBendingStiffness"));
 		inline const FName FabricBucklingStiffness(TEXT("FabricBucklingStiffness"));
 		inline const FName FabricStretchStiffness(TEXT("FabricStretchStiffness"));
 		inline const FName FabricBucklingRatio(TEXT("FabricBucklingRatio"));
-		inline const FName FabricClothDensity(TEXT("FabricClothDensity"));
-		inline const FName FabricClothFriction(TEXT("FabricClothFriction"));
-		inline const FName FabricClothThickness(TEXT("FabricClothThickness"));
-		inline const FName FabricClothDamping(TEXT("FabricClothDamping"));
+		inline const FName FabricDensity(TEXT("FabricClothDensity"));
+		inline const FName FabricFriction(TEXT("FabricClothFriction"));
+		inline const FName FabricDamping(TEXT("FabricClothDamping"));
+		inline const FName FabricPressure(TEXT("FabricPressure"));
+		inline const FName FabricLayer(TEXT("FabricLayer"));
+		inline const FName FabricCollisionThickness(TEXT("FabricollisionThickness"));
 
 		// Seam Group
 		inline const FName SeamStitchStart(TEXT("SeamStitchStart"));
diff --git a/Engine/Plugins/Experimental/ChaosClothAsset/Source/ChaosClothAsset/Public/ChaosClothAsset/ClothCollectionGroup.h b/Engine/Plugins/Experimental/ChaosClothAsset/Source/ChaosClothAsset/Public/ChaosClothAsset/ClothCollectionGroup.h
--- a/Engine/Plugins/Experimental/ChaosClothAsset/Source/ChaosClothAsset/Public/ChaosClothAsset/ClothCollectionGroup.h
+++ b/Engine/Plugins/Experimental/ChaosClothAsset/Source/ChaosClothAsset/Public/ChaosClothAsset/ClothCollectionGroup.h
@@ -13,6 +13,8 @@ namespace UE::Chaos::ClothAsset
 	{
 		/** LOD information (only one LOD per collection). */
 		inline const FName Lods = FName(TEXT("Lods"));
+		/** Solvers information (only one Solver per collection for now). */
+		inline const FName Solvers = FName(TEXT("Solvers"));
 		/** Collection of seam stitches. */
 		inline const FName Seams = FName(TEXT("Seams"));
 		/** Contains pairs of stitched sim vertex indices. */
diff --git a/Engine/Plugins/Experimental/ChaosClothAsset/Source/ChaosClothAsset/Public/ChaosClothAsset/ClothCollectionOptionalSchemas.h b/Engine/Plugins/Experimental/ChaosClothAsset/Source/ChaosClothAsset/Public/ChaosClothAsset/ClothCollectionOptionalSchemas.h
--- a/Engine/Plugins/Experimental/ChaosClothAsset/Source/ChaosClothAsset/Public/ChaosClothAsset/ClothCollectionOptionalSchemas.h
+++ b/Engine/Plugins/Experimental/ChaosClothAsset/Source/ChaosClothAsset/Public/ChaosClothAsset/ClothCollectionOptionalSchemas.h
@@ -11,7 +11,8 @@ namespace UE::Chaos::ClothAsset
 	enum class EClothCollectionOptionalSchemas : uint8
 	{
 		None = 0,
-		RenderDeformer = 1 << 1
+		RenderDeformer = 1 << 1,
+		Solvers = 1 << 2
 	};
 	ENUM_CLASS_FLAGS(EClothCollectionOptionalSchemas)
 }  // End namespace UE::Chaos::ClothAsset
diff --git a/Engine/Plugins/Experimental/ChaosClothAsset/Source/ChaosClothAsset/Public/ChaosClothAsset/CollectionClothFabricFacade.h b/Engine/Plugins/Experimental/ChaosClothAsset/Source/ChaosClothAsset/Public/ChaosClothAsset/CollectionClothFabricFacade.h
--- a/Engine/Plugins/Experimental/ChaosClothAsset/Source/ChaosClothAsset/Public/ChaosClothAsset/CollectionClothFabricFacade.h
+++ b/Engine/Plugins/Experimental/ChaosClothAsset/Source/ChaosClothAsset/Public/ChaosClothAsset/CollectionClothFabricFacade.h
@@ -6,6 +6,22 @@
 
 namespace UE::Chaos::ClothAsset
 {
+	struct FDefaultFabric
+	{
+		inline static constexpr float BendingStiffness = 100.0f;
+		inline static constexpr float StretchStiffness = 100.0f;
+		inline static constexpr float BucklingRatio = 0.5f;
+		inline static constexpr float BucklingStiffness = 50.0f;
+		inline static constexpr float Density = 0.35f;
+		inline static constexpr float Friction = 0.8f;
+		inline static constexpr float Damping = 0.1f;
+		inline static constexpr float Pressure = 0.0f;
+		inline static constexpr int32 Layer = INDEX_NONE;
+		inline static constexpr float CollisionThickness = 1.0f;
+		inline static constexpr float SelfFriction = 0.0f;
+		inline static constexpr float SelfCollisionThickness = 0.5f;
+	};
+	
 	/**
 	 * Cloth Asset collection fabric facade class to access cloth fabric data.
 	 * Constructed from FCollectionClothConstFacade.
@@ -40,28 +56,37 @@ namespace UE::Chaos::ClothAsset
 		};
 
 		/** Return the anisotropic bending stiffness */
-		FAnisotropicData GetXPBDAnisoBendingStiffness() const;
+		FAnisotropicData GetBendingStiffness() const;
 
 		/** Return the buckling ratio */
-		float GetXPBDAnisoBucklingRatio() const;
+		float GetBucklingRatio() const;
 
 		/** Return the anisotropic buckling stiffness */
-		FAnisotropicData GetXPBDAnisoBucklingStiffness() const;
+		FAnisotropicData GetBucklingStiffness() const;
 
 		/** Return the anisotropic stretch stiffness */
-		FAnisotropicData GetXPBDAnisoSpringStiffness() const;
+		FAnisotropicData GetStretchStiffness() const;
 
-		/** Return the cloth density */
-		float GetDensityWeighted() const;
+		/** Return the fabric density */
+		float GetDensity() const;
+		
+		/** Return the fabric damping */
+		float GetDamping() const;
 
-		/** Return the cloth thickness */
-		float GetCollisionThickness() const;
+		/** Return the fabric friction */
+		float GetFriction() const;
 
-		/** Return the cloth damping */
-		float GetXPBDAnisoDamping() const;
+		/** Return the fabric pressure */
+		float GetPressure() const;
+
+		/** Return the fabric layer */
+		int32 GetLayer() const;
+
+		/** Return the collision thickness */
+		float GetCollisionThickness() const;
 
-		/** Return the cloth friction */
-		float GetFrictionCoefficient() const;
+		/** Get the global element index */
+		int32 GetElementIndex() const { return GetBaseElementIndex() + FabricIndex; }
 
 	protected:
 		friend class FCollectionClothFabricFacade;  // For other instances access
@@ -70,9 +95,6 @@ namespace UE::Chaos::ClothAsset
 
 		static constexpr int32 GetBaseElementIndex() { return 0; }
 
-		/** Get the global element index */
-		int32 GetElementIndex() const { return GetBaseElementIndex() + FabricIndex; }
-
 		/** Cloth collection modified by  the fabric facade */
 		TSharedRef<const class FClothCollection> ClothCollection;
 
@@ -101,12 +123,15 @@ namespace UE::Chaos::ClothAsset
 		/** Initialize the cloth fabric with simulation parameters. */
 		void Initialize(const FAnisotropicData& BendingStiffness, const float BucklingRatio,
 			const FAnisotropicData& BucklingStiffness, const FAnisotropicData& StretchStiffness,
-			const float ClothDensity, const float ClothFriction,
-			const float ClothDamping, const float ClothThickness);
+			const float Density, const float Friction, const float Damping, const float Pressure, const int32 Layer, const float CollisionThickness);
 
 		/** Initialize the cloth fabric with another one. */
 		void Initialize(const FCollectionClothFabricConstFacade& OtherFabricFacade);
 
+		/** Initialize the cloth fabric from another one and from pattern datas */
+		void Initialize(const FCollectionClothFabricConstFacade& OtherFabricFacade,
+					const float Pressure, const int32 Layer, const float CollisionThickness);
+
 	private:
 		friend class FCollectionClothFacade;
 		FCollectionClothFabricFacade(const TSharedRef<class FClothCollection>& InClothCollection, int32 InFabricIndex);
diff --git a/Engine/Plugins/Experimental/ChaosClothAsset/Source/ChaosClothAsset/Public/ChaosClothAsset/CollectionClothFacade.h b/Engine/Plugins/Experimental/ChaosClothAsset/Source/ChaosClothAsset/Public/ChaosClothAsset/CollectionClothFacade.h
--- a/Engine/Plugins/Experimental/ChaosClothAsset/Source/ChaosClothAsset/Public/ChaosClothAsset/CollectionClothFacade.h
+++ b/Engine/Plugins/Experimental/ChaosClothAsset/Source/ChaosClothAsset/Public/ChaosClothAsset/CollectionClothFacade.h
@@ -2,6 +2,7 @@
 
 #pragma once
 
+#include "Math/Vector.h"
 #include "ChaosClothAsset/ClothCollectionOptionalSchemas.h"
 #include "ChaosClothAsset/CollectionClothFabricFacade.h"
 #include "ChaosClothAsset/CollectionClothRenderPatternFacade.h"
@@ -16,6 +17,14 @@ class FChaosArchive;
 
 namespace UE::Chaos::ClothAsset
 {
+	struct FDefaultSolver
+	{
+		inline static const FVector3f Gravity = FVector3f(0.0f, 0.0f, -980.665f);
+		inline static constexpr float AirDamping = 0.035f;
+		inline static constexpr int32 SubSteps = 1;
+		inline static constexpr float TimeStep = 0.033f;
+	};
+
 	/**
 	 * Cloth Asset collection facade class focused on draping and pattern information.
 	 * Const access (read only) version.
@@ -53,6 +62,19 @@ namespace UE::Chaos::ClothAsset
 		const FString& GetPhysicsAssetPathName() const;
 		/** Return the skeleton asset path names used for this collection. */
 		const FString& GetSkeletalMeshPathName() const;
+		
+
+		//~ Solver (single per collection) Group
+		/** Return true if the solver group has one element*/
+		bool HasSolverElement() const;
+		/** Return the solver gravity vector used for this collection. */
+		const FVector3f& GetSolverGravity() const;
+		/** Return the solver air damping used for this collection. */
+		float GetSolverAirDamping() const;
+		/** Return the solver time step used for this collection. */
+		float GetSolverTimeStep() const;
+		/** Return the solver sub steps used for this collection. */
+		int32 GetSolverSubSteps() const;
 
 		//~ Sim Vertices 2D Group
 		/** Return the total number of 2D simulation vertices for this collection. */
@@ -155,6 +177,7 @@ namespace UE::Chaos::ClothAsset
 			TArray<uint32>& PatternToWeldedIndices, TArray<TArray<int32>>* OptionalWeldedToPatternIndices = nullptr) const;
 
 	protected:
+		friend class FCollectionClothFacade;
 		TSharedRef<const class FClothCollection> ClothCollection;
 
 		friend class FCollectionClothFacade;  // To enable access from a different instance
@@ -205,6 +228,16 @@ namespace UE::Chaos::ClothAsset
 		/** Set the skeletal mesh asset path name and the reference skeleton that will be used with this asset. */
 		void SetSkeletalMeshPathName(const FString& PathName);
 
+		//~ Solver (max 1 per collection) Group
+		/** Set the solver gravity */
+		void SetSolverGravity(const FVector3f& SolverGravity);
+		/** Set the solver air damping */
+		void SetSolverAirDamping(const float SolverAirDamping);
+		/** Set the solver time step */
+		void SetSolverTimeStep(const float SolverTimeStep);
+		/** Set the solver substeps */
+		void SetSolverSubSteps(const int32 SolverSubSteps);
+
 		//~ Pattern Sim Vertices 2D Group
 		/** SetNumSimVertices2D per pattern within pattern facade. */
 		TArrayView<FVector2f> GetSimPosition2D();
@@ -321,6 +354,7 @@ namespace UE::Chaos::ClothAsset
 		static bool IsValidClothCollectionGroupName(const FName& GroupName);
 
 	private:
+		
 		void SetDefaults();
 
 		TSharedRef<class FClothCollection> GetClothCollection() { return ConstCastSharedRef<class FClothCollection>(ClothCollection); }
diff --git a/Engine/Plugins/Experimental/ChaosClothAssetEditor/Source/ChaosClothAssetDataflowNodes/Private/ChaosClothAsset/ChaosClothAssetDataflowNodesModule.cpp b/Engine/Plugins/Experimental/ChaosClothAssetEditor/Source/ChaosClothAssetDataflowNodes/Private/ChaosClothAsset/ChaosClothAssetDataflowNodesModule.cpp
--- a/Engine/Plugins/Experimental/ChaosClothAssetEditor/Source/ChaosClothAssetDataflowNodes/Private/ChaosClothAsset/ChaosClothAssetDataflowNodesModule.cpp
+++ b/Engine/Plugins/Experimental/ChaosClothAssetEditor/Source/ChaosClothAssetDataflowNodes/Private/ChaosClothAsset/ChaosClothAssetDataflowNodesModule.cpp
@@ -59,7 +59,10 @@
 #include "ChaosClothAsset/TransformUVsNode.h"
 #include "ChaosClothAsset/USDImportNode.h"
 #include "ChaosClothAsset/WeightedValueCustomization.h"
+#include "ChaosClothAsset/ImportedValueCustomization.h"
 #include "ChaosClothAsset/WeightMapToSelectionNode.h"
+#include "ChaosClothAsset/ConnectableValueCustomization.h"
+#include "ChaosClothAsset/ConnectableValue.h"
 #include "Dataflow/DataflowNodeColorsRegistry.h"
 #include "Dataflow/DataflowNodeFactory.h"
 #include "Modules/ModuleManager.h"
@@ -146,13 +149,16 @@ namespace UE::Chaos::ClothAsset
 				PropertyModule->RegisterCustomPropertyTypeLayout(FChaosClothAssetWeightedValueNonAnimatableNoLowHighRange::StaticStruct()->GetFName(), FOnGetPropertyTypeCustomizationInstance::CreateStatic(&FWeightedValueCustomization::MakeInstance));
 				PropertyModule->RegisterCustomPropertyTypeLayout(FChaosClothAssetConnectableStringValue::StaticStruct()->GetFName(), FOnGetPropertyTypeCustomizationInstance::CreateStatic(&FWeightedValueCustomization::MakeInstance));
 				PropertyModule->RegisterCustomPropertyTypeLayout(FChaosClothAssetWeightedValueOverride::StaticStruct()->GetFName(), FOnGetPropertyTypeCustomizationInstance::CreateStatic(&FWeightedValueCustomization::MakeInstance));
-				PropertyModule->RegisterCustomPropertyTypeLayout(FChaosClothAssetConnectableIStringValue::StaticStruct()->GetFName(), FOnGetPropertyTypeCustomizationInstance::CreateStatic(&FWeightedValueCustomization::MakeInstance));
+				PropertyModule->RegisterCustomPropertyTypeLayout(FChaosClothAssetConnectableIStringValue::StaticStruct()->GetFName(), FOnGetPropertyTypeCustomizationInstance::CreateStatic(&FConnectableValueCustomization::MakeInstance));
 				PropertyModule->RegisterCustomPropertyTypeLayout(FChaosClothAssetConnectableIOStringValue::StaticStruct()->GetFName(), FOnGetPropertyTypeCustomizationInstance::CreateStatic(&FWeightedValueCustomization::MakeInstance));
 				PropertyModule->RegisterCustomPropertyTypeLayout(FChaosClothAssetNodeSelectionGroup::StaticStruct()->GetFName(), FOnGetPropertyTypeCustomizationInstance::CreateStatic(&FSelectionGroupCustomization::MakeInstance));
 				PropertyModule->RegisterCustomPropertyTypeLayout(FChaosClothAssetImportFilePath::StaticStruct()->GetFName(), FOnGetPropertyTypeCustomizationInstance::CreateStatic(&FImportFilePathCustomization::MakeInstance));
 				PropertyModule->RegisterCustomPropertyTypeLayout(FChaosClothAssetTerminalNodeRefreshAsset::StaticStruct()->GetFName(), FOnGetPropertyTypeCustomizationInstance::CreateStatic(&FTerminalNodeRefreshAssetCustomization::MakeInstance));
 				PropertyModule->RegisterCustomPropertyTypeLayout(FChaosClothAssetImportNodeRefreshAsset::StaticStruct()->GetFName(), FOnGetPropertyTypeCustomizationInstance::CreateStatic(&FTerminalNodeRefreshAssetCustomization::MakeInstance));
 				PropertyModule->RegisterCustomPropertyTypeLayout(FChaosClothAssetNodeAttributeGroup::StaticStruct()->GetFName(), FOnGetPropertyTypeCustomizationInstance::CreateStatic(&FSelectionGroupCustomization::MakeInstance));
+				PropertyModule->RegisterCustomPropertyTypeLayout(FChaosClothAssetImportedVectorValue::StaticStruct()->GetFName(), FOnGetPropertyTypeCustomizationInstance::CreateStatic(&FImportedValueCustomization::MakeInstance));
+				PropertyModule->RegisterCustomPropertyTypeLayout(FChaosClothAssetImportedFloatValue::StaticStruct()->GetFName(), FOnGetPropertyTypeCustomizationInstance::CreateStatic(&FImportedValueCustomization::MakeInstance));
+				PropertyModule->RegisterCustomPropertyTypeLayout(FChaosClothAssetImportedIntValue::StaticStruct()->GetFName(), FOnGetPropertyTypeCustomizationInstance::CreateStatic(&FImportedValueCustomization::MakeInstance));
 			}
 
 			// Register modular features
@@ -169,6 +175,7 @@ namespace UE::Chaos::ClothAsset
 					PropertyModule->UnregisterCustomPropertyTypeLayout(FChaosClothAssetWeightedValue::StaticStruct()->GetFName());
 					PropertyModule->UnregisterCustomPropertyTypeLayout(FChaosClothAssetWeightedValueNonAnimatable::StaticStruct()->GetFName());
 					PropertyModule->UnregisterCustomPropertyTypeLayout(FChaosClothAssetWeightedValueNonAnimatableNoLowHighRange::StaticStruct()->GetFName());
+					PropertyModule->UnregisterCustomPropertyTypeLayout(FChaosClothAssetConnectableStringValue::StaticStruct()->GetFName());
 					PropertyModule->UnregisterCustomPropertyTypeLayout(FChaosClothAssetWeightedValueOverride::StaticStruct()->GetFName());
 					PropertyModule->UnregisterCustomPropertyTypeLayout(FChaosClothAssetConnectableIStringValue::StaticStruct()->GetFName());
 					PropertyModule->UnregisterCustomPropertyTypeLayout(FChaosClothAssetConnectableIOStringValue::StaticStruct()->GetFName());
@@ -177,6 +184,9 @@ namespace UE::Chaos::ClothAsset
 					PropertyModule->UnregisterCustomPropertyTypeLayout(FChaosClothAssetTerminalNodeRefreshAsset::StaticStruct()->GetFName());
 					PropertyModule->UnregisterCustomPropertyTypeLayout(FChaosClothAssetImportNodeRefreshAsset::StaticStruct()->GetFName());
 					PropertyModule->UnregisterCustomPropertyTypeLayout(FChaosClothAssetNodeAttributeGroup::StaticStruct()->GetFName());
+					PropertyModule->UnregisterCustomPropertyTypeLayout(FChaosClothAssetImportedVectorValue::StaticStruct()->GetFName());
+					PropertyModule->UnregisterCustomPropertyTypeLayout(FChaosClothAssetImportedFloatValue::StaticStruct()->GetFName());
+					PropertyModule->UnregisterCustomPropertyTypeLayout(FChaosClothAssetImportedIntValue::StaticStruct()->GetFName());
 				}
 			}
 
diff --git a/Engine/Plugins/Experimental/ChaosClothAssetEditor/Source/ChaosClothAssetDataflowNodes/Private/ChaosClothAsset/ConnectableValueCustomization.cpp b/Engine/Plugins/Experimental/ChaosClothAssetEditor/Source/ChaosClothAssetDataflowNodes/Private/ChaosClothAsset/ConnectableValueCustomization.cpp
new file mode 100644
--- /dev/null
+++ b/Engine/Plugins/Experimental/ChaosClothAssetEditor/Source/ChaosClothAssetDataflowNodes/Private/ChaosClothAsset/ConnectableValueCustomization.cpp
@@ -0,0 +1,197 @@
+// Copyright Epic Games, Inc. All Rights Reserved.
+
+#include "ChaosClothAsset/ConnectableValueCustomization.h"
+#include "ChaosClothAsset/ImportedValueCustomization.h"
+#include "ChaosClothAsset/ClothAssetEditorStyle.h"
+#include "ChaosClothAsset/WeightedValue.h"
+#include "Widgets/Input/SCheckBox.h"
+#include "Widgets/Input/SNumericEntryBox.h"
+#include "Widgets/Input/SEditableTextBox.h"
+#include "DetailLayoutBuilder.h"
+#include "DetailWidgetRow.h"
+#include "Editor.h"
+
+#define LOCTEXT_NAMESPACE "ChaosClothAssetWeightedValueCustomization"
+
+namespace UE::Chaos::ClothAsset
+{
+	namespace Private
+	{
+		static const FString OverridePrefix = TEXT("_Override");
+		static const FString BuildFabricMaps = TEXT("BuildFabricMaps");
+		static const FString CouldUseFabrics = TEXT("CouldUseFabrics");
+
+		
+	}
+	
+	bool FConnectableValueCustomization::IsOverrideProperty(const TSharedPtr<IPropertyHandle>& Property)
+	{
+		const FStringView PropertyPath = Property ? Property->GetPropertyPath() : FStringView();
+		return PropertyPath.EndsWith(Private::OverridePrefix, ESearchCase::CaseSensitive);
+	}
+	bool FConnectableValueCustomization::IsOverridePropertyOf(const TSharedPtr<IPropertyHandle>& OverrideProperty, const TSharedPtr<IPropertyHandle>& Property)
+	{
+		const FStringView OverridePropertyPath = OverrideProperty ? OverrideProperty->GetPropertyPath() : FStringView();
+		const FStringView PropertyPath = Property ? Property->GetPropertyPath() : FStringView();
+		return OverridePropertyPath == FString(PropertyPath) + Private::OverridePrefix;
+	}
+	bool FConnectableValueCustomization::BuildFabricMapsProperty(const TSharedPtr<IPropertyHandle>& Property)
+	{
+		const FStringView PropertyPath = Property ? Property->GetPropertyPath() : FStringView();
+		return PropertyPath.EndsWith(Private::BuildFabricMaps, ESearchCase::CaseSensitive);
+	}
+	bool FConnectableValueCustomization::CouldUseFabricsProperty(const TSharedPtr<IPropertyHandle>& Property)
+	{
+		const FStringView PropertyPath = Property ? Property->GetPropertyPath() : FStringView();
+		return PropertyPath.EndsWith(Private::CouldUseFabrics, ESearchCase::CaseSensitive);
+	}
+	
+	TSharedRef<IPropertyTypeCustomization> FConnectableValueCustomization::MakeInstance()
+	{
+		return MakeShareable(new FConnectableValueCustomization);
+	}
+
+	FConnectableValueCustomization::FConnectableValueCustomization() = default;
+	
+	FConnectableValueCustomization::~FConnectableValueCustomization() = default;
+
+
+	void FConnectableValueCustomization::MakeHeaderRow(TSharedRef<class IPropertyHandle>& StructPropertyHandle, FDetailWidgetRow& Row)
+	{
+		const TWeakPtr<IPropertyHandle> StructWeakHandlePtr = StructPropertyHandle;
+
+		TSharedPtr<SHorizontalBox> ValueHorizontalBox;
+		TSharedPtr<SHorizontalBox> NameHorizontalBox;
+
+		Row.NameContent()
+			[
+				SAssignNew(NameHorizontalBox, SHorizontalBox)
+				.IsEnabled(this, &FMathStructCustomization::IsValueEnabled, StructWeakHandlePtr)
+			]
+		.ValueContent()
+			// Make enough space for each child handle
+			.MinDesiredWidth(125.f * SortedChildHandles.Num())
+			.MaxDesiredWidth(125.f * SortedChildHandles.Num())
+			[
+				SAssignNew(ValueHorizontalBox, SHorizontalBox)
+				.IsEnabled(this, &FMathStructCustomization::IsValueEnabled, StructWeakHandlePtr)
+			];
+		
+		for (int32 ChildIndex = 0; ChildIndex < SortedChildHandles.Num(); ++ChildIndex)
+		{
+			TSharedRef<IPropertyHandle> ChildHandle = SortedChildHandles[ChildIndex];
+			
+			if (CouldUseFabricsProperty(ChildHandle))
+			{
+				bool bValue = false;
+				ChildHandle->GetValue(bValue);
+				if(!bValue)
+				{
+					break;
+				}
+			}
+			else if (BuildFabricMapsProperty(ChildHandle))
+			{
+				AddToggledCheckBox(ChildHandle, NameHorizontalBox, UE::Chaos::ClothAsset::FClothAssetEditorStyle::Get().GetBrush("ClassIcon.ChaosClothPreset"));
+			}
+		}
+		
+		NameHorizontalBox->AddSlot().VAlign(VAlign_Center)
+				.Padding(FMargin(4.f, 2.f, 4.0f, 2.f))
+				.HAlign(HAlign_Right)
+				.AutoWidth()
+				[
+					StructPropertyHandle->CreatePropertyNameWidget()
+				];
+
+		for (int32 ChildIndex = 0; ChildIndex < SortedChildHandles.Num(); ++ChildIndex)
+		{
+			TSharedRef<IPropertyHandle> ChildHandle = SortedChildHandles[ChildIndex];
+
+			if (IsOverrideProperty(ChildHandle))
+			{
+				continue;  // Skip overrides
+			}
+
+			const bool bLastChild = SortedChildHandles.Num() - 1 == ChildIndex;
+
+			TSharedRef<SWidget> ChildWidget = MakeChildWidget(StructPropertyHandle, ChildHandle);
+			if(ChildWidget != SNullWidget::NullWidget)
+			{
+				ValueHorizontalBox->AddSlot()
+					.Padding(FMargin(0.f, 2.f, bLastChild ? 0.f : 3.f, 2.f))
+					[
+						ChildWidget
+					];
+			}
+		}
+	}
+
+	TSharedRef<SWidget> FConnectableValueCustomization::MakeChildWidget(
+		TSharedRef<IPropertyHandle>& StructurePropertyHandle,
+		TSharedRef<IPropertyHandle>& PropertyHandle)
+	{
+		const FFieldClass* PropertyClass = PropertyHandle->GetPropertyClass();
+		
+		if (PropertyClass == FStrProperty::StaticClass())
+		{
+			// Manage override property values (properties ending with _Override)
+			TWeakPtr<IPropertyHandle> OverrideHandleWeakPtr;
+
+			for (int32 ChildIndex = 0; ChildIndex < SortedChildHandles.Num(); ++ChildIndex)
+			{
+				const bool bLastChild = SortedChildHandles.Num() - 1 == ChildIndex;
+				const TSharedRef<IPropertyHandle>& ChildHandle = SortedChildHandles[ChildIndex];
+
+				if (IsOverridePropertyOf(ChildHandle, PropertyHandle))
+				{
+					OverrideHandleWeakPtr = ChildHandle;
+					break;
+				}
+			}
+
+			TWeakPtr<IPropertyHandle> HandleWeakPtr = PropertyHandle;
+			return
+				SNew(SEditableTextBox)
+				.ToolTipText(PropertyHandle->GetToolTipText())
+				.Text_Lambda([HandleWeakPtr, OverrideHandleWeakPtr]() -> FText
+					{
+						FString Text;
+						if (const TSharedPtr<IPropertyHandle> OverrideHandlePtr = OverrideHandleWeakPtr.Pin())
+						{
+							OverrideHandlePtr->GetValue(Text);
+						}
+						if (Text == UE::Chaos::ClothAsset::FWeightMapTools::NotOverridden)
+						{
+							Text.Empty();  // GetValue seems to concatenate the text if the string isn't emptied first
+							if (const TSharedPtr<IPropertyHandle> HandlePtr = HandleWeakPtr.Pin())
+							{
+								HandlePtr->GetValue(Text);
+							}
+						}
+						return FText::FromString(Text);
+
+					})
+				.OnTextCommitted_Lambda([HandleWeakPtr](const FText& Text, ETextCommit::Type)
+					{
+						if (const TSharedPtr<IPropertyHandle> HandlePtr = HandleWeakPtr.Pin())
+						{
+							HandlePtr->SetValue(Text.ToString(), EPropertyValueSetFlags::DefaultFlags);
+						}
+					})
+				.IsEnabled_Lambda([OverrideHandleWeakPtr]() -> bool
+					{
+						FString Text;
+						if (const TSharedPtr<IPropertyHandle> OverrideHandlePtr = OverrideHandleWeakPtr.Pin())
+						{
+							OverrideHandlePtr->GetValue(Text);
+						}
+						return Text == UE::Chaos::ClothAsset::FWeightMapTools::NotOverridden;
+					})
+				.Font(IPropertyTypeCustomizationUtils::GetRegularFont());
+		}
+		return SNullWidget::NullWidget;
+	}
+}  // End namespace UE::Chaos::ClothAsset
+
+#undef LOCTEXT_NAMESPACE
diff --git a/Engine/Plugins/Experimental/ChaosClothAssetEditor/Source/ChaosClothAssetDataflowNodes/Private/ChaosClothAsset/ConnectableValueCustomization.h b/Engine/Plugins/Experimental/ChaosClothAssetEditor/Source/ChaosClothAssetDataflowNodes/Private/ChaosClothAsset/ConnectableValueCustomization.h
new file mode 100644
--- /dev/null
+++ b/Engine/Plugins/Experimental/ChaosClothAssetEditor/Source/ChaosClothAssetDataflowNodes/Private/ChaosClothAsset/ConnectableValueCustomization.h
@@ -0,0 +1,31 @@
+// Copyright Epic Games, Inc. All Rights Reserved.
+
+#pragma once
+
+#include "ChaosClothAsset/ImportedValueCustomization.h"
+
+namespace UE::Chaos::ClothAsset
+{
+	/**
+	 * Customization for all connectable property that could be imported
+	 * Works like a FMathStructCustomization.
+	 */
+	class FConnectableValueCustomization : public FImportedValueCustomization
+	{
+	public:
+		static TSharedRef<IPropertyTypeCustomization> MakeInstance();
+
+		FConnectableValueCustomization();
+		virtual ~FConnectableValueCustomization() override;
+
+	protected:
+		virtual void MakeHeaderRow(TSharedRef<class IPropertyHandle>& StructPropertyHandle, FDetailWidgetRow& Row) override;
+		virtual TSharedRef<SWidget> MakeChildWidget(TSharedRef<IPropertyHandle>& StructurePropertyHandle, TSharedRef<IPropertyHandle>& PropertyHandle) override;
+		
+		static bool IsOverrideProperty(const TSharedPtr<IPropertyHandle>& Property);
+		static bool IsOverridePropertyOf(const TSharedPtr<IPropertyHandle>& OverrideProperty, const TSharedPtr<IPropertyHandle>& Property);
+		static bool BuildFabricMapsProperty(const TSharedPtr<IPropertyHandle>& Property);
+		static bool CouldUseFabricsProperty(const TSharedPtr<IPropertyHandle>& Property);
+		
+	};
+}
diff --git a/Engine/Plugins/Experimental/ChaosClothAssetEditor/Source/ChaosClothAssetDataflowNodes/Private/ChaosClothAsset/ImportedValueCustomization.cpp b/Engine/Plugins/Experimental/ChaosClothAssetEditor/Source/ChaosClothAssetDataflowNodes/Private/ChaosClothAsset/ImportedValueCustomization.cpp
new file mode 100644
--- /dev/null
+++ b/Engine/Plugins/Experimental/ChaosClothAssetEditor/Source/ChaosClothAssetDataflowNodes/Private/ChaosClothAsset/ImportedValueCustomization.cpp
@@ -0,0 +1,206 @@
+// Copyright Epic Games, Inc. All Rights Reserved.
+
+#include "ChaosClothAsset/ImportedValueCustomization.h"
+#include "Widgets/Input/SCheckBox.h"
+#include "Widgets/Input/SEditableTextBox.h"
+#include "Widgets/Input/SVectorInputBox.h"
+#include "DetailLayoutBuilder.h"
+#include "DetailWidgetRow.h"
+#include "Editor.h"
+
+#define LOCTEXT_NAMESPACE "ChaosClothAssetImportedValueCustomization"
+
+namespace UE::Chaos::ClothAsset
+{
+	
+namespace Private
+{
+	static const FString UseImportedValue = TEXT("UseImportedValue");
+	
+	static bool UseImportedValueProperty(const TSharedPtr<IPropertyHandle>& Property)
+	{
+		const FStringView PropertyPath = Property ? Property->GetPropertyPath() : FStringView();
+		return PropertyPath.EndsWith(UseImportedValue, ESearchCase::CaseSensitive);
+	}
+}
+
+void FImportedValueCustomization::AddToggledCheckBox(const TSharedRef<IPropertyHandle>& PropertyHandle, const TSharedPtr<SHorizontalBox>& HorizontalBox, const FSlateBrush* SlateBrush)
+{
+	TWeakPtr<IPropertyHandle> WeakHandle = PropertyHandle;
+	
+	HorizontalBox->AddSlot()
+		.VAlign(VAlign_Center)
+		.HAlign(HAlign_Left)
+		.Padding(FMargin(2.f, 2.f, 2.0f, 2.f))
+		.AutoWidth()
+		[
+			SNew(SCheckBox)
+			.Padding(FMargin(0.f, 0.f, 0.0f, 0.f))
+			.Style(&FAppStyle::Get().GetWidgetStyle<FCheckBoxStyle>("ToggleButtonCheckBox"))
+			.Type(ESlateCheckBoxType::ToggleButton)
+			.ToolTipText(WeakHandle.Pin()->GetToolTipText())
+			.IsChecked_Lambda([WeakHandle]()->ECheckBoxState
+				{
+					bool bValue = false;
+					WeakHandle.Pin()->GetValue(bValue);
+					return bValue ? ECheckBoxState::Checked : ECheckBoxState::Unchecked;
+				})
+				.OnCheckStateChanged_Lambda([WeakHandle](ECheckBoxState CheckBoxState)
+				{
+					WeakHandle.Pin()->SetValue(CheckBoxState == ECheckBoxState::Checked, EPropertyValueSetFlags::DefaultFlags);
+				})
+			[
+			SNew(SImage)
+				.ColorAndOpacity(FSlateColor::UseForeground())
+				.Image(SlateBrush)
+			]
+		];
+}
+
+TSharedRef<IPropertyTypeCustomization> FImportedValueCustomization::MakeInstance()
+{
+	return MakeShareable(new FImportedValueCustomization);
+}
+
+FImportedValueCustomization::FImportedValueCustomization() = default;
+
+FImportedValueCustomization::~FImportedValueCustomization() = default;
+
+void FImportedValueCustomization::MakeHeaderRow(TSharedRef<class IPropertyHandle>& StructPropertyHandle, FDetailWidgetRow& Row)
+{
+	const TWeakPtr<IPropertyHandle> StructWeakHandlePtr = StructPropertyHandle;
+
+	TSharedPtr<SHorizontalBox> ValueHorizontalBox;
+	TSharedPtr<SHorizontalBox> NameHorizontalBox;
+
+	Row.NameContent()
+		[
+			SAssignNew(NameHorizontalBox, SHorizontalBox)
+			.IsEnabled(this, &FMathStructCustomization::IsValueEnabled, StructWeakHandlePtr)
+		]
+	.ValueContent()
+		// Make enough space for each child handle
+		.MinDesiredWidth(125.f * SortedChildHandles.Num())
+		.MaxDesiredWidth(125.f * SortedChildHandles.Num())
+		[
+			SAssignNew(ValueHorizontalBox, SHorizontalBox)
+			.IsEnabled(this, &FMathStructCustomization::IsValueEnabled, StructWeakHandlePtr)
+		];
+	
+	for (int32 ChildIndex = 0; ChildIndex < SortedChildHandles.Num(); ++ChildIndex)
+	{
+		TSharedRef<IPropertyHandle> ChildHandle = SortedChildHandles[ChildIndex];
+		
+		if (Private::UseImportedValueProperty(ChildHandle))
+		{
+			AddToggledCheckBox(ChildHandle, NameHorizontalBox, FAppStyle::Get().GetBrush("Icons.Import"));
+		}
+	}
+	
+	NameHorizontalBox->AddSlot().VAlign(VAlign_Center)
+			.HAlign(HAlign_Right)
+			.Padding(FMargin(4.f, 2.f, 4.0f, 2.f))
+			.AutoWidth()
+			[
+				StructPropertyHandle->CreatePropertyNameWidget()
+			];
+
+	for (int32 ChildIndex = 0; ChildIndex < SortedChildHandles.Num(); ++ChildIndex)
+	{
+		TSharedRef<IPropertyHandle> ChildHandle = SortedChildHandles[ChildIndex];
+		if (!Private::UseImportedValueProperty(ChildHandle))
+		{
+			const bool bLastChild = SortedChildHandles.Num() - 1 == ChildIndex;
+			
+			TSharedRef<SWidget> ChildWidget = MakeChildWidget(StructPropertyHandle, ChildHandle);
+			if(ChildWidget != SNullWidget::NullWidget)
+			{
+				ValueHorizontalBox->AddSlot()
+					.Padding(FMargin(0.f, 2.f, bLastChild ? 0.f : 3.f, 2.f))
+					[
+						ChildWidget
+					];
+			}
+		}
+	}
+}
+
+TOptional<float> FImportedValueCustomization::GetAxisValue(TSharedRef<IPropertyHandle> InPropertyHandle, EAxis::Type InAxis) const
+{
+	TOptional<float> Result;
+
+	if (InPropertyHandle->IsValidHandle())
+	{
+		const uint32 ChildIndex = InAxis == EAxis::X ? 0 : InAxis == EAxis::Y ? 1 : 2;
+		float VectorValue = 0.0;
+		if (InPropertyHandle->GetChildHandle(ChildIndex)->GetValue(VectorValue) == FPropertyAccess::Success)
+		{
+			Result = VectorValue;
+		}
+	}
+	return Result;
+}
+
+void FImportedValueCustomization::CommitAxisValue(float InNewValue, ETextCommit::Type CommitInfo, TSharedRef<IPropertyHandle> InPropertyHandle, EAxis::Type InAxis) const
+{
+	if (CommitInfo == ETextCommit::Default)
+	{
+		return;
+	}
+	if (InPropertyHandle->IsValidHandle())
+	{
+		const uint32 ChildIndex = InAxis == EAxis::X ? 0 : InAxis == EAxis::Y ? 1 : 2;
+		InPropertyHandle->GetChildHandle(ChildIndex)->SetValue(InNewValue);
+	}
+}
+
+void FImportedValueCustomization::ChangeAxisValue(float InNewValue, TSharedRef<IPropertyHandle> InPropertyHandle, EAxis::Type InAxis) const
+{
+	if (InPropertyHandle->IsValidHandle())
+	{
+		const uint32 ChildIndex = InAxis == EAxis::X ? 0 : InAxis == EAxis::Y ? 1 : 2;
+		InPropertyHandle->GetChildHandle(ChildIndex)->SetValue(InNewValue);
+	}
+}
+
+TSharedRef<SWidget> FImportedValueCustomization::MakeChildWidget(
+	TSharedRef<IPropertyHandle>& StructurePropertyHandle,
+	TSharedRef<IPropertyHandle>& PropertyHandle)
+{
+	const FFieldClass* PropertyClass = PropertyHandle->GetPropertyClass();
+
+	uint32 NumChildren = 0;
+	PropertyHandle->GetNumChildren(NumChildren);
+
+	if(NumChildren == 3 && PropertyHandle->GetChildHandle(0)->GetPropertyClass() == FFloatProperty::StaticClass())
+	{
+		const TWeakPtr<IPropertyHandle> WeakHandlePtr = PropertyHandle;
+		
+		return SNew(SNumericVectorInputBox<float>)
+			.IsEnabled(this, &FMathStructCustomization::IsValueEnabled, WeakHandlePtr)
+			.X(this, &FImportedValueCustomization::GetAxisValue, PropertyHandle, EAxis::X)
+			.Y(this, &FImportedValueCustomization::GetAxisValue, PropertyHandle, EAxis::Y)
+			.Z(this, &FImportedValueCustomization::GetAxisValue, PropertyHandle,  EAxis::Z)
+			.OnXChanged(this, &FImportedValueCustomization::ChangeAxisValue, PropertyHandle, EAxis::X)
+			.OnYChanged(this, &FImportedValueCustomization::ChangeAxisValue, PropertyHandle, EAxis::Y)
+			.OnZChanged(this, &FImportedValueCustomization::ChangeAxisValue, PropertyHandle, EAxis::Z)
+			.OnXCommitted(this, &FImportedValueCustomization::CommitAxisValue, PropertyHandle, EAxis::X)
+			.OnYCommitted(this, &FImportedValueCustomization::CommitAxisValue, PropertyHandle, EAxis::Y)
+			.OnZCommitted(this, &FImportedValueCustomization::CommitAxisValue, PropertyHandle, EAxis::Z)
+			.Font(IDetailLayoutBuilder::GetDetailFont())
+			.bColorAxisLabels(true)
+			.AllowSpin(true)
+			.SpinDelta(1.0f);
+	}
+	else if(PropertyClass == FFloatProperty::StaticClass() || PropertyClass == FIntProperty::StaticClass())
+	{
+		TSharedRef<SWidget> NumericWidget = FMathStructCustomization::MakeChildWidget(StructurePropertyHandle, PropertyHandle);
+		NumericEntryBoxWidgetList.Add(NumericWidget);
+		return NumericWidget;
+	}
+	return SNullWidget::NullWidget;
+}
+	
+}  // End namespace UE::Chaos::ClothAsset
+
+#undef LOCTEXT_NAMESPACE
diff --git a/Engine/Plugins/Experimental/ChaosClothAssetEditor/Source/ChaosClothAssetDataflowNodes/Private/ChaosClothAsset/ImportedValueCustomization.h b/Engine/Plugins/Experimental/ChaosClothAssetEditor/Source/ChaosClothAssetDataflowNodes/Private/ChaosClothAsset/ImportedValueCustomization.h
new file mode 100644
--- /dev/null
+++ b/Engine/Plugins/Experimental/ChaosClothAssetEditor/Source/ChaosClothAssetDataflowNodes/Private/ChaosClothAsset/ImportedValueCustomization.h
@@ -0,0 +1,35 @@
+// Copyright Epic Games, Inc. All Rights Reserved.
+
+#pragma once
+
+#include "Customizations/MathStructCustomizations.h"
+
+template <typename OptionalType> struct TOptional;
+
+namespace UE::Chaos::ClothAsset
+{
+	/**
+	 * Customization for all base property that could be imported
+	 * Works like a FMathStructCustomization.
+	 */
+	class FImportedValueCustomization : public FMathStructCustomization
+	{
+	public:
+		static TSharedRef<IPropertyTypeCustomization> MakeInstance();
+
+		FImportedValueCustomization();
+		virtual ~FImportedValueCustomization() override;
+
+	protected:
+		virtual void MakeHeaderRow(TSharedRef<class IPropertyHandle>& StructPropertyHandle, FDetailWidgetRow& Row) override;
+		virtual TSharedRef<SWidget> MakeChildWidget(TSharedRef<IPropertyHandle>& StructurePropertyHandle, TSharedRef<IPropertyHandle>& PropertyHandle) override;
+
+		/** Add a toggled check box to the horizontal one */
+		static void AddToggledCheckBox(const TSharedRef<IPropertyHandle>& PropertyHandle, const TSharedPtr<SHorizontalBox>& HorizontalBox, const FSlateBrush* SlateBrush);
+
+	private : 
+		TOptional<float> GetAxisValue(TSharedRef<IPropertyHandle> InPropertyHandle, EAxis::Type InAxis) const;
+		void CommitAxisValue(float InNewValue, ETextCommit::Type CommitInfo, TSharedRef<IPropertyHandle> InPropertyHandle, EAxis::Type InAxis) const;
+		void ChangeAxisValue(float InNewValue, TSharedRef<IPropertyHandle> InPropertyHandle, EAxis::Type InAxis) const;
+	};
+}
diff --git a/Engine/Plugins/Experimental/ChaosClothAssetEditor/Source/ChaosClothAssetDataflowNodes/Private/ChaosClothAsset/SimulationAerodynamicsConfigNode.cpp b/Engine/Plugins/Experimental/ChaosClothAssetEditor/Source/ChaosClothAssetDataflowNodes/Private/ChaosClothAsset/SimulationAerodynamicsConfigNode.cpp
--- a/Engine/Plugins/Experimental/ChaosClothAssetEditor/Source/ChaosClothAssetDataflowNodes/Private/ChaosClothAsset/SimulationAerodynamicsConfigNode.cpp
+++ b/Engine/Plugins/Experimental/ChaosClothAssetEditor/Source/ChaosClothAssetDataflowNodes/Private/ChaosClothAsset/SimulationAerodynamicsConfigNode.cpp
@@ -2,6 +2,7 @@
 
 #include "ChaosClothAsset/SimulationAerodynamicsConfigNode.h"
 #include "Chaos/CollectionPropertyFacade.h"
+#include "ChaosClothAsset/SimulationBaseConfigNode.h"
 
 #include UE_INLINE_GENERATED_CPP_BY_NAME(SimulationAerodynamicsConfigNode)
 
@@ -16,7 +17,17 @@ FChaosClothAssetSimulationAerodynamicsConfigNode::FChaosClothAssetSimulationAero
 void FChaosClothAssetSimulationAerodynamicsConfigNode::AddProperties(FPropertyHelper& PropertyHelper) const
 {
 	PropertyHelper.SetProperty(this, &FluidDensity);
-	PropertyHelper.SetPropertyWeighted(this, &Drag);
-	PropertyHelper.SetPropertyWeighted(this, &Lift);
 	PropertyHelper.SetProperty(this, &WindVelocity);
+	
+	PropertyHelper.SetSolverPropertyWeighted(FName(TEXT("Drag")), Drag, [](
+				const UE::Chaos::ClothAsset::FCollectionClothFacade& ClothFacade)-> float
+	{
+		return ClothFacade.GetSolverAirDamping();
+	},{});
+
+	PropertyHelper.SetSolverPropertyWeighted(FName(TEXT("Lift")), Lift, [](
+				const UE::Chaos::ClothAsset::FCollectionClothFacade& ClothFacade)-> float
+	{
+		return ClothFacade.GetSolverAirDamping();
+	},{});
 }
diff --git a/Engine/Plugins/Experimental/ChaosClothAssetEditor/Source/ChaosClothAssetDataflowNodes/Private/ChaosClothAsset/SimulationBaseConfigNode.cpp b/Engine/Plugins/Experimental/ChaosClothAssetEditor/Source/ChaosClothAssetDataflowNodes/Private/ChaosClothAsset/SimulationBaseConfigNode.cpp
--- a/Engine/Plugins/Experimental/ChaosClothAssetEditor/Source/ChaosClothAssetDataflowNodes/Private/ChaosClothAsset/SimulationBaseConfigNode.cpp
+++ b/Engine/Plugins/Experimental/ChaosClothAssetEditor/Source/ChaosClothAssetDataflowNodes/Private/ChaosClothAsset/SimulationBaseConfigNode.cpp
@@ -4,6 +4,7 @@
 #include "ChaosClothAsset/CollectionClothFacade.h"
 #include "ChaosClothAsset/ClothDataflowTools.h"
 #include "Chaos/CollectionPropertyFacade.h"
+#include "ChaosClothAsset/ClothCollectionGroup.h"
 #include "Dataflow/DataflowInputOutput.h"
 
 #define LOCTEXT_NAMESPACE "ChaosClothAssetSimulationBaseConfigNode"
@@ -42,6 +43,116 @@ namespace UE::Chaos::ClothAsset::Private
 
 		FClothDataflowTools::LogAndToastWarning(DataflowNode, Headline, Details);
 	}
+
+	static FWeightedValueBounds ComputeFabricWeightedValueBounds(FCollectionClothFacade& ClothFacade,
+		 TArray<float>& PatternValues, const TFunction<float(const FCollectionClothFabricFacade&)>& FabricValueFunction)
+	{
+		const int32 NumPatterns = ClothFacade.GetNumSimPatterns();
+		
+		PatternValues.Init(0.0f, NumPatterns);
+		float MinValue = FLT_MAX, MaxValue = 0.0f;
+		
+		for(int32 PatternIndex = 0; PatternIndex < NumPatterns; ++PatternIndex)
+		{
+			FCollectionClothSimPatternFacade PatternFacade = ClothFacade.GetSimPattern(PatternIndex);
+			const int32 FabricIndex = PatternFacade.GetFabricIndex();
+
+			if(FabricIndex >= 0 && FabricIndex < ClothFacade.GetNumFabrics())
+			{
+				FCollectionClothFabricFacade FabricFacade = ClothFacade.GetFabric(FabricIndex);
+				const float FabricValue = FabricValueFunction(FabricFacade);
+
+				MinValue = FMath::Min(MinValue, FabricValue);
+				MaxValue = FMath::Max(MaxValue, FabricValue);
+
+				PatternValues[PatternIndex] = FabricValue;
+			}
+		}
+		return {MinValue, MaxValue};
+	}
+	
+	static FWeightedValueBounds ComputePatternWeightedValueBounds(FCollectionClothFacade& ClothFacade,
+		 TArray<float>& PatternValues, const TFunction<float(const FCollectionClothSimPatternFacade&)>& PatternValueFunction)
+	{
+		const int32 NumPatterns = ClothFacade.GetNumSimPatterns();
+		
+		PatternValues.Init(0.0f, NumPatterns);
+		float MinValue = FLT_MAX, MaxValue = 0.0f;
+		
+		for(int32 PatternIndex = 0; PatternIndex < NumPatterns; ++PatternIndex)
+		{
+			FCollectionClothSimPatternFacade PatternFacade = ClothFacade.GetSimPattern(PatternIndex);
+			const float PatternValue = PatternValueFunction(PatternFacade);
+			
+			MinValue = FMath::Min(MinValue, PatternValue);
+			MaxValue = FMath::Max(MaxValue, PatternValue);
+
+			PatternValues[PatternIndex] = PatternValue;
+		}
+		return {MinValue, MaxValue};
+	}
+
+	static FWeightedValueBounds ComputeWeightedValueMap(FCollectionClothFacade& ClothFacade,
+					  const FString& WeightMapName, const FWeightedValueBounds& WeightValueBounds, const TArray<float>& PatternValues)
+	{
+		const int32 NumPatterns = ClothFacade.GetNumSimPatterns();
+		if(WeightValueBounds.Low != WeightValueBounds.High)
+		{
+			const int32 NumVertices = ClothFacade.GetNumSimVertices3D();
+		
+			TArray<int32> NumValues;
+			NumValues.Init(0, NumVertices);
+
+			TArray<float> VertexValues;
+			VertexValues.Init(0.0f, NumVertices);
+			
+			ClothFacade.AddWeightMap(*WeightMapName);
+			const TArrayView<float> ValueWeightMap = ClothFacade.GetWeightMap(*WeightMapName);
+			
+			for(int32 PatternIndex = 0; PatternIndex < NumPatterns; ++PatternIndex)
+			{
+				FCollectionClothSimPatternFacade PatternFacade = ClothFacade.GetSimPattern(PatternIndex);
+				
+				const TConstArrayView<int32> SimVertex3DLookup =
+					static_cast<FCollectionClothSimPatternConstFacade&>(PatternFacade).GetSimVertex3DLookup();
+
+				// Average of the pattern values at the seam
+				for(const int32& SimVertex3DIndex : SimVertex3DLookup)
+				{
+					VertexValues[SimVertex3DIndex] += PatternValues[PatternIndex];
+					NumValues[SimVertex3DIndex]++;
+				}
+			}
+			for(int32 VertexIndex = 0; VertexIndex < NumVertices; ++VertexIndex)
+			{
+				if(NumValues[VertexIndex] > 0)
+				{
+					ValueWeightMap[VertexIndex] = (VertexValues[VertexIndex] / NumValues[VertexIndex]  - WeightValueBounds.Low) /
+						(WeightValueBounds.High - WeightValueBounds.Low);
+				}
+			}
+		}
+		
+		return WeightValueBounds;
+	}
+
+	static FWeightedValueBounds BuildFabricWeightedValue(FCollectionClothFacade& ClothFacade,
+					  const FString& WeightMapName, const TFunction<float(const FCollectionClothFabricFacade&)>& FabricValueFunction)
+	{
+		TArray<float> PatternValues;
+		const FWeightedValueBounds WeightValueBounds = ComputeFabricWeightedValueBounds(ClothFacade, PatternValues, FabricValueFunction);
+
+		return ComputeWeightedValueMap(ClothFacade, WeightMapName, WeightValueBounds, PatternValues);
+	}
+
+	static FWeightedValueBounds BuildPatternWeightedValue(FCollectionClothFacade& ClothFacade,
+					  const FString& WeightMapName, const TFunction<float(const FCollectionClothSimPatternFacade&)>& PatternValueFunction)
+	{
+		TArray<float> PatternValues;
+		const FWeightedValueBounds WeightValueBounds = ComputePatternWeightedValueBounds(ClothFacade, PatternValues, PatternValueFunction);
+
+		return ComputeWeightedValueMap(ClothFacade, WeightMapName, WeightValueBounds, PatternValues);
+	}
 }
 
 FChaosClothAssetSimulationBaseConfigNode::FChaosClothAssetSimulationBaseConfigNode(const Dataflow::FNodeParameters& InParam, FGuid InGuid)
@@ -211,6 +322,7 @@ int32 FChaosClothAssetSimulationBaseConfigNode::FPropertyHelper::SetPropertyWeig
 		1.0f, PropertyValue.WeightMap, PropertyValue.WeightMap_Override, SimilarPropertyNames, PropertyFlags);
 }
 
+
 void FChaosClothAssetSimulationBaseConfigNode::FPropertyHelper::OverridePropertiesBool(const TArray<FName>& PropertyNames, bool bPropertyValue)
 {
 	for (const FName& PropertyName : PropertyNames)
@@ -260,6 +372,7 @@ void FChaosClothAssetSimulationBaseConfigNode::FPropertyHelper::OverrideProperti
 			if (OverrideType == EChaosClothAssetConstraintOverrideType::Override)
 			{
 				Properties.SetWeightedValue(PropertyKeyIndex, OverrideValue.Low, OverrideValue.High);
+				
 			}
 			else if (OverrideType == EChaosClothAssetConstraintOverrideType::Multiply)
 			{
@@ -281,4 +394,198 @@ int32 FChaosClothAssetSimulationBaseConfigNode::FPropertyHelper::SetPropertyStri
 	return PropertyKeyIndex;
 }
 
+template<typename PropertyType>
+void FChaosClothAssetSimulationBaseConfigNode::FPropertyHelper::SetSolverProperty(const FName& PropertyName, const PropertyType& PropertyValue,
+	const TFunction<typename PropertyType::ImportedType(UE::Chaos::ClothAsset::FCollectionClothFacade&)>& SolverValueFunction,
+	const TArray<FName>& SimilarPropertyNames, ECollectionPropertyFlags PropertyFlags)
+{
+	UE::Chaos::ClothAsset::FCollectionClothFacade ClothFacade(GetClothCollection());
+	if(PropertyValue.bUseImportedValue && ClothFacade.IsValid() && ClothFacade.HasSolverElement())
+	{
+		PropertyValue.ImportedValue = SolverValueFunction(ClothFacade);
+	}
+	SetProperty(PropertyName, PropertyValue.ImportedValue, SimilarPropertyNames, PropertyFlags);
+}
+
+template<typename PropertyType>
+void FChaosClothAssetSimulationBaseConfigNode::FPropertyHelper::SetFabricProperty(const FName& PropertyName, const PropertyType& PropertyValue,
+	const TFunction<typename PropertyType::ImportedType(UE::Chaos::ClothAsset::FCollectionClothFabricFacade&)>& FabricValueFunction,
+	const TArray<FName>& SimilarPropertyNames, ECollectionPropertyFlags PropertyFlags)
+{
+	UE::Chaos::ClothAsset::FCollectionClothFacade ClothFacade(GetClothCollection());
+	if(PropertyValue.bUseImportedValue && ClothFacade.IsValid() && ClothFacade.GetNumFabrics() > 0)
+	{
+		const int32 NumFabrics = ClothFacade.GetNumFabrics();
+		
+		typename PropertyType::ImportedType AveragedFabricValue(0.0f);
+		for(int32 FabricIndex = 0; FabricIndex < NumFabrics; ++FabricIndex)
+		{
+			UE::Chaos::ClothAsset::FCollectionClothFabricFacade FabricFacade = ClothFacade.GetFabric(FabricIndex);
+			AveragedFabricValue += FabricValueFunction(FabricFacade);
+		}
+		AveragedFabricValue /= NumFabrics;
+		PropertyValue.ImportedValue = AveragedFabricValue;
+	}
+	SetProperty(PropertyName, PropertyValue.ImportedValue, SimilarPropertyNames, PropertyFlags);
+}
+
+template<typename PropertyType>
+void FChaosClothAssetSimulationBaseConfigNode::FPropertyHelper::SetFabricPropertyWeighted(
+	const FName& PropertyName, const PropertyType& PropertyValue,
+	const TFunction<float(const UE::Chaos::ClothAsset::FCollectionClothFabricFacade&)>& FabricValueFunction,
+	const TArray<FName>& SimilarPropertyNames, ECollectionPropertyFlags PropertyFlags)
+{
+	UE::Chaos::ClothAsset::FCollectionClothFacade ClothFacade(GetClothCollection());
+	if(PropertyValue.bCouldUseFabrics && ClothFacade.IsValid() && (ClothFacade.GetNumFabrics() > 0) &&
+			(PropertyValue.bImportFabricBounds || PropertyValue.bBuildFabricMaps))
+	{
+		UE::Chaos::ClothAsset::FWeightedValueBounds WeightedValueBounds;
+		if(PropertyValue.bBuildFabricMaps)
+		{
+			WeightedValueBounds = UE::Chaos::ClothAsset::Private::BuildFabricWeightedValue(ClothFacade, 
+			 GetPropertyString(&PropertyValue.WeightMap), FabricValueFunction);
+		}
+		else 
+		{
+			TArray<float> PatternValues;
+			WeightedValueBounds = UE::Chaos::ClothAsset::Private::ComputeFabricWeightedValueBounds(
+				ClothFacade, PatternValues, FabricValueFunction);
+		}
+		if(PropertyValue.bImportFabricBounds)
+		{
+			PropertyValue.Low = WeightedValueBounds.Low;
+			PropertyValue.High = WeightedValueBounds.High;
+		}
+	}
+	SetPropertyWeighted(PropertyName, PropertyValue, SimilarPropertyNames, PropertyFlags);
+}
+
+
+template<typename MapType, typename PropertyType>
+void FChaosClothAssetSimulationBaseConfigNode::FPropertyHelper::SetFabricPropertyString(
+	const FName& PropertyName, const PropertyType& PropertyValue,
+	const TFunction<MapType(const UE::Chaos::ClothAsset::FCollectionClothFabricFacade&)>& FabricValueFunction,
+	const TArray<FName>& SimilarPropertyNames, ECollectionPropertyFlags PropertyFlags, const FName& GroupName)
+{
+	UE::Chaos::ClothAsset::FCollectionClothFacade ClothFacade(GetClothCollection());
+	if(PropertyValue.bCouldUseFabrics && ClothFacade.IsValid() && (ClothFacade.GetNumFabrics() > 0) && PropertyValue.bBuildFabricMaps)
+	{
+		const int32 NumPatterns = ClothFacade.GetNumSimPatterns();
+
+		const FName StringValue(*GetPropertyString(&PropertyValue.StringValue));
+		if (!ClothFacade.HasUserDefinedAttribute<MapType>(StringValue, GroupName))
+		{
+			ClothFacade.AddUserDefinedAttribute<MapType>(StringValue, GroupName);
+		}
+		TArrayView<MapType> UserMap = ClothFacade.GetUserDefinedAttribute<MapType>(StringValue, GroupName);
+		
+		for(int32 PatternIndex = 0; PatternIndex < NumPatterns; ++PatternIndex) 
+		{
+			UE::Chaos::ClothAsset::FCollectionClothSimPatternFacade PatternFacade = ClothFacade.GetSimPattern(PatternIndex);
+			const int32 FabricIndex = PatternFacade.GetFabricIndex();
+
+			if(FabricIndex >= 0 && FabricIndex < ClothFacade.GetNumFabrics())
+			{
+				UE::Chaos::ClothAsset::FCollectionClothFabricFacade FabricFacade = ClothFacade.GetFabric(FabricIndex);
+				const float FabricValue = FabricValueFunction(FabricFacade);
+
+				if(GroupName == UE::Chaos::ClothAsset::ClothCollectionGroup::SimFaces)
+				{
+					const int32 PatternFacesStart = PatternFacade.GetSimFacesOffset();
+					const int32 PatternFacesEnd = PatternFacade.GetNumSimFaces() + PatternFacesStart;
+				
+					for(int32 SimFaceIndex = PatternFacesStart; SimFaceIndex < PatternFacesEnd; ++SimFaceIndex)
+					{
+						UserMap[SimFaceIndex] = FabricValue;
+					}
+				}
+				else if(GroupName == UE::Chaos::ClothAsset::ClothCollectionGroup::SimVertices3D)
+				{
+					const TConstArrayView<int32> SimVertex3DLookup =
+					static_cast<UE::Chaos::ClothAsset::FCollectionClothSimPatternConstFacade&>(PatternFacade).GetSimVertex3DLookup();
+
+					for(const int32& SimVertex3DIndex : SimVertex3DLookup)
+					{
+						UserMap[SimVertex3DIndex] = FabricValue;
+					}
+				}
+			}
+		}
+	}
+	SetPropertyString(PropertyName, PropertyValue, SimilarPropertyNames, PropertyFlags);
+}
+
+template<typename PropertyType>
+void FChaosClothAssetSimulationBaseConfigNode::FPropertyHelper::SetSolverPropertyWeighted(
+	const FName& PropertyName, const PropertyType& PropertyValue,
+	const TFunction<float(const UE::Chaos::ClothAsset::FCollectionClothFacade&)>& SolverValueFunction,
+	const TArray<FName>& SimilarPropertyNames, ECollectionPropertyFlags PropertyFlags)
+{
+	UE::Chaos::ClothAsset::FCollectionClothFacade ClothFacade(GetClothCollection());
+	if(PropertyValue.bCouldUseFabrics && (ClothFacade.IsValid() && ClothFacade.HasSolverElement()) && PropertyValue.bImportFabricBounds)
+	{
+		const float SolverValue = SolverValueFunction(ClothFacade);
+	
+		PropertyValue.Low = SolverValue;
+		PropertyValue.High = SolverValue;
+	}
+	SetPropertyWeighted(PropertyName, PropertyValue, SimilarPropertyNames, PropertyFlags);
+}
+
+template void FChaosClothAssetSimulationBaseConfigNode::FPropertyHelper::SetSolverProperty<FChaosClothAssetImportedVectorValue>(const FName& PropertyName,
+	const FChaosClothAssetImportedVectorValue& PropertyValue, 
+	const TFunction<FChaosClothAssetImportedVectorValue::ImportedType(UE::Chaos::ClothAsset::FCollectionClothFacade&)>& SolverValueFunction,
+	const TArray<FName>& SimilarPropertyNames, ECollectionPropertyFlags PropertyFlags);
+
+template void FChaosClothAssetSimulationBaseConfigNode::FPropertyHelper::SetSolverProperty<FChaosClothAssetImportedFloatValue>(const FName& PropertyName,
+	const FChaosClothAssetImportedFloatValue& PropertyValue, 
+	const TFunction<FChaosClothAssetImportedFloatValue::ImportedType(UE::Chaos::ClothAsset::FCollectionClothFacade&)>& SolverValueFunction,
+	const TArray<FName>& SimilarPropertyNames, ECollectionPropertyFlags PropertyFlags);
+
+template void FChaosClothAssetSimulationBaseConfigNode::FPropertyHelper::SetSolverProperty<FChaosClothAssetImportedIntValue>(const FName& PropertyName,
+	const FChaosClothAssetImportedIntValue& PropertyValue, 
+	const TFunction<FChaosClothAssetImportedIntValue::ImportedType(UE::Chaos::ClothAsset::FCollectionClothFacade&)>& SolverValueFunction,
+	const TArray<FName>& SimilarPropertyNames, ECollectionPropertyFlags PropertyFlags);
+	
+template void FChaosClothAssetSimulationBaseConfigNode::FPropertyHelper::SetFabricProperty<FChaosClothAssetImportedVectorValue>(const FName& PropertyName,
+	const FChaosClothAssetImportedVectorValue& PropertyValue, 
+	const TFunction<FChaosClothAssetImportedVectorValue::ImportedType(UE::Chaos::ClothAsset::FCollectionClothFabricFacade&)>& FabricValueFunction,
+	const TArray<FName>& SimilarPropertyNames, ECollectionPropertyFlags PropertyFlags);
+
+template void FChaosClothAssetSimulationBaseConfigNode::FPropertyHelper::SetFabricProperty<FChaosClothAssetImportedFloatValue>(const FName& PropertyName,
+	const FChaosClothAssetImportedFloatValue& PropertyValue, 
+	const TFunction<FChaosClothAssetImportedFloatValue::ImportedType(UE::Chaos::ClothAsset::FCollectionClothFabricFacade&)>& FabricValueFunction,
+	const TArray<FName>& SimilarPropertyNames, ECollectionPropertyFlags PropertyFlags);
+
+template void FChaosClothAssetSimulationBaseConfigNode::FPropertyHelper::SetFabricProperty<FChaosClothAssetImportedIntValue>(const FName& PropertyName,
+	const FChaosClothAssetImportedIntValue& PropertyValue, 
+	const TFunction<FChaosClothAssetImportedIntValue::ImportedType(UE::Chaos::ClothAsset::FCollectionClothFabricFacade&)>& FabricValueFunction,
+	const TArray<FName>& SimilarPropertyNames, ECollectionPropertyFlags PropertyFlags);
+
+template void FChaosClothAssetSimulationBaseConfigNode::FPropertyHelper::SetSolverPropertyWeighted(
+	const FName& PropertyName, const FChaosClothAssetWeightedValueNonAnimatable& PropertyValue,
+	const TFunction<float(const UE::Chaos::ClothAsset::FCollectionClothFacade&)>& SolverValueFunction,
+	const TArray<FName>& SimilarPropertyNames, ECollectionPropertyFlags PropertyFlags);
+
+template void FChaosClothAssetSimulationBaseConfigNode::FPropertyHelper::SetSolverPropertyWeighted(
+	const FName& PropertyName, const FChaosClothAssetWeightedValue& PropertyValue,
+	const TFunction<float(const UE::Chaos::ClothAsset::FCollectionClothFacade&)>& SolverValueFunction,
+	const TArray<FName>& SimilarPropertyNames, ECollectionPropertyFlags PropertyFlags);
+
+template void FChaosClothAssetSimulationBaseConfigNode::FPropertyHelper::SetFabricPropertyWeighted<FChaosClothAssetWeightedValue>(const FName& PropertyName,
+	const FChaosClothAssetWeightedValue& PropertyValue, 
+	const TFunction<float(const UE::Chaos::ClothAsset::FCollectionClothFabricFacade&)>& FabricValueFunction,
+	const TArray<FName>& SimilarPropertyNames, ECollectionPropertyFlags PropertyFlags);
+
+template void FChaosClothAssetSimulationBaseConfigNode::FPropertyHelper::SetFabricPropertyWeighted<FChaosClothAssetWeightedValueNonAnimatable>(const FName& PropertyName,
+	const FChaosClothAssetWeightedValueNonAnimatable& PropertyValue, 
+	const TFunction<float(const UE::Chaos::ClothAsset::FCollectionClothFabricFacade&)>& FabricValueFunction,
+	const TArray<FName>& SimilarPropertyNames, ECollectionPropertyFlags PropertyFlags);
+
+template void FChaosClothAssetSimulationBaseConfigNode::FPropertyHelper::SetFabricPropertyString<int32,FChaosClothAssetConnectableIStringValue>(
+	const FName& PropertyName, const FChaosClothAssetConnectableIStringValue& PropertyValue,
+	const TFunction<int32(const UE::Chaos::ClothAsset::FCollectionClothFabricFacade&)>& FabricValueFunction,
+	const TArray<FName>& SimilarPropertyNames, ECollectionPropertyFlags PropertyFlags, const FName& GroupName);
+
+
 #undef LOCTEXT_NAMESPACE
diff --git a/Engine/Plugins/Experimental/ChaosClothAssetEditor/Source/ChaosClothAssetDataflowNodes/Private/ChaosClothAsset/SimulationBendingConfigNode.cpp b/Engine/Plugins/Experimental/ChaosClothAssetEditor/Source/ChaosClothAssetDataflowNodes/Private/ChaosClothAsset/SimulationBendingConfigNode.cpp
--- a/Engine/Plugins/Experimental/ChaosClothAssetEditor/Source/ChaosClothAssetDataflowNodes/Private/ChaosClothAsset/SimulationBendingConfigNode.cpp
+++ b/Engine/Plugins/Experimental/ChaosClothAssetEditor/Source/ChaosClothAssetDataflowNodes/Private/ChaosClothAsset/SimulationBendingConfigNode.cpp
@@ -6,7 +6,7 @@
 #include UE_INLINE_GENERATED_CPP_BY_NAME(SimulationBendingConfigNode)
 
 FChaosClothAssetSimulationBendingConfigNode::FChaosClothAssetSimulationBendingConfigNode(const Dataflow::FNodeParameters& InParam, FGuid InGuid)
-	: FChaosClothAssetSimulationFabricConfigNode(InParam, InGuid)
+	: FChaosClothAssetSimulationBaseConfigNode(InParam, InGuid)
 {
 	RegisterCollectionConnections();
 	RegisterInputConnection(&FlatnessRatio.WeightMap);
@@ -25,12 +25,6 @@ FChaosClothAssetSimulationBendingConfigNode::FChaosClothAssetSimulationBendingCo
 
 void FChaosClothAssetSimulationBendingConfigNode::AddProperties(FPropertyHelper& PropertyHelper) const
 {
-	UE::Chaos::ClothAsset::FCollectionClothFacade ClothFacade(PropertyHelper.GetClothCollection());
-	if(!ClothFacade.IsValid())
-	{
-		return;
-	}
-	
 	if(SolverType == EChaosClothAssetConstraintSolverType::XPBD)
 	{
 		if(DistributionType == EChaosClothAssetConstraintDistributionType::Isotropic)
@@ -84,52 +78,52 @@ void FChaosClothAssetSimulationBendingConfigNode::AddProperties(FPropertyHelper&
 		}
 		else
 		{
-			SetFabricPropertyWeighted(FName(TEXT("XPBDAnisoBendingStiffnessWarp")), BendingStiffnessWarp, ClothFacade, PropertyHelper, [](
+			PropertyHelper.SetFabricPropertyWeighted(FName(TEXT("XPBDAnisoBendingStiffnessWarp")), BendingStiffnessWarp, [](
 				const UE::Chaos::ClothAsset::FCollectionClothFabricFacade& FabricFacade)-> float
 			 {
-				 return FabricFacade.GetXPBDAnisoBendingStiffness().Warp;
+				 return FabricFacade.GetBendingStiffness().Warp;
 			 }, {
 				FName(TEXT("BendingSpringStiffness")),     
 				FName(TEXT("BendingElementStiffness")),     
 				FName(TEXT("XPBDBendingSpringStiffness")),  
 				FName(TEXT("XPBDBendingElementStiffness"))});
 
-			SetFabricPropertyWeighted(FName(TEXT("XPBDAnisoBendingStiffnessWeft")), BendingStiffnessWeft, ClothFacade, PropertyHelper, [](
+			PropertyHelper.SetFabricPropertyWeighted(FName(TEXT("XPBDAnisoBendingStiffnessWeft")), BendingStiffnessWeft, [](
 				const UE::Chaos::ClothAsset::FCollectionClothFabricFacade& FabricFacade)-> float
 			 {
-				 return FabricFacade.GetXPBDAnisoBendingStiffness().Weft;
+				 return FabricFacade.GetBendingStiffness().Weft;
 			 }, {});
 			
-			SetFabricPropertyWeighted(FName(TEXT("XPBDAnisoBendingStiffnessBias")), BendingStiffnessBias,ClothFacade, PropertyHelper, [](
+			PropertyHelper.SetFabricPropertyWeighted(FName(TEXT("XPBDAnisoBendingStiffnessBias")), BendingStiffnessBias, [](
 				const UE::Chaos::ClothAsset::FCollectionClothFabricFacade& FabricFacade)-> float
 			 {
-				 return FabricFacade.GetXPBDAnisoBendingStiffness().Bias;
+				 return FabricFacade.GetBendingStiffness().Bias;
 			 }, {});
 			
-			SetFabricPropertyWeighted(FName(TEXT("XPBDAnisoBucklingStiffnessWarp")), BucklingStiffnessWarp,ClothFacade, PropertyHelper, [](
+			PropertyHelper.SetFabricPropertyWeighted(FName(TEXT("XPBDAnisoBucklingStiffnessWarp")), BucklingStiffnessWarp, [](
 				const UE::Chaos::ClothAsset::FCollectionClothFabricFacade& FabricFacade)-> float
 			 {
-				 return FabricFacade.GetXPBDAnisoBucklingStiffness().Warp;
+				 return FabricFacade.GetBucklingStiffness().Warp;
 			 }, {
 				FName(TEXT("BucklingStiffness")),    
 				FName(TEXT("XPBDBucklingStiffness"))});
 			
-			SetFabricPropertyWeighted(FName(TEXT("XPBDAnisoBucklingStiffnessWeft")), BucklingStiffnessWeft, ClothFacade, PropertyHelper, [](
+			PropertyHelper.SetFabricPropertyWeighted(FName(TEXT("XPBDAnisoBucklingStiffnessWeft")), BucklingStiffnessWeft, [](
 				const UE::Chaos::ClothAsset::FCollectionClothFabricFacade& FabricFacade)-> float
 			{
-				return FabricFacade.GetXPBDAnisoBucklingStiffness().Weft;
+				return FabricFacade.GetBucklingStiffness().Weft;
 			},{});
 			
-			SetFabricPropertyWeighted(FName(TEXT("XPBDAnisoBucklingStiffnessBias")), BucklingStiffnessBias, ClothFacade, PropertyHelper, [](
+			PropertyHelper.SetFabricPropertyWeighted(FName(TEXT("XPBDAnisoBucklingStiffnessBias")), BucklingStiffnessBias, [](
 				const UE::Chaos::ClothAsset::FCollectionClothFabricFacade& FabricFacade)-> float
 			{
-				return FabricFacade.GetXPBDAnisoBucklingStiffness().Bias;
+				return FabricFacade.GetBucklingStiffness().Bias;
 			},{});
 
-			SetFabricPropertyWeighted(FName(TEXT("XPBDAnisoBendingDamping")), BendingAnisoDamping, ClothFacade, PropertyHelper, [](
+			PropertyHelper.SetFabricPropertyWeighted(FName(TEXT("XPBDAnisoBendingDamping")), BendingAnisoDamping, [](
 				const UE::Chaos::ClothAsset::FCollectionClothFabricFacade& FabricFacade)-> float
 			{
-				return FabricFacade.GetXPBDAnisoDamping();
+				return FabricFacade.GetDamping();
 			},{
 				FName(TEXT("XPBDBendingSpringDamping")),  
 				FName(TEXT("XPBDBendingElementDamping"))});
diff --git a/Engine/Plugins/Experimental/ChaosClothAssetEditor/Source/ChaosClothAssetDataflowNodes/Private/ChaosClothAsset/SimulationCollisionConfigNode.cpp b/Engine/Plugins/Experimental/ChaosClothAssetEditor/Source/ChaosClothAssetDataflowNodes/Private/ChaosClothAsset/SimulationCollisionConfigNode.cpp
--- a/Engine/Plugins/Experimental/ChaosClothAssetEditor/Source/ChaosClothAssetDataflowNodes/Private/ChaosClothAsset/SimulationCollisionConfigNode.cpp
+++ b/Engine/Plugins/Experimental/ChaosClothAssetEditor/Source/ChaosClothAssetDataflowNodes/Private/ChaosClothAsset/SimulationCollisionConfigNode.cpp
@@ -13,9 +13,39 @@ FChaosClothAssetSimulationCollisionConfigNode::FChaosClothAssetSimulationCollisi
 
 void FChaosClothAssetSimulationCollisionConfigNode::AddProperties(FPropertyHelper& PropertyHelper) const
 {
-	PropertyHelper.SetProperty(this, &FrictionCoefficient);
-	PropertyHelper.SetProperty(this, &CollisionThickness);
-	
 	PropertyHelper.SetPropertyBool(this, &bUseCCD);
 	PropertyHelper.SetProperty(this, &ProximityStiffness);
+	
+	PropertyHelper.SetFabricProperty(FName(TEXT("CollisionThickness")), CollisionThicknessImported,
+		[](UE::Chaos::ClothAsset::FCollectionClothFabricFacade& FabricFacade)-> float
+		{
+			return FabricFacade.GetCollisionThickness();
+		}, {});
+
+	PropertyHelper.SetFabricProperty(FName(TEXT("FrictionCoefficient")), FrictionCoefficientImported,
+		[](UE::Chaos::ClothAsset::FCollectionClothFabricFacade& FabricFacade)-> float
+		{
+			return FabricFacade.GetFriction();
+		}, {});
 }
+
+void FChaosClothAssetSimulationCollisionConfigNode::Serialize(FArchive& Ar)
+{
+	Super::Serialize(Ar);
+	if (Ar.IsLoading())
+	{
+#if WITH_EDITORONLY_DATA
+		if (FrictionCoefficient_DEPRECATED != UE::Chaos::ClothAsset::FDefaultFabric::Friction)
+		{
+			FrictionCoefficientImported.ImportedValue = FrictionCoefficient_DEPRECATED;
+			FrictionCoefficient_DEPRECATED = UE::Chaos::ClothAsset::FDefaultFabric::Friction;
+		}
+		if (CollisionThickness_DEPRECATED != UE::Chaos::ClothAsset::FDefaultFabric::CollisionThickness)
+		{
+			CollisionThicknessImported.ImportedValue = CollisionThickness_DEPRECATED;
+			CollisionThickness_DEPRECATED = UE::Chaos::ClothAsset::FDefaultFabric::CollisionThickness;
+		}
+#endif
+	}
+}
+
diff --git a/Engine/Plugins/Experimental/ChaosClothAssetEditor/Source/ChaosClothAssetDataflowNodes/Private/ChaosClothAsset/SimulationFabricConfigNode.cpp b/Engine/Plugins/Experimental/ChaosClothAssetEditor/Source/ChaosClothAssetDataflowNodes/Private/ChaosClothAsset/SimulationFabricConfigNode.cpp
deleted file mode 100644
diff --git a/Engine/Plugins/Experimental/ChaosClothAssetEditor/Source/ChaosClothAssetDataflowNodes/Private/ChaosClothAsset/SimulationGravityConfigNode.cpp b/Engine/Plugins/Experimental/ChaosClothAssetEditor/Source/ChaosClothAssetDataflowNodes/Private/ChaosClothAsset/SimulationGravityConfigNode.cpp
--- a/Engine/Plugins/Experimental/ChaosClothAssetEditor/Source/ChaosClothAssetDataflowNodes/Private/ChaosClothAsset/SimulationGravityConfigNode.cpp
+++ b/Engine/Plugins/Experimental/ChaosClothAssetEditor/Source/ChaosClothAssetDataflowNodes/Private/ChaosClothAsset/SimulationGravityConfigNode.cpp
@@ -3,6 +3,7 @@
 #include "ChaosClothAsset/SimulationGravityConfigNode.h"
 #include "Chaos/CollectionPropertyFacade.h"
 #include "UObject/FortniteValkyrieBranchObjectVersion.h"
+#include "UObject/FortniteReleaseBranchCustomObjectVersion.h"
 
 #include UE_INLINE_GENERATED_CPP_BY_NAME(SimulationGravityConfigNode)
 
@@ -17,7 +18,12 @@ void FChaosClothAssetSimulationGravityConfigNode::AddProperties(FPropertyHelper&
 {
 	PropertyHelper.SetPropertyBool(this, &bUseGravityOverride);
 	PropertyHelper.SetPropertyWeighted(TEXT("GravityScale"), GravityScaleWeighted);
-	PropertyHelper.SetProperty(this, &GravityOverride);
+	
+	PropertyHelper.SetSolverProperty(FName(TEXT("GravityOverride")), GravityOverrideImported,
+		[](UE::Chaos::ClothAsset::FCollectionClothFacade& ClothFacade)-> FVector3f
+		 {
+			 return ClothFacade.GetSolverGravity();
+		 }, {});
 }
 
 void FChaosClothAssetSimulationGravityConfigNode::Serialize(FArchive& Ar)
@@ -31,6 +37,11 @@ void FChaosClothAssetSimulationGravityConfigNode::Serialize(FArchive& Ar)
 		{
 			GravityScaleWeighted.Low = GravityScaleWeighted.High = GravityScale_DEPRECATED;
 		}
+		if (GravityOverride_DEPRECATED != UE::Chaos::ClothAsset::FDefaultSolver::Gravity)
+		{
+			GravityOverrideImported.ImportedValue = GravityOverride_DEPRECATED;
+			GravityOverride_DEPRECATED = UE::Chaos::ClothAsset::FDefaultSolver::Gravity;
+		}
 #endif
 	}
 }
diff --git a/Engine/Plugins/Experimental/ChaosClothAssetEditor/Source/ChaosClothAssetDataflowNodes/Private/ChaosClothAsset/SimulationMassConfigNode.cpp b/Engine/Plugins/Experimental/ChaosClothAssetEditor/Source/ChaosClothAssetDataflowNodes/Private/ChaosClothAsset/SimulationMassConfigNode.cpp
--- a/Engine/Plugins/Experimental/ChaosClothAssetEditor/Source/ChaosClothAssetDataflowNodes/Private/ChaosClothAsset/SimulationMassConfigNode.cpp
+++ b/Engine/Plugins/Experimental/ChaosClothAssetEditor/Source/ChaosClothAssetDataflowNodes/Private/ChaosClothAsset/SimulationMassConfigNode.cpp
@@ -1,7 +1,7 @@
 // Copyright Epic Games, Inc. All Rights Reserved.
 
 #include "ChaosClothAsset/SimulationMassConfigNode.h"
-#include "ChaosClothAsset/SimulationFabricConfigNode.h"
+#include "ChaosClothAsset/SimulationBaseConfigNode.h"
 #include "ChaosClothAsset/CollectionClothFacade.h"
 #include "Chaos/CollectionPropertyFacade.h"
 #include "UObject/FortniteValkyrieBranchObjectVersion.h"
@@ -9,7 +9,7 @@
 #include UE_INLINE_GENERATED_CPP_BY_NAME(SimulationMassConfigNode)
 
 FChaosClothAssetSimulationMassConfigNode::FChaosClothAssetSimulationMassConfigNode(const Dataflow::FNodeParameters& InParam, FGuid InGuid)
-	: FChaosClothAssetSimulationFabricConfigNode(InParam, InGuid)
+	: FChaosClothAssetSimulationBaseConfigNode(InParam, InGuid)
 {
 	RegisterCollectionConnections();
 	RegisterInputConnection(&UniformMassWeighted.WeightMap);
@@ -18,11 +18,6 @@ FChaosClothAssetSimulationMassConfigNode::FChaosClothAssetSimulationMassConfigNo
 
 void FChaosClothAssetSimulationMassConfigNode::AddProperties(FPropertyHelper& PropertyHelper) const
 {
-	UE::Chaos::ClothAsset::FCollectionClothFacade ClothFacade(PropertyHelper.GetClothCollection());
-	if(!ClothFacade.IsValid())
-	{
-		return;
-	}
 	PropertyHelper.SetPropertyEnum(this, &MassMode, {}, ECollectionPropertyFlags::Intrinsic);
 	switch (MassMode)
 	{
@@ -39,9 +34,9 @@ void FChaosClothAssetSimulationMassConfigNode::AddProperties(FPropertyHelper& Pr
 		break;
 	case EClothMassMode::Density:
 		{
-			SetFabricPropertyWeighted(FName(TEXT("MassValue")), DensityWeighted, ClothFacade, PropertyHelper, [](const UE::Chaos::ClothAsset::FCollectionClothFabricFacade& FabricFacade)-> float
+			PropertyHelper.SetFabricPropertyWeighted(FName(TEXT("MassValue")), DensityWeighted, [](const UE::Chaos::ClothAsset::FCollectionClothFabricFacade& FabricFacade)-> float
 			 {
-				 return FabricFacade.GetDensityWeighted();
+				 return FabricFacade.GetDensity();
 			 }, {}, ECollectionPropertyFlags::Intrinsic);
 		}
 		break;
diff --git a/Engine/Plugins/Experimental/ChaosClothAssetEditor/Source/ChaosClothAssetDataflowNodes/Private/ChaosClothAsset/SimulationPressureConfigNode.cpp b/Engine/Plugins/Experimental/ChaosClothAssetEditor/Source/ChaosClothAssetDataflowNodes/Private/ChaosClothAsset/SimulationPressureConfigNode.cpp
--- a/Engine/Plugins/Experimental/ChaosClothAssetEditor/Source/ChaosClothAssetDataflowNodes/Private/ChaosClothAsset/SimulationPressureConfigNode.cpp
+++ b/Engine/Plugins/Experimental/ChaosClothAssetEditor/Source/ChaosClothAssetDataflowNodes/Private/ChaosClothAsset/SimulationPressureConfigNode.cpp
@@ -14,5 +14,9 @@ FChaosClothAssetSimulationPressureConfigNode::FChaosClothAssetSimulationPressure
 
 void FChaosClothAssetSimulationPressureConfigNode::AddProperties(FPropertyHelper& PropertyHelper) const
 {
-	PropertyHelper.SetPropertyWeighted(this, &Pressure);
+	PropertyHelper.SetFabricPropertyWeighted(FName(TEXT("Pressure")), Pressure, [](
+				const UE::Chaos::ClothAsset::FCollectionClothFabricFacade& FabricFacade)-> float
+	{
+		return FabricFacade.GetPressure();
+	}, {});
 }
diff --git a/Engine/Plugins/Experimental/ChaosClothAssetEditor/Source/ChaosClothAssetDataflowNodes/Private/ChaosClothAsset/SimulationSelfCollisionConfigNode.cpp b/Engine/Plugins/Experimental/ChaosClothAssetEditor/Source/ChaosClothAssetDataflowNodes/Private/ChaosClothAsset/SimulationSelfCollisionConfigNode.cpp
--- a/Engine/Plugins/Experimental/ChaosClothAssetEditor/Source/ChaosClothAssetDataflowNodes/Private/ChaosClothAsset/SimulationSelfCollisionConfigNode.cpp
+++ b/Engine/Plugins/Experimental/ChaosClothAssetEditor/Source/ChaosClothAssetDataflowNodes/Private/ChaosClothAsset/SimulationSelfCollisionConfigNode.cpp
@@ -3,6 +3,7 @@
 #include "ChaosClothAsset/SimulationSelfCollisionConfigNode.h"
 #include "Dataflow/DataflowInputOutput.h"
 #include "Chaos/CollectionPropertyFacade.h"
+#include "ChaosClothAsset/ClothCollectionGroup.h"
 
 #include UE_INLINE_GENERATED_CPP_BY_NAME(SimulationSelfCollisionConfigNode)
 
@@ -20,11 +21,8 @@ FChaosClothAssetSimulationSelfCollisionConfigNode::FChaosClothAssetSimulationSel
 void FChaosClothAssetSimulationSelfCollisionConfigNode::AddProperties(FPropertyHelper& PropertyHelper) const
 {
 	PropertyHelper.SetPropertyBool(FName("UseSelfCollisions"), true);
-	PropertyHelper.SetPropertyWeighted(TEXT("SelfCollisionThickness"), SelfCollisionThicknessWeighted);
 	PropertyHelper.SetProperty(this, &SelfCollisionStiffness);
-	PropertyHelper.SetProperty(this, &SelfCollisionFriction);
 	PropertyHelper.SetProperty(this, &SelfCollisionDisableNeighborDistance, {}, ECollectionPropertyFlags::None); // Non animatable
-	PropertyHelper.SetPropertyString(this, &SelfCollisionLayers);
 	PropertyHelper.SetPropertyString(this, &SelfCollisionDisabledFaces);
 	PropertyHelper.SetPropertyBool(this, &bSelfCollideAgainstKinematicCollidersOnly);
 	PropertyHelper.SetPropertyBool(this, &bSelfCollideAgainstAllKinematicVertices);
@@ -39,6 +37,24 @@ void FChaosClothAssetSimulationSelfCollisionConfigNode::AddProperties(FPropertyH
 	PropertyHelper.SetProperty(this, &NumContourMinimizationPostSteps);
 	PropertyHelper.SetPropertyBool(this, &bUseGlobalPostStepContours);
 	PropertyHelper.SetProperty(this, &SelfCollisionProximityStiffness);
+
+	PropertyHelper.SetFabricPropertyString<int32,FChaosClothAssetConnectableIStringValue>(FName(TEXT("SelfCollisionLayers")), SelfCollisionLayers,
+		[](const UE::Chaos::ClothAsset::FCollectionClothFabricFacade& FabricFacade)-> float
+		{
+			return FabricFacade.GetLayer();
+		}, {}, ECollectionPropertyFlags::None, UE::Chaos::ClothAsset::ClothCollectionGroup::SimFaces);
+	
+	PropertyHelper.SetFabricPropertyWeighted(FName(TEXT("SelfCollisionThickness")), SelfCollisionThicknessWeighted,
+		[](const UE::Chaos::ClothAsset::FCollectionClothFabricFacade& FabricFacade)-> float
+		{
+			return FabricFacade.GetCollisionThickness();
+		}, {});
+
+	PropertyHelper.SetFabricProperty(FName(TEXT("SelfCollisionFriction")), SelfCollisionFrictionImported,
+		[](UE::Chaos::ClothAsset::FCollectionClothFabricFacade& FabricFacade)-> float
+		{
+			return FabricFacade.GetFriction();
+		}, {});
 }
 
 void FChaosClothAssetSimulationSelfCollisionConfigNode::Serialize(FArchive& Ar)
@@ -47,6 +63,11 @@ void FChaosClothAssetSimulationSelfCollisionConfigNode::Serialize(FArchive& Ar)
 	if (Ar.IsLoading())
 	{
 #if WITH_EDITORONLY_DATA
+		if (SelfCollisionFriction_DEPRECATED != UE::Chaos::ClothAsset::FDefaultFabric::SelfFriction)
+		{
+			SelfCollisionFrictionImported.ImportedValue = SelfCollisionFriction_DEPRECATED;
+			SelfCollisionFriction_DEPRECATED = UE::Chaos::ClothAsset::FDefaultFabric::SelfFriction;
+		}
 		if (SelfCollisionKinematicColliderFriction_DEPRECATED != FrictionDeprecatedValue)
 		{
 			SelfCollisionKinematicColliderFrictionWeighted.Low = SelfCollisionKinematicColliderFrictionWeighted.High = SelfCollisionKinematicColliderFriction_DEPRECATED;
diff --git a/Engine/Plugins/Experimental/ChaosClothAssetEditor/Source/ChaosClothAssetDataflowNodes/Private/ChaosClothAsset/SimulationSolverConfigNode.cpp b/Engine/Plugins/Experimental/ChaosClothAssetEditor/Source/ChaosClothAssetDataflowNodes/Private/ChaosClothAsset/SimulationSolverConfigNode.cpp
--- a/Engine/Plugins/Experimental/ChaosClothAssetEditor/Source/ChaosClothAssetDataflowNodes/Private/ChaosClothAsset/SimulationSolverConfigNode.cpp
+++ b/Engine/Plugins/Experimental/ChaosClothAssetEditor/Source/ChaosClothAssetDataflowNodes/Private/ChaosClothAsset/SimulationSolverConfigNode.cpp
@@ -15,7 +15,6 @@ void FChaosClothAssetSimulationSolverConfigNode::AddProperties(FPropertyHelper&
 {
 	PropertyHelper.SetProperty(this, &NumIterations);
 	PropertyHelper.SetProperty(this, &MaxNumIterations);
-	PropertyHelper.SetProperty(this, &NumSubsteps);
 	const float DynamicSubstepDeltaTimeValue = bEnableDynamicSubstepping ? DynamicSubstepDeltaTime : 0.f;
 	PropertyHelper.SetProperty(TEXT("DynamicSubstepDeltaTime"), DynamicSubstepDeltaTimeValue);
 	PropertyHelper.SetPropertyBool(this, &bEnableNumSelfCollisionSubsteps);
@@ -25,6 +24,12 @@ void FChaosClothAssetSimulationSolverConfigNode::AddProperties(FPropertyHelper&
 	PropertyHelper.SetProperty(this, &MaxNumCGIterations);
 	PropertyHelper.SetProperty(this, &CGResidualTolerance);
 	PropertyHelper.SetPropertyBool(this, &bDoQuasistatics);
+	
+	PropertyHelper.SetSolverProperty(FName(TEXT("NumSubsteps")), NumSubstepsImported,
+		[](UE::Chaos::ClothAsset::FCollectionClothFacade& ClothFacade)-> int32
+		 {
+			 return ClothFacade.GetSolverSubSteps();
+		 }, {});
 }
 
 void FChaosClothAssetSimulationSolverConfigNode::Serialize(FArchive& Ar)
@@ -40,5 +45,12 @@ void FChaosClothAssetSimulationSolverConfigNode::Serialize(FArchive& Ar)
 				NumNewtonIterations = 0;
 			}
 		}
+#if WITH_EDITORONLY_DATA
+		if (NumSubsteps_DEPRECATED != UE::Chaos::ClothAsset::FDefaultSolver::SubSteps)
+		{
+			NumSubstepsImported.ImportedValue = NumSubsteps_DEPRECATED;
+			NumSubsteps_DEPRECATED = UE::Chaos::ClothAsset::FDefaultSolver::SubSteps;
+		}
+#endif
 	}
 }
diff --git a/Engine/Plugins/Experimental/ChaosClothAssetEditor/Source/ChaosClothAssetDataflowNodes/Private/ChaosClothAsset/SimulationStretchConfigNode.cpp b/Engine/Plugins/Experimental/ChaosClothAssetEditor/Source/ChaosClothAssetDataflowNodes/Private/ChaosClothAsset/SimulationStretchConfigNode.cpp
--- a/Engine/Plugins/Experimental/ChaosClothAssetEditor/Source/ChaosClothAssetDataflowNodes/Private/ChaosClothAsset/SimulationStretchConfigNode.cpp
+++ b/Engine/Plugins/Experimental/ChaosClothAssetEditor/Source/ChaosClothAssetDataflowNodes/Private/ChaosClothAsset/SimulationStretchConfigNode.cpp
@@ -6,7 +6,7 @@
 #include UE_INLINE_GENERATED_CPP_BY_NAME(SimulationStretchConfigNode)
 
 FChaosClothAssetSimulationStretchConfigNode::FChaosClothAssetSimulationStretchConfigNode(const Dataflow::FNodeParameters& InParam, FGuid InGuid)
-	: FChaosClothAssetSimulationFabricConfigNode(InParam, InGuid)
+	: FChaosClothAssetSimulationBaseConfigNode(InParam, InGuid)
 {
 	RegisterCollectionConnections();
 	RegisterInputConnection(&StretchStiffness.WeightMap);
@@ -22,12 +22,6 @@ FChaosClothAssetSimulationStretchConfigNode::FChaosClothAssetSimulationStretchCo
 
 void FChaosClothAssetSimulationStretchConfigNode::AddProperties(FPropertyHelper& PropertyHelper) const
 {
-	UE::Chaos::ClothAsset::FCollectionClothFacade ClothFacade(PropertyHelper.GetClothCollection());
-	if(!ClothFacade.IsValid())
-	{
-		return;
-	}
-	
 	if(SolverType == EChaosClothAssetConstraintSolverType::XPBD)
 	{
 		if (DistributionType == EChaosClothAssetConstraintDistributionType::Anisotropic)
@@ -35,31 +29,31 @@ void FChaosClothAssetSimulationStretchConfigNode::AddProperties(FPropertyHelper&
 			PropertyHelper.SetPropertyBool(FName(TEXT("XPBDAnisoSpringUse3dRestLengths")), bStretchUse3dRestLengths, {
 				FName(TEXT("XPBDAnisoStretchUse3dRestLengths"))}, ECollectionPropertyFlags::None);  // Non animatable
 
-			SetFabricPropertyWeighted(FName(TEXT("XPBDAnisoSpringStiffnessWarp")), StretchStiffnessWarp,ClothFacade, PropertyHelper, [](
+			PropertyHelper.SetFabricPropertyWeighted(FName(TEXT("XPBDAnisoSpringStiffnessWarp")), StretchStiffnessWarp, [](
 				const UE::Chaos::ClothAsset::FCollectionClothFabricFacade& FabricFacade)-> float
 			 {
-				 return FabricFacade.GetXPBDAnisoSpringStiffness().Warp;
+				 return FabricFacade.GetStretchStiffness().Warp;
 			 }, {
 				FName(TEXT("EdgeSpringStiffness")),
 				FName(TEXT("XPBDEdgeSpringStiffness")),
 				FName(TEXT("XPBDAnisoStretchStiffnessWarp"))});
 
-			SetFabricPropertyWeighted(FName(TEXT("XPBDAnisoSpringStiffnessWeft")), StretchStiffnessWeft, ClothFacade, PropertyHelper, [](
+			PropertyHelper.SetFabricPropertyWeighted(FName(TEXT("XPBDAnisoSpringStiffnessWeft")), StretchStiffnessWeft, [](
 				const UE::Chaos::ClothAsset::FCollectionClothFabricFacade& FabricFacade)-> float
 			 {
-				 return FabricFacade.GetXPBDAnisoSpringStiffness().Weft;
+				 return FabricFacade.GetStretchStiffness().Weft;
 			 }, { FName(TEXT("XPBDAnisoStretchStiffnessWeft")) });
 	
-			SetFabricPropertyWeighted(FName(TEXT("XPBDAnisoSpringStiffnessBias")), StretchStiffnessBias, ClothFacade, PropertyHelper, [](
+			PropertyHelper.SetFabricPropertyWeighted(FName(TEXT("XPBDAnisoSpringStiffnessBias")), StretchStiffnessBias, [](
 				const UE::Chaos::ClothAsset::FCollectionClothFabricFacade& FabricFacade)-> float
 			 {
-				 return FabricFacade.GetXPBDAnisoSpringStiffness().Bias;
+				 return FabricFacade.GetStretchStiffness().Bias;
 			 }, { FName(TEXT("XPBDAnisoStretchStiffnessBias")) });
 
-			SetFabricPropertyWeighted(FName(TEXT("XPBDAnisoSpringDamping")), StretchAnisoDamping, ClothFacade, PropertyHelper, [](
+			PropertyHelper.SetFabricPropertyWeighted(FName(TEXT("XPBDAnisoSpringDamping")), StretchAnisoDamping, [](
 				const UE::Chaos::ClothAsset::FCollectionClothFabricFacade& FabricFacade)-> float
 			 {
-				 return FabricFacade.GetXPBDAnisoDamping();
+				 return FabricFacade.GetDamping();
 			 }, { FName(TEXT("XPBDEdgeSpringDamping")),
 				FName(TEXT("XPBDAnisoStretchDamping")) });
 
diff --git a/Engine/Plugins/Experimental/ChaosClothAssetEditor/Source/ChaosClothAssetDataflowNodes/Private/ChaosClothAsset/USDImportNode.cpp b/Engine/Plugins/Experimental/ChaosClothAssetEditor/Source/ChaosClothAssetDataflowNodes/Private/ChaosClothAsset/USDImportNode.cpp
--- a/Engine/Plugins/Experimental/ChaosClothAssetEditor/Source/ChaosClothAssetDataflowNodes/Private/ChaosClothAsset/USDImportNode.cpp
+++ b/Engine/Plugins/Experimental/ChaosClothAssetEditor/Source/ChaosClothAssetDataflowNodes/Private/ChaosClothAsset/USDImportNode.cpp
@@ -604,9 +604,9 @@ bool FChaosClothAssetUSDImportNode::ImportFromFile(const FString& UsdFilePath, c
 		}
 	}
 
-	auto FillIntDatas = [](const UE::FUsdPrim& FabricPrim, const FString& DatasName, uint32& IntDatas) 
+	auto FillIntDatas = [](const UE::FUsdPrim& UsdPrim, const FString& DatasName, uint32& IntDatas) 
 	{
-		const UE::FUsdAttribute IntDatasAttr = FabricPrim.GetAttribute(*DatasName);
+		const UE::FUsdAttribute IntDatasAttr = UsdPrim.GetAttribute(*DatasName);
 		if (IntDatasAttr.HasValue() && IntDatasAttr.GetTypeName() == TEXT("uint"))
 		{
 			UE::FVtValue Value;
@@ -615,23 +615,107 @@ bool FChaosClothAssetUSDImportNode::ImportFromFile(const FString& UsdFilePath, c
 			IntDatas = Optional.IsSet() ? Optional.GetValue() : 0;
 		}
 	};
+	auto FillFloatDatas = [](const UE::FUsdPrim& UsdPrim, const FString& DatasName, float& FloatDatas) 
+	{
+		const UE::FUsdAttribute FloatDatasAttr = UsdPrim.GetAttribute(*DatasName);
+		if (FloatDatasAttr.HasValue() && FloatDatasAttr.GetTypeName() == TEXT("float"))
+		{
+			UE::FVtValue Value;
+			FloatDatasAttr.Get(Value);
+			const TOptional<float> Optional = UsdUtils::GetUnderlyingValue<float>(Value);
+			FloatDatas = Optional.IsSet() ? Optional.GetValue() : 0.0f;
+		}
+	};
+	auto FillVectorDatas = [&AxesOrder](const UE::FUsdPrim& UsdPrim, const FString& DatasName, FVector3f& VectorDatas) 
+	{
+		const UE::FUsdAttribute VectorDatasAttr = UsdPrim.GetAttribute(*DatasName);
+		if (VectorDatasAttr.HasValue() && VectorDatasAttr.GetTypeName() == TEXT("float3"))
+		{
+			UE::FVtValue Value;
+			VectorDatasAttr.Get(Value);
+			UsdUtils::FConvertedVtValue ConvertedVtValue;
+			if (UsdToUnreal::ConvertValue(Value, ConvertedVtValue) && !ConvertedVtValue.bIsArrayValued && !ConvertedVtValue.bIsEmpty)
+			{
+				VectorDatas = FVector3f(ConvertedVtValue.Entries[0][AxesOrder[0]].Get<float>(),
+				ConvertedVtValue.Entries[0][AxesOrder[1]].Get<float>(),
+				ConvertedVtValue.Entries[0][AxesOrder[2]].Get<float>());
+			}
+		}
+	};
+
+	auto FillFloatArrayDatas = [](const UE::FUsdPrim& UsdPrim, const FString& DatasName, TArray<float>& FloatArrayDatas) 
+	{
+		const UE::FUsdAttribute FloatArrayDatasAttr = UsdPrim.GetAttribute(*DatasName);
+		if (FloatArrayDatasAttr.HasValue() && FloatArrayDatasAttr.GetTypeName() == TEXT("float[]"))
+		{
+			UE::FVtValue Value;
+			FloatArrayDatasAttr.Get(Value);
+			UsdUtils::FConvertedVtValue ConvertedVtValue;
+			if (UsdToUnreal::ConvertValue(Value, ConvertedVtValue) && ConvertedVtValue.bIsArrayValued && !ConvertedVtValue.bIsEmpty)
+			{
+				if (const int32 NumEntries = ConvertedVtValue.Entries.Num())
+				{
+					for (int32 Index = 0; Index < NumEntries; ++Index)
+					{
+						FloatArrayDatas.Add(ConvertedVtValue.Entries[Index][0].Get<float>());
+					}
+				}
+			}
+		}
+	};
+
+	auto FillIntArrayDatas = [](const UE::FUsdPrim& UsdPrim, const FString& DatasName, TArray<int32>& IntArrayDatas) 
+	{
+		const UE::FUsdAttribute IntArrayDatasAttr = UsdPrim.GetAttribute(*DatasName);
+		if (IntArrayDatasAttr.HasValue() && IntArrayDatasAttr.GetTypeName() == TEXT("int[]"))
+		{
+			UE::FVtValue Value;
+			IntArrayDatasAttr.Get(Value);
+			UsdUtils::FConvertedVtValue ConvertedVtValue;
+			if (UsdToUnreal::ConvertValue(Value, ConvertedVtValue) && ConvertedVtValue.bIsArrayValued && !ConvertedVtValue.bIsEmpty)
+			{
+				if (const int32 NumTriangles = ConvertedVtValue.Entries.Num())
+				{
+					UE_LOG(LogChaosClothAssetDataflowNodes, Log, TEXT("Num Triangles = %d"), NumTriangles);
+					for (int32 Index = 0; Index < NumTriangles; ++Index)
+					{
+						UE_LOG(LogChaosClothAssetDataflowNodes, Log, TEXT("Num Entries[%d] = %d"), Index, ConvertedVtValue.Entries[Index].Num());
+						IntArrayDatas.Add(ConvertedVtValue.Entries[Index][0].Get<int32>());
+					}
+				}
+			}
+		}
+	};
+
+	// Simulation properties
+	const UE::FSdfPath SimulationPropertiesPath = UE::FSdfPath(UE::FSdfPath::AbsoluteRootPath()).AppendChild(TEXT("SimulationData")).AppendChild(TEXT("SimulationProperties"));
+	if (const UE::FUsdPrim SimulationPropertiesPrim = UsdStage.GetPrimAtPath(SimulationPropertiesPath))
+	{
+		float AirDamping = 0.1f;
+		FillFloatDatas(SimulationPropertiesPrim, TEXT("AirDamping"), AirDamping);
+
+		FVector3f Gravity(0.0f, 0.0f, -9810.0f);
+		FillVectorDatas(SimulationPropertiesPrim, TEXT("Gravity"), Gravity);
+
+		float TimeStep = 0.033f;
+		FillFloatDatas(SimulationPropertiesPrim, TEXT("TimeStep"), TimeStep);
+		
+		uint32 SubSteps = 1;
+		FillIntDatas(SimulationPropertiesPrim, TEXT("SubStepCount"), SubSteps);
+
+		static constexpr float GravityScaling = 1e-1f; // from mm to cm
+
+		ClothFacade.SetSolverGravity(Gravity * GravityScaling);
+		ClothFacade.SetSolverAirDamping(AirDamping);
+		ClothFacade.SetSolverTimeStep(TimeStep);
+		ClothFacade.SetSolverSubSteps(SubSteps);
+	}
 
 	// Fabrics
 	TArray<uint32> FabricIds;
 	const UE::FSdfPath FabricsPath = UE::FSdfPath(UE::FSdfPath::AbsoluteRootPath()).AppendChild(TEXT("SimulationData")).AppendChild(TEXT("Fabrics"));
 	if (const UE::FUsdPrim FabricsPrim = UsdStage.GetPrimAtPath(FabricsPath))
 	{
-		auto FillFloatDatas = [](const UE::FUsdPrim& FabricPrim, const FString& DatasName, float& FloatDatas) 
-		{
-			const UE::FUsdAttribute FloatDatasAttr = FabricPrim.GetAttribute(*DatasName);
-			if (FloatDatasAttr.HasValue() && FloatDatasAttr.GetTypeName() == TEXT("float"))
-			{
-				UE::FVtValue Value;
-				FloatDatasAttr.Get(Value);
-				const TOptional<float> Optional = UsdUtils::GetUnderlyingValue<float>(Value);
-				FloatDatas = Optional.IsSet() ? Optional.GetValue() : 0.0f;
-			}
-		};
 		for (const UE::FUsdPrim& FabricPrim : FabricsPrim.GetChildren())
 		{
 			float BendingBiasLeft = 0.0f;
@@ -676,8 +760,8 @@ bool FChaosClothAssetUSDImportNode::ImportFromFile(const FString& UsdFilePath, c
 			float Friction = 0.0f;
 			FillFloatDatas(FabricPrim, TEXT("Friction"), Friction);
 
-			float InternalDamping = 0.0f;
-			FillFloatDatas(FabricPrim, TEXT("InternalDamping"), InternalDamping);
+			float Damping = 0.0f;
+			FillFloatDatas(FabricPrim, TEXT("InternalDamping"), Damping);
 
 			float Thickness = 0.0f;
 			FillFloatDatas(FabricPrim, TEXT("Thickness"), Thickness);
@@ -700,7 +784,7 @@ bool FChaosClothAssetUSDImportNode::ImportFromFile(const FString& UsdFilePath, c
 			FCollectionClothFabricFacade Fabric = ClothFacade.AddGetFabric();
 
 			static constexpr float BendingScaling = 1e-5f; // from g.mm2/s2 to kg.cm2/s2
-			static constexpr float StretchShearScaling = 1e-3f; // from g/s2 to kg.cm/s2
+			static constexpr float StretchShearScaling = 1e-3f; // from g/s2 to kg/s2
 			static constexpr float DensityScaling = 1e+3f; // from g/mm2 to kg/m2
 			static constexpr float ThicknessScaling = 1e-1f; // from mm to cm
 			
@@ -719,7 +803,7 @@ bool FChaosClothAssetUSDImportNode::ImportFromFile(const FString& UsdFilePath, c
 				0.5f * (BucklingRatioBiasLeft+BucklingRatioBiasRight)) / 3.0f;
 
 			Fabric.Initialize(BendingStiffness, BucklingRatio, BucklingStiffness, StretchStiffness,
-				Density * DensityScaling, Friction, InternalDamping, Thickness * ThicknessScaling);
+				Density * DensityScaling, Friction, Damping, 0.0f, 0, Thickness * ThicknessScaling);
 
 			FabricIds.Add(FabricId);
 		}
@@ -925,6 +1009,103 @@ bool FChaosClothAssetUSDImportNode::ImportFromFile(const FString& UsdFilePath, c
 		}
 	}
 
+	// Triangles
+	const UE::FSdfPath TrianglesPath = UE::FSdfPath(UE::FSdfPath::AbsoluteRootPath()).AppendChild(TEXT("SimulationData")).AppendChild(TEXT("Triangles"));
+	if (const UE::FUsdPrim TrianglesPrim = UsdStage.GetPrimAtPath(TrianglesPath))
+	{
+		const int32 NumSimFaces = ClothFacade.GetNumSimFaces();
+		
+		TArray<float> TrianglesCollisionThickness;
+		TrianglesCollisionThickness.Reserve(NumSimFaces);
+		FillFloatArrayDatas(TrianglesPrim, TEXT("CollisionThickness"), TrianglesCollisionThickness);
+
+		TArray<int32> TrianglesPatternLayer;
+		TrianglesPatternLayer.Reserve(NumSimFaces);
+		FillIntArrayDatas(TrianglesPrim, TEXT("Layer"), TrianglesPatternLayer);
+
+		TArray<float> TrianglesPatternPressure;
+		TrianglesPatternPressure.Reserve(NumSimFaces);
+		FillFloatArrayDatas(TrianglesPrim, TEXT("Pressure"), TrianglesPatternPressure);
+
+		if(TrianglesCollisionThickness.Num() == TrianglesPatternLayer.Num() &&
+			TrianglesCollisionThickness.Num() == TrianglesPatternPressure.Num() &&
+			TrianglesCollisionThickness.Num() == NumSimFaces)
+		{
+			// Struct datas that will be used to retrieve the correct fabric index
+			struct FFabricPatternDatas
+			{
+				float PatternPressure;
+				int32 PatternLayer;
+				float CollisionThickness;
+				int32 FabricIndex;
+			};
+			const int32 NumPatterns = ClothFacade.GetNumSimPatterns();
+			TArray<TArray<FFabricPatternDatas>> FabricPatternDatas;
+			FabricPatternDatas.SetNum(ClothFacade.GetNumFabrics());
+			for(int32 PatternIndex = 0; PatternIndex < NumPatterns; ++PatternIndex)
+			{
+				FCollectionClothSimPatternFacade PatternFacade = ClothFacade.GetSimPattern(PatternIndex);
+				
+				const int32 PatternFacesStart = PatternFacade.GetSimFacesOffset();
+				const int32 PatternFacesEnd = PatternFacade.GetNumSimFaces() + PatternFacesStart;
+
+				float CollisionThickness = 0.0;
+				float PatternPressure = 0.0;
+				int32 PatternLayer = (PatternFacade.GetNumSimFaces() > 0) ? TrianglesPatternLayer[PatternFacesStart] : INDEX_NONE;
+				bool bHasUniformLayer = true;
+				for(int32 PatternFaceIndex = PatternFacesStart; PatternFaceIndex < PatternFacesEnd; ++PatternFaceIndex)
+				{
+					CollisionThickness += TrianglesCollisionThickness[PatternFaceIndex];
+					PatternPressure += TrianglesPatternPressure[PatternFaceIndex];
+					bHasUniformLayer = (TrianglesPatternLayer[PatternFaceIndex] != PatternLayer) ? false : bHasUniformLayer;
+				}
+				PatternLayer = !bHasUniformLayer ? INDEX_NONE : PatternLayer;
+				CollisionThickness /= PatternFacade.GetNumSimFaces();
+				PatternPressure /= PatternFacade.GetNumSimFaces();
+
+				static constexpr float ThicknessScaling = 1e-1f; // from mm to cm
+				CollisionThickness *= ThicknessScaling;
+
+				const int32 FabricIndex = PatternFacade.GetFabricIndex();
+				if(FabricIndex >= 0 && FabricIndex < ClothFacade.GetNumFabrics())
+				{
+					FCollectionClothFabricFacade OldFabricFacade = ClothFacade.GetFabric(FabricIndex);
+					if(FabricPatternDatas[FabricIndex].IsEmpty())
+					{
+						// If empty we update the existing fabric
+						OldFabricFacade.Initialize(OldFabricFacade, PatternPressure, PatternLayer, CollisionThickness);
+						
+						// Store the existing fabric into the array for future potential reuse
+						FabricPatternDatas[FabricIndex].Add({PatternPressure, PatternLayer, CollisionThickness, FabricIndex});
+					}
+					else
+					{
+						bool bFoundMatchingFabric = false;
+						for(FFabricPatternDatas& PatternDatas : FabricPatternDatas[FabricIndex])
+						{
+							// If the fabric already in use and if the pattern datas are matching, reuse the fabric 
+							if((PatternDatas.CollisionThickness == CollisionThickness) && (PatternDatas.PatternLayer == PatternLayer) && (PatternDatas.PatternPressure == PatternPressure))
+							{
+								bFoundMatchingFabric = true;
+								PatternFacade.SetFabricIndex(PatternDatas.FabricIndex);
+								break;
+							}
+						}
+						if(!bFoundMatchingFabric)
+						{
+							FCollectionClothFabricFacade NewFabricFacade = ClothFacade.AddGetFabric();
+							NewFabricFacade.Initialize(OldFabricFacade, PatternPressure, PatternLayer, CollisionThickness);
+
+							// Store the new fabric into the array for future potential reuse
+							FabricPatternDatas[FabricIndex].Add({PatternPressure, PatternLayer, CollisionThickness, NewFabricFacade.GetElementIndex()});
+							PatternFacade.SetFabricIndex(NewFabricFacade.GetElementIndex());
+						}
+					}
+				}
+			}
+		}
+	}
+
 	// Check for duplicate stitches
 	Private::RemoveDuplicateStitches(SeamStitches);
 
diff --git a/Engine/Plugins/Experimental/ChaosClothAssetEditor/Source/ChaosClothAssetDataflowNodes/Private/ChaosClothAsset/WeightedValueCustomization.cpp b/Engine/Plugins/Experimental/ChaosClothAssetEditor/Source/ChaosClothAssetDataflowNodes/Private/ChaosClothAsset/WeightedValueCustomization.cpp
--- a/Engine/Plugins/Experimental/ChaosClothAssetEditor/Source/ChaosClothAssetDataflowNodes/Private/ChaosClothAsset/WeightedValueCustomization.cpp
+++ b/Engine/Plugins/Experimental/ChaosClothAssetEditor/Source/ChaosClothAssetDataflowNodes/Private/ChaosClothAsset/WeightedValueCustomization.cpp
@@ -16,69 +16,13 @@ namespace UE::Chaos::ClothAsset
 {
 	namespace Private
 	{
-		static const FString OverridePrefix = TEXT("_Override");
-		
 		static const FString ImportFabricBounds = TEXT("ImportFabricBounds");
-		static const FString BuildFabricMaps = TEXT("BuildFabricMaps");
-		static const FString CouldUseFabrics = TEXT("CouldUseFabrics");
-
-		static bool IsOverrideProperty(const TSharedPtr<IPropertyHandle>& Property)
-		{
-			const FStringView PropertyPath = Property ? Property->GetPropertyPath() : FStringView();
-			return PropertyPath.EndsWith(OverridePrefix, ESearchCase::CaseSensitive);
-		}
-		static bool IsOverridePropertyOf(const TSharedPtr<IPropertyHandle>& OverrideProperty, const TSharedPtr<IPropertyHandle>& Property)
-		{
-			const FStringView OverridePropertyPath = OverrideProperty ? OverrideProperty->GetPropertyPath() : FStringView();
-			const FStringView PropertyPath = Property ? Property->GetPropertyPath() : FStringView();
-			return OverridePropertyPath == FString(PropertyPath) + OverridePrefix;
-		}
-		static bool IsImportFabricBoundsProperty(const TSharedPtr<IPropertyHandle>& Property)
+		
+		static bool ImportFabricBoundsProperty(const TSharedPtr<IPropertyHandle>& Property)
 		{
 			const FStringView PropertyPath = Property ? Property->GetPropertyPath() : FStringView();
 			return PropertyPath.EndsWith(ImportFabricBounds, ESearchCase::CaseSensitive);
 		}
-		static bool IsBuildFabricMapsProperty(const TSharedPtr<IPropertyHandle>& Property)
-		{
-			const FStringView PropertyPath = Property ? Property->GetPropertyPath() : FStringView();
-			return PropertyPath.EndsWith(BuildFabricMaps, ESearchCase::CaseSensitive);
-		}
-		static bool CouldUseFabricsProperty(const TSharedPtr<IPropertyHandle>& Property)
-		{
-			const FStringView PropertyPath = Property ? Property->GetPropertyPath() : FStringView();
-			return PropertyPath.EndsWith(CouldUseFabrics, ESearchCase::CaseSensitive);
-		}
-		
-		static void AddToggledCheckBox(const TSharedRef<IPropertyHandle>& PropertyHandle, const TSharedPtr<SHorizontalBox>& HorizontalBox, const FSlateBrush* SlateBrush)
-		{
-			TWeakPtr<IPropertyHandle> WeakHandle = PropertyHandle;
-		
-			HorizontalBox->AddSlot()
-				.VAlign(VAlign_Center)
-				.HAlign(HAlign_Left)
-				.AutoWidth()
-				[
-					SNew(SCheckBox)
-					.Style(&FAppStyle::Get().GetWidgetStyle<FCheckBoxStyle>("ToggleButtonCheckBox"))
-					.Type(ESlateCheckBoxType::ToggleButton)
-					.ToolTipText(WeakHandle.Pin()->GetToolTipText())
-					.IsChecked_Lambda([WeakHandle]()->ECheckBoxState
-						{
-							bool bValue = false;
-							WeakHandle.Pin()->GetValue(bValue);
-							return bValue ? ECheckBoxState::Checked : ECheckBoxState::Unchecked;
-						})
-						.OnCheckStateChanged_Lambda([WeakHandle](ECheckBoxState CheckBoxState)
-						{
-							WeakHandle.Pin()->SetValue(CheckBoxState == ECheckBoxState::Checked, EPropertyValueSetFlags::DefaultFlags);
-						})
-					[
-					SNew(SImage)
-						.ColorAndOpacity(FSlateColor::UseForeground())
-						.Image(SlateBrush)
-					]
-				];
-		}
 	}
 
 	TSharedRef<IPropertyTypeCustomization> FWeightedValueCustomization::MakeInstance()
@@ -116,7 +60,7 @@ namespace UE::Chaos::ClothAsset
 		{
 			TSharedRef<IPropertyHandle> ChildHandle = SortedChildHandles[ChildIndex];
 			
-			if (Private::CouldUseFabricsProperty(ChildHandle))
+			if (CouldUseFabricsProperty(ChildHandle))
 			{
 				bool bValue = false;
 				ChildHandle->GetValue(bValue);
@@ -125,17 +69,18 @@ namespace UE::Chaos::ClothAsset
 					break;
 				}
 			}
-			if (Private::IsImportFabricBoundsProperty(ChildHandle))
+			if (Private::ImportFabricBoundsProperty(ChildHandle))
 			{
-				Private::AddToggledCheckBox(ChildHandle, NameHorizontalBox, FAppStyle::Get().GetBrush("Icons.Import"));
+				AddToggledCheckBox(ChildHandle, NameHorizontalBox, FAppStyle::Get().GetBrush("Icons.Import"));
 			}
-			else if (Private::IsBuildFabricMapsProperty(ChildHandle))
+			else if (BuildFabricMapsProperty(ChildHandle))
 			{
-				Private::AddToggledCheckBox(ChildHandle, NameHorizontalBox, UE::Chaos::ClothAsset::FClothAssetEditorStyle::Get().GetBrush("ClassIcon.ChaosClothPreset"));
+				AddToggledCheckBox(ChildHandle, NameHorizontalBox, UE::Chaos::ClothAsset::FClothAssetEditorStyle::Get().GetBrush("ClassIcon.ChaosClothPreset"));
 			}
 		}
 		
 		NameHorizontalBox->AddSlot().VAlign(VAlign_Center)
+				.Padding(FMargin(4.f, 2.f, 4.0f, 2.f))
 				.HAlign(HAlign_Right)
 				.AutoWidth()
 				[
@@ -146,7 +91,7 @@ namespace UE::Chaos::ClothAsset
 		{
 			TSharedRef<IPropertyHandle> ChildHandle = SortedChildHandles[ChildIndex];
 
-			if (Private::IsOverrideProperty(ChildHandle))
+			if (IsOverrideProperty(ChildHandle))
 			{
 				continue;  // Skip overrides
 			}
@@ -208,7 +153,7 @@ namespace UE::Chaos::ClothAsset
 		if (PropertyClass == FBoolProperty::StaticClass())
 		{
 			TWeakPtr<IPropertyHandle> WeakHandlePtr = PropertyHandle;
-			if (!Private::IsImportFabricBoundsProperty(PropertyHandle) && !Private::IsBuildFabricMapsProperty(PropertyHandle) && !Private::CouldUseFabricsProperty(PropertyHandle))
+			if (!Private::ImportFabricBoundsProperty(PropertyHandle) && !BuildFabricMapsProperty(PropertyHandle) && !CouldUseFabricsProperty(PropertyHandle))
 			{
 				return
 					SNew(SCheckBox)
@@ -236,7 +181,7 @@ namespace UE::Chaos::ClothAsset
 				const bool bLastChild = SortedChildHandles.Num() - 1 == ChildIndex;
 				const TSharedRef<IPropertyHandle>& ChildHandle = SortedChildHandles[ChildIndex];
 
-				if (Private::IsOverridePropertyOf(ChildHandle, PropertyHandle))
+				if (IsOverridePropertyOf(ChildHandle, PropertyHandle))
 				{
 					OverrideHandleWeakPtr = ChildHandle;
 					break;
diff --git a/Engine/Plugins/Experimental/ChaosClothAssetEditor/Source/ChaosClothAssetDataflowNodes/Private/ChaosClothAsset/WeightedValueCustomization.h b/Engine/Plugins/Experimental/ChaosClothAssetEditor/Source/ChaosClothAssetDataflowNodes/Private/ChaosClothAsset/WeightedValueCustomization.h
--- a/Engine/Plugins/Experimental/ChaosClothAssetEditor/Source/ChaosClothAssetDataflowNodes/Private/ChaosClothAsset/WeightedValueCustomization.h
+++ b/Engine/Plugins/Experimental/ChaosClothAssetEditor/Source/ChaosClothAssetDataflowNodes/Private/ChaosClothAsset/WeightedValueCustomization.h
@@ -2,7 +2,7 @@
 
 #pragma once
 
-#include "Customizations/MathStructCustomizations.h"
+#include "ChaosClothAsset/ConnectableValueCustomization.h"
 
 template <typename OptionalType> struct TOptional;
 
@@ -14,7 +14,7 @@ namespace UE::Chaos::ClothAsset
 	 * the structure member short names are displayed in the header's row.
 	 */
 	class FWeightedValueCustomization
-		: public FMathStructCustomization
+		: public FConnectableValueCustomization
 	{
 	public:
 		static TSharedRef<IPropertyTypeCustomization> MakeInstance();
diff --git a/Engine/Plugins/Experimental/ChaosClothAssetEditor/Source/ChaosClothAssetDataflowNodes/Public/ChaosClothAsset/ConnectableValue.h b/Engine/Plugins/Experimental/ChaosClothAssetEditor/Source/ChaosClothAssetDataflowNodes/Public/ChaosClothAsset/ConnectableValue.h
--- a/Engine/Plugins/Experimental/ChaosClothAssetEditor/Source/ChaosClothAssetDataflowNodes/Public/ChaosClothAsset/ConnectableValue.h
+++ b/Engine/Plugins/Experimental/ChaosClothAssetEditor/Source/ChaosClothAssetDataflowNodes/Public/ChaosClothAsset/ConnectableValue.h
@@ -28,9 +28,21 @@ struct FChaosClothAssetConnectableIStringValue
 	UPROPERTY(EditAnywhere, Category = "Value", Meta = (DataflowInput))
 	FString StringValue = TEXT("StringValue");
 
+	/**
+	  * Whether the property could import fabrics datas or not
+	  */
+	UPROPERTY(VisibleAnywhere, Transient, Category = "Value")
+	bool bCouldUseFabrics = false;
+
 	/** The string value override value for when the StringValue has a connection that replaces the provided string value. */
 	UPROPERTY(VisibleAnywhere, Category = "Value", Transient)
 	mutable FString StringValue_Override;  // _Override has a special meaning to the property customization, mutable because this property is set while getting the original value
+	
+	/**
+	 * Whether the property can override the weight map based on the imported fabrics
+	 */
+	UPROPERTY(EditAnywhere, Category = "Value", Meta = (EditCondition = "bCouldUseFabrics", EditConditionHides))
+	bool bBuildFabricMaps = false;
 };
 
 USTRUCT()
diff --git a/Engine/Plugins/Experimental/ChaosClothAssetEditor/Source/ChaosClothAssetDataflowNodes/Public/ChaosClothAsset/ImportedValue.h b/Engine/Plugins/Experimental/ChaosClothAssetEditor/Source/ChaosClothAssetDataflowNodes/Public/ChaosClothAsset/ImportedValue.h
new file mode 100644
--- /dev/null
+++ b/Engine/Plugins/Experimental/ChaosClothAssetEditor/Source/ChaosClothAssetDataflowNodes/Public/ChaosClothAsset/ImportedValue.h
@@ -0,0 +1,66 @@
+// Copyright Epic Games, Inc. All Rights Reserved.
+#pragma once
+
+#include "UObject/ObjectMacros.h"
+#include "ImportedValue.generated.h"
+
+USTRUCT()
+struct FChaosClothAssetImportedVectorValue
+{
+	GENERATED_BODY()
+
+	using ImportedType = FVector3f;
+
+	/**
+	 * Property vector Value
+	*/
+	UPROPERTY(EditAnywhere, Category = "Imported Value", Meta = (EditCondition = "!bUseImportedValue"))
+	mutable FVector3f ImportedValue;
+	
+	/**
+	 * Whether the property can use the values imported from USD
+	 */
+	UPROPERTY(EditAnywhere, Category = "Imported Value")
+	bool bUseImportedValue = false;
+};
+
+USTRUCT()
+struct FChaosClothAssetImportedIntValue
+{
+	GENERATED_BODY()
+
+	using ImportedType = int32;
+
+	/**
+	 * Property integer Value
+	*/
+	UPROPERTY(EditAnywhere, Category = "Imported Value", Meta = (EditCondition = "!bUseImportedValue"))
+	mutable int32 ImportedValue;
+	
+	/**
+	 * Whether the property can use the values imported from USD
+	 */
+	UPROPERTY(EditAnywhere, Category = "Imported Value")
+	bool bUseImportedValue = false;
+};
+
+USTRUCT()
+struct FChaosClothAssetImportedFloatValue
+{
+	GENERATED_BODY()
+
+	using ImportedType = float;
+
+	/**
+	 * Property float value
+	*/
+	UPROPERTY(EditAnywhere, Category = "Imported Value", Meta = (EditCondition = "!bUseImportedValue"))
+	mutable float ImportedValue;
+	
+	/**
+	 * Whether the property can use the values imported from USD
+	 */
+	UPROPERTY(EditAnywhere, Category = "Imported Value")
+	bool bUseImportedValue = false;
+};
+
diff --git a/Engine/Plugins/Experimental/ChaosClothAssetEditor/Source/ChaosClothAssetDataflowNodes/Public/ChaosClothAsset/SimulationAerodynamicsConfigNode.h b/Engine/Plugins/Experimental/ChaosClothAssetEditor/Source/ChaosClothAssetDataflowNodes/Public/ChaosClothAsset/SimulationAerodynamicsConfigNode.h
--- a/Engine/Plugins/Experimental/ChaosClothAssetEditor/Source/ChaosClothAssetDataflowNodes/Public/ChaosClothAsset/SimulationAerodynamicsConfigNode.h
+++ b/Engine/Plugins/Experimental/ChaosClothAssetEditor/Source/ChaosClothAssetDataflowNodes/Public/ChaosClothAsset/SimulationAerodynamicsConfigNode.h
@@ -21,6 +21,13 @@ public:
 	 */
 	UPROPERTY(EditAnywhere, Category = "Aerodynamics", Meta = (UIMin = "0", ClampMin = "0", ClampMax = "10"))
 	float FluidDensity = 1.225f;
+	
+	/**
+	 * The fixed wind velocity [m/s] for this asset.
+	 * For reference a wind gust is above 8m/s (18mph).
+	 */
+	UPROPERTY(EditAnywhere, Category = "Aerodynamics", Meta = (UIMin = "0", UIMax = "10"))
+	FVector3f WindVelocity = { 0.f, 0.f, 0.f };
 
 	/**
 	 * The aerodynamic coefficient of drag applying on each particle.
@@ -29,7 +36,7 @@ public:
 	 * Otherwise all particles are considered to have a zero weight, and only the Low value is meaningful.
 	 */
 	UPROPERTY(EditAnywhere, Category = "Aerodynamics", Meta = (UIMin = "0", UIMax = "1", ClampMin = "0", ClampMax = "10"))
-	FChaosClothAssetWeightedValue Drag = { true, 0.035f, 1.f, TEXT("Drag") };
+	FChaosClothAssetWeightedValue Drag = { true, 0.035f, 1.f, TEXT("Drag"), true };
 
 	/**
 	 * The aerodynamic coefficient of lift applying on each particle.
@@ -38,15 +45,8 @@ public:
 	 * Otherwise all particles are considered to have a zero weight, and only the Low value is meaningful.
 	 */
 	UPROPERTY(EditAnywhere, Category = "Aerodynamics", Meta = (UIMin = "0", UIMax = "1", ClampMin = "0", ClampMax = "10"))
-	FChaosClothAssetWeightedValue Lift = { true, 0.035f, 1.f, TEXT("Lift") };
-
-	/**
-	 * The fixed wind velocity [m/s] for this asset.
-	 * For reference a wind gust is above 8m/s (18mph).
-	 */
-	UPROPERTY(EditAnywhere, Category = "Aerodynamics", Meta = (UIMin = "0", UIMax = "10"))
-	FVector3f WindVelocity = { 0.f, 0.f, 0.f };
-
+	FChaosClothAssetWeightedValue Lift = { true, 0.035f, 1.f, TEXT("Lift"), true };
+	
 	FChaosClothAssetSimulationAerodynamicsConfigNode(const Dataflow::FNodeParameters& InParam, FGuid InGuid = FGuid::NewGuid());
 
 private:
diff --git a/Engine/Plugins/Experimental/ChaosClothAssetEditor/Source/ChaosClothAssetDataflowNodes/Public/ChaosClothAsset/SimulationBaseConfigNode.h b/Engine/Plugins/Experimental/ChaosClothAssetEditor/Source/ChaosClothAssetDataflowNodes/Public/ChaosClothAsset/SimulationBaseConfigNode.h
--- a/Engine/Plugins/Experimental/ChaosClothAssetEditor/Source/ChaosClothAssetDataflowNodes/Public/ChaosClothAsset/SimulationBaseConfigNode.h
+++ b/Engine/Plugins/Experimental/ChaosClothAssetEditor/Source/ChaosClothAssetDataflowNodes/Public/ChaosClothAsset/SimulationBaseConfigNode.h
@@ -3,9 +3,11 @@
 #pragma once 
 
 #include "Chaos/CollectionPropertyFacade.h"
+#include "ChaosClothAsset/CollectionClothFacade.h"
 #include "ChaosClothAsset/ConnectableValue.h"
 #include "ChaosClothAsset/WeightedValue.h"
 #include "ChaosClothAsset/SimulationConfigNodePropertyTypes.h"
+#include "ChaosClothAsset/ImportedValue.h"
 #include "Dataflow/DataflowNode.h"
 #include "GeometryCollection/ManagedArrayCollection.h"
 #include "SimulationBaseConfigNode.generated.h"
@@ -82,6 +84,37 @@ protected:
         {
         	return ClothCollection;
         }
+
+		/** Set an imported solver value onto the property */
+		template<typename PropertyType>
+		void SetSolverProperty(const FName& PropertyName, const PropertyType& PropertyValue, 
+			const TFunction<typename PropertyType::ImportedType(UE::Chaos::ClothAsset::FCollectionClothFacade&)>& SolverValueFunction,
+			const TArray<FName>& SimilarPropertyNames, ECollectionPropertyFlags PropertyFlags = ECollectionPropertyFlags::None);
+
+		/** Set an imported averaged fabric value onto the property */
+		template<typename PropertyType>
+		void SetFabricProperty(const FName& PropertyName, const PropertyType& PropertyValue, 
+			const TFunction<typename PropertyType::ImportedType(UE::Chaos::ClothAsset::FCollectionClothFabricFacade&)>& FabricValueFunction,
+			const TArray<FName>& SimilarPropertyNames, ECollectionPropertyFlags PropertyFlags = ECollectionPropertyFlags::None);
+
+		/** Set an imported solver value onto the weighted value property (animatable or not) */
+		template<typename PropertyType>
+		void SetSolverPropertyWeighted(const FName& PropertyName, const PropertyType& PropertyValue,
+			const TFunction<float(const UE::Chaos::ClothAsset::FCollectionClothFacade&)>& SolverValueFunction,
+			const TArray<FName>& SimilarPropertyNames, ECollectionPropertyFlags PropertyFlags = ECollectionPropertyFlags::None);
+
+		/** Set an imported fabric value onto the weighted value property (animatable or not) */
+		template<typename PropertyType>
+		void SetFabricPropertyWeighted(const FName& PropertyName, const PropertyType& PropertyValue,
+			const TFunction<float(const UE::Chaos::ClothAsset::FCollectionClothFabricFacade&)>& FabricValueFunction,
+			const TArray<FName>& SimilarPropertyNames, ECollectionPropertyFlags PropertyFlags = ECollectionPropertyFlags::None);
+
+		template<typename MapType, typename PropertyType>
+		void SetFabricPropertyString(
+			const FName& PropertyName, const PropertyType& PropertyValue,
+			const TFunction<MapType(const UE::Chaos::ClothAsset::FCollectionClothFabricFacade&)>& FabricValueFunction,
+			const TArray<FName>& SimilarPropertyNames, ECollectionPropertyFlags PropertyFlags = ECollectionPropertyFlags::None,
+			const FName& GroupName = "");
 		
 	private:
 		
diff --git a/Engine/Plugins/Experimental/ChaosClothAssetEditor/Source/ChaosClothAssetDataflowNodes/Public/ChaosClothAsset/SimulationBendingConfigNode.h b/Engine/Plugins/Experimental/ChaosClothAssetEditor/Source/ChaosClothAssetDataflowNodes/Public/ChaosClothAsset/SimulationBendingConfigNode.h
--- a/Engine/Plugins/Experimental/ChaosClothAssetEditor/Source/ChaosClothAssetDataflowNodes/Public/ChaosClothAsset/SimulationBendingConfigNode.h
+++ b/Engine/Plugins/Experimental/ChaosClothAssetEditor/Source/ChaosClothAssetDataflowNodes/Public/ChaosClothAsset/SimulationBendingConfigNode.h
@@ -2,14 +2,14 @@
 
 #pragma once 
 
-#include "ChaosClothAsset/SimulationFabricConfigNode.h"
+#include "ChaosClothAsset/SimulationBaseConfigNode.h"
 #include "ChaosClothAsset/WeightedValue.h"
 #include "ChaosClothAsset/SimulationConfigNodePropertyTypes.h"
 #include "SimulationBendingConfigNode.generated.h"
 
 /** Bending constraint property configuration node. */
 USTRUCT(Meta = (DataflowCloth))
-struct FChaosClothAssetSimulationBendingConfigNode : public FChaosClothAssetSimulationFabricConfigNode
+struct FChaosClothAssetSimulationBendingConfigNode : public FChaosClothAssetSimulationBaseConfigNode
 {
 	GENERATED_USTRUCT_BODY()
 	DATAFLOW_NODE_DEFINE_INTERNAL(FChaosClothAssetSimulationBendingConfigNode, "SimulationBendingConfig", "Cloth", "Cloth Simulation Bending Config")
diff --git a/Engine/Plugins/Experimental/ChaosClothAssetEditor/Source/ChaosClothAssetDataflowNodes/Public/ChaosClothAsset/SimulationCollisionConfigNode.h b/Engine/Plugins/Experimental/ChaosClothAssetEditor/Source/ChaosClothAssetDataflowNodes/Public/ChaosClothAsset/SimulationCollisionConfigNode.h
--- a/Engine/Plugins/Experimental/ChaosClothAssetEditor/Source/ChaosClothAssetDataflowNodes/Public/ChaosClothAsset/SimulationCollisionConfigNode.h
+++ b/Engine/Plugins/Experimental/ChaosClothAssetEditor/Source/ChaosClothAssetDataflowNodes/Public/ChaosClothAsset/SimulationCollisionConfigNode.h
@@ -14,12 +14,12 @@ struct FChaosClothAssetSimulationCollisionConfigNode : public FChaosClothAssetSi
 
 public:
 	/** The added thickness of collision shapes. */
-	UPROPERTY(EditAnywhere, Category = "Collision Properties", Meta = (UIMin = "0", UIMax = "100", ClampMin = "0", ClampMax = "1000"))
-	float CollisionThickness = 1.0f;
+	UPROPERTY(EditAnywhere, Category = "Collision Properties", DisplayName = "Collision Thickness", Meta = (UIMin = "0", UIMax = "100", ClampMin = "0", ClampMax = "1000"))
+	FChaosClothAssetImportedFloatValue CollisionThicknessImported = {UE::Chaos::ClothAsset::FDefaultFabric::CollisionThickness};
 
 	/** Friction coefficient for cloth - collider interaction. */
-	UPROPERTY(EditAnywhere, Category = "Collision Properties", Meta = (UIMin = "0", UIMax = "1", ClampMin = "0", ClampMax = "10"))
-	float FrictionCoefficient = 0.8f;
+	UPROPERTY(EditAnywhere, Category = "Collision Properties", DisplayName = "Friction Coefficient", Meta = (UIMin = "0", UIMax = "1", ClampMin = "0", ClampMax = "10"))
+	FChaosClothAssetImportedFloatValue FrictionCoefficientImported = {UE::Chaos::ClothAsset::FDefaultFabric::Friction};
 
 	/** Stiffness for proximity repulsion forces (Force-based solver only). Units = kg cm/ s^2 (same as XPBD springs)*/
 	UPROPERTY(EditAnywhere, Category = "Proximity Force Properties", Meta = (UIMin = "0", UIMax = "10000", ClampMin = "0", ClampMax = "10000000"))
@@ -34,6 +34,18 @@ public:
 
 	FChaosClothAssetSimulationCollisionConfigNode(const Dataflow::FNodeParameters& InParam, FGuid InGuid = FGuid::NewGuid());
 
+	virtual void Serialize(FArchive& Ar) override;
+
 private:
 	virtual void AddProperties(FPropertyHelper& PropertyHelper) const override;
+
+	// Deprecated properties
+#if WITH_EDITORONLY_DATA
+
+	UPROPERTY()
+	float CollisionThickness_DEPRECATED = UE::Chaos::ClothAsset::FDefaultFabric::CollisionThickness;
+
+	UPROPERTY()
+	float FrictionCoefficient_DEPRECATED  = UE::Chaos::ClothAsset::FDefaultFabric::Friction;
+#endif
 };
diff --git a/Engine/Plugins/Experimental/ChaosClothAssetEditor/Source/ChaosClothAssetDataflowNodes/Public/ChaosClothAsset/SimulationFabricConfigNode.h b/Engine/Plugins/Experimental/ChaosClothAssetEditor/Source/ChaosClothAssetDataflowNodes/Public/ChaosClothAsset/SimulationFabricConfigNode.h
deleted file mode 100644
diff --git a/Engine/Plugins/Experimental/ChaosClothAssetEditor/Source/ChaosClothAssetDataflowNodes/Public/ChaosClothAsset/SimulationGravityConfigNode.h b/Engine/Plugins/Experimental/ChaosClothAssetEditor/Source/ChaosClothAssetDataflowNodes/Public/ChaosClothAsset/SimulationGravityConfigNode.h
--- a/Engine/Plugins/Experimental/ChaosClothAssetEditor/Source/ChaosClothAssetDataflowNodes/Public/ChaosClothAsset/SimulationGravityConfigNode.h
+++ b/Engine/Plugins/Experimental/ChaosClothAssetEditor/Source/ChaosClothAssetDataflowNodes/Public/ChaosClothAsset/SimulationGravityConfigNode.h
@@ -24,18 +24,22 @@ public:
 	FChaosClothAssetWeightedValue GravityScaleWeighted = {true, 1.f, 1.f, TEXT("GravityScale")};
 
 	/** The gravitational acceleration vector [cm/s^2]. */
-	UPROPERTY(EditAnywhere, Category = "Gravity Properties", Meta = (UIMin = "0", UIMax = "10", EditCondition = "bUseGravityOverride"))
-	FVector3f GravityOverride = { 0.f, 0.f, -980.665f };  // TODO: Should we make this a S.I. unit?
-
+	UPROPERTY(EditAnywhere, Category = "Gravity Properties", DisplayName = "Gravity Override", Meta = (UIMin = "0", UIMax = "10", EditCondition = "bUseGravityOverride"))
+	FChaosClothAssetImportedVectorValue GravityOverrideImported = {UE::Chaos::ClothAsset::FDefaultSolver::Gravity};  // TODO: Should we make this a S.I. unit?
+	
 	FChaosClothAssetSimulationGravityConfigNode(const Dataflow::FNodeParameters& InParam, FGuid InGuid = FGuid::NewGuid());
 
 	virtual void Serialize(FArchive& Ar) override;
 
 private:
+
 	// Deprecated properties
 #if WITH_EDITORONLY_DATA
 	UPROPERTY()
 	float GravityScale_DEPRECATED = 1.f;
+
+	UPROPERTY()
+	FVector3f GravityOverride_DEPRECATED = UE::Chaos::ClothAsset::FDefaultSolver::Gravity;  // TODO: Should we make this a S.I. unit?
 #endif
 
 	virtual void AddProperties(FPropertyHelper& PropertyHelper) const override;
diff --git a/Engine/Plugins/Experimental/ChaosClothAssetEditor/Source/ChaosClothAssetDataflowNodes/Public/ChaosClothAsset/SimulationMassConfigNode.h b/Engine/Plugins/Experimental/ChaosClothAssetEditor/Source/ChaosClothAssetDataflowNodes/Public/ChaosClothAsset/SimulationMassConfigNode.h
--- a/Engine/Plugins/Experimental/ChaosClothAssetEditor/Source/ChaosClothAssetDataflowNodes/Public/ChaosClothAsset/SimulationMassConfigNode.h
+++ b/Engine/Plugins/Experimental/ChaosClothAssetEditor/Source/ChaosClothAssetDataflowNodes/Public/ChaosClothAsset/SimulationMassConfigNode.h
@@ -2,13 +2,13 @@
 
 #pragma once 
 
-#include "ChaosClothAsset/SimulationFabricConfigNode.h"
+#include "ChaosClothAsset/SimulationBaseConfigNode.h"
 #include "ClothConfig.h"  // For EClothMassMode
 #include "SimulationMassConfigNode.generated.h"
 
 /** Mass properties configuration node. */
 USTRUCT(Meta = (DataflowCloth))
-struct FChaosClothAssetSimulationMassConfigNode : public FChaosClothAssetSimulationFabricConfigNode
+struct FChaosClothAssetSimulationMassConfigNode : public FChaosClothAssetSimulationBaseConfigNode
 {
 	GENERATED_USTRUCT_BODY()
 	DATAFLOW_NODE_DEFINE_INTERNAL(FChaosClothAssetSimulationMassConfigNode, "SimulationMassConfig", "Cloth", "Cloth Simulation Mass Config")
diff --git a/Engine/Plugins/Experimental/ChaosClothAssetEditor/Source/ChaosClothAssetDataflowNodes/Public/ChaosClothAsset/SimulationPressureConfigNode.h b/Engine/Plugins/Experimental/ChaosClothAssetEditor/Source/ChaosClothAssetDataflowNodes/Public/ChaosClothAsset/SimulationPressureConfigNode.h
--- a/Engine/Plugins/Experimental/ChaosClothAssetEditor/Source/ChaosClothAssetDataflowNodes/Public/ChaosClothAsset/SimulationPressureConfigNode.h
+++ b/Engine/Plugins/Experimental/ChaosClothAssetEditor/Source/ChaosClothAssetDataflowNodes/Public/ChaosClothAsset/SimulationPressureConfigNode.h
@@ -21,7 +21,7 @@ public:
 	 * Otherwise all particles are considered to have a zero weight, and only the Low value is meaningful.
 	 */
 	UPROPERTY(EditAnywhere, Category = "Pressure Properties", meta = (UIMin = "-10", UIMax = "10", ClampMin = "-100", ClampMax = "100"))
-	FChaosClothAssetWeightedValue Pressure = { true, 0.0f, 1.f, TEXT("Pressure") };
+	FChaosClothAssetWeightedValue Pressure = { true, 0.0f, 1.f, TEXT("Pressure"),true };
 
 	FChaosClothAssetSimulationPressureConfigNode(const Dataflow::FNodeParameters& InParam, FGuid InGuid = FGuid::NewGuid());
 
diff --git a/Engine/Plugins/Experimental/ChaosClothAssetEditor/Source/ChaosClothAssetDataflowNodes/Public/ChaosClothAsset/SimulationSelfCollisionConfigNode.h b/Engine/Plugins/Experimental/ChaosClothAssetEditor/Source/ChaosClothAssetDataflowNodes/Public/ChaosClothAsset/SimulationSelfCollisionConfigNode.h
--- a/Engine/Plugins/Experimental/ChaosClothAssetEditor/Source/ChaosClothAssetDataflowNodes/Public/ChaosClothAsset/SimulationSelfCollisionConfigNode.h
+++ b/Engine/Plugins/Experimental/ChaosClothAssetEditor/Source/ChaosClothAssetDataflowNodes/Public/ChaosClothAsset/SimulationSelfCollisionConfigNode.h
@@ -15,16 +15,17 @@ struct FChaosClothAssetSimulationSelfCollisionConfigNode : public FChaosClothAss
 
 public:
 	/** The self collision offset per side. Total thickness of cloth is 2x this value. */
-	UPROPERTY(EditAnywhere, Category = "Self-Collision Properties", meta = (UIMin = "0", UIMax = "100", ClampMin = "0", ClampMax = "1000", EditCondition = "bUseSelfCollisions"))
-	FChaosClothAssetWeightedValue SelfCollisionThicknessWeighted = {true, 0.5f, 0.5f, TEXT("SelfCollisionThickness")};
+	UPROPERTY(EditAnywhere, Category = "Self-Collision Properties", meta = (UIMin = "0", UIMax = "100", ClampMin = "0", ClampMax = "1000"))
+	FChaosClothAssetWeightedValue SelfCollisionThicknessWeighted = {true, UE::Chaos::ClothAsset::FDefaultFabric::SelfCollisionThickness,
+		UE::Chaos::ClothAsset::FDefaultFabric::SelfCollisionThickness, TEXT("SelfCollisionThickness"), true};
 
 	/** The stiffness of the springs used to control self collision (PBD Solver). */
 	UPROPERTY(EditAnywhere, Category = "Self-Collision Properties", meta = (UIMin = "0", UIMax = "1", ClampMin = "0", ClampMax = "1", EditCondition = "!bSelfCollideAgainstKinematicCollidersOnly"))
 	float SelfCollisionStiffness = 0.5f;
 
 	/** Friction coefficient for cloth - cloth interaction. */
-	UPROPERTY(EditAnywhere, Category = "Self-Collision Properties", meta = (UIMin = "0", UIMax = "1", ClampMin = "0", ClampMax = "1", EditCondition = "!bSelfCollideAgainstKinematicCollidersOnly"))
-	float SelfCollisionFriction = 0.0f;
+	UPROPERTY(EditAnywhere, Category = "Self-Collision Properties", DisplayName = "Self Collision Friction", meta = (UIMin = "0", UIMax = "1", ClampMin = "0", ClampMax = "1", EditCondition = "!bSelfCollideAgainstKinematicCollidersOnly"))
+	FChaosClothAssetImportedFloatValue SelfCollisionFrictionImported = {UE::Chaos::ClothAsset::FDefaultFabric::SelfFriction};
 
 	/** Disabled neighbor collision ring. Collisions are disabled between vertices within this N-ring connectivity distance.*/
 	UPROPERTY(EditAnywhere, Category = "Self-Collision Properties", meta = (UIMin = "1", UIMax = "5", ClampMin = "1", EditCondition = "!bSelfCollideAgainstKinematicCollidersOnly"))
@@ -35,7 +36,7 @@ public:
 	* Faces labeled with any other number will keep higher layer numbers outside lower layer numbers (outside = front facing normal direction).
 	*/
 	UPROPERTY(EditAnywhere, Category = "Self-Collision Properties", meta = (EditCondition = "!bSelfCollideAgainstKinematicCollidersOnly"))
-	FChaosClothAssetConnectableIStringValue SelfCollisionLayers = { TEXT("SelfCollisionLayers") };
+	FChaosClothAssetConnectableIStringValue SelfCollisionLayers = { TEXT("SelfCollisionLayers"), true };
 
 	/** Sim face selection set of faces which should not self collide */
 	UPROPERTY(EditAnywhere, Category = "Self-Collision Properties")
@@ -50,7 +51,7 @@ public:
 	FChaosClothAssetConnectableIStringValue SelfCollisionEnabledKinematicFaces = { TEXT("SelfCollisionEnabledKinematicFaces") };
 
 	/** Thickness of kinematic colliders. Total offset between cloth and kinematic colliders is SelfCollisionThickness + SelfCollisionKinematicColliderThickness. */
-	UPROPERTY(EditAnywhere, Category = "Self-Collision Properties - Kinematic Colliders", meta = (UIMin = "0", UIMax = "100", ClampMin = "0", ClampMax = "1000", EditCondition = "bUseSelfCollisions"))
+	UPROPERTY(EditAnywhere, Category = "Self-Collision Properties - Kinematic Colliders", meta = (UIMin = "0", UIMax = "100", ClampMin = "0", ClampMax = "1000"))
 	float SelfCollisionKinematicColliderThickness = 0.f;
 
 	/** The stiffness of the springs used to control self collision (PBD Solver). */
@@ -99,6 +100,9 @@ private:
 	UPROPERTY()
 	float SelfCollisionKinematicColliderFriction_DEPRECATED = FrictionDeprecatedValue;
 
+	UPROPERTY()
+	float SelfCollisionFriction_DEPRECATED = UE::Chaos::ClothAsset::FDefaultFabric::SelfFriction;
+
 	static constexpr float SelfCollisionThicknessDeprecatedValue = -1.f;
 	UPROPERTY()
 	float SelfCollisionThickness_DEPRECATED = SelfCollisionThicknessDeprecatedValue;
diff --git a/Engine/Plugins/Experimental/ChaosClothAssetEditor/Source/ChaosClothAssetDataflowNodes/Public/ChaosClothAsset/SimulationSolverConfigNode.h b/Engine/Plugins/Experimental/ChaosClothAssetEditor/Source/ChaosClothAssetDataflowNodes/Public/ChaosClothAsset/SimulationSolverConfigNode.h
--- a/Engine/Plugins/Experimental/ChaosClothAssetEditor/Source/ChaosClothAssetDataflowNodes/Public/ChaosClothAsset/SimulationSolverConfigNode.h
+++ b/Engine/Plugins/Experimental/ChaosClothAssetEditor/Source/ChaosClothAssetDataflowNodes/Public/ChaosClothAsset/SimulationSolverConfigNode.h
@@ -33,8 +33,8 @@ public:
 	 * The number of solver substeps.
 	 * This will increase the precision of the collision inputs and help with constraint resolutions but will increase the CPU cost.
 	 */
-	UPROPERTY(EditAnywhere, Category = Simulation, meta = (UIMin = "1", UIMax = "10", ClampMin = "1", ClampMax = "100"))
-	int32 NumSubsteps = 1;
+	UPROPERTY(EditAnywhere, Category = Simulation, DisplayName = "Num Substeps", meta = (UIMin = "1", UIMax = "10", ClampMin = "1", ClampMax = "100"))
+	FChaosClothAssetImportedIntValue NumSubstepsImported = {UE::Chaos::ClothAsset::FDefaultSolver::SubSteps};
 
 	/**
 	 * Enable dynamic substepping.
@@ -96,5 +96,12 @@ public:
 	virtual void Serialize(FArchive& Ar) override;
 
 private:
+
+	// Deprecated properties
+#if WITH_EDITORONLY_DATA
+
+	UPROPERTY()
+	int32 NumSubsteps_DEPRECATED = UE::Chaos::ClothAsset::FDefaultSolver::SubSteps;
+#endif
 	virtual void AddProperties(FPropertyHelper& PropertyHelper) const override;
 };
diff --git a/Engine/Plugins/Experimental/ChaosClothAssetEditor/Source/ChaosClothAssetDataflowNodes/Public/ChaosClothAsset/SimulationStretchConfigNode.h b/Engine/Plugins/Experimental/ChaosClothAssetEditor/Source/ChaosClothAssetDataflowNodes/Public/ChaosClothAsset/SimulationStretchConfigNode.h
--- a/Engine/Plugins/Experimental/ChaosClothAssetEditor/Source/ChaosClothAssetDataflowNodes/Public/ChaosClothAsset/SimulationStretchConfigNode.h
+++ b/Engine/Plugins/Experimental/ChaosClothAssetEditor/Source/ChaosClothAssetDataflowNodes/Public/ChaosClothAsset/SimulationStretchConfigNode.h
@@ -2,14 +2,14 @@
 
 #pragma once 
 
-#include "ChaosClothAsset/SimulationFabricConfigNode.h"
+#include "ChaosClothAsset/SimulationBaseConfigNode.h"
 #include "ChaosClothAsset/WeightedValue.h"
 #include "ChaosClothAsset/SimulationConfigNodePropertyTypes.h"
 #include "SimulationStretchConfigNode.generated.h"
 
 /** Stretching constraint property configuration node. */
 USTRUCT(Meta = (DataflowCloth))
-struct FChaosClothAssetSimulationStretchConfigNode : public FChaosClothAssetSimulationFabricConfigNode
+struct FChaosClothAssetSimulationStretchConfigNode : public FChaosClothAssetSimulationBaseConfigNode
 {
 	GENERATED_USTRUCT_BODY()
 	DATAFLOW_NODE_DEFINE_INTERNAL(FChaosClothAssetSimulationStretchConfigNode, "SimulationStretchConfig", "Cloth", "Cloth Simulation Stretching Config")
diff --git a/Engine/Source/Runtime/Experimental/Chaos/Public/Chaos/XPBDAnisotropicSpringConstraints.h b/Engine/Source/Runtime/Experimental/Chaos/Public/Chaos/XPBDAnisotropicSpringConstraints.h
--- a/Engine/Source/Runtime/Experimental/Chaos/Public/Chaos/XPBDAnisotropicSpringConstraints.h
+++ b/Engine/Source/Runtime/Experimental/Chaos/Public/Chaos/XPBDAnisotropicSpringConstraints.h
@@ -22,7 +22,7 @@ class FXPBDAnisotropicEdgeSpringConstraints : public FPBDSpringConstraintsBase
 	typedef FPBDSpringConstraintsBase Base;
 
 public:
-	// Stiffness is in kg cm /s^2
+	// Stiffness is in kg /s^2
 	static constexpr FSolverReal MinStiffness = (FSolverReal)0; // We're not checking against MinStiffness (except when it's constant and == 0)
 	static constexpr FSolverReal MaxStiffness = (FSolverReal)1e9;
 	static constexpr FSolverReal MinDampingRatio = (FSolverReal)0.;
@@ -138,7 +138,7 @@ class FXPBDAnisotropicAxialSpringConstraints : public FPBDAxialSpringConstraints
 	typedef FPBDAxialSpringConstraintsBase Base;
 
 public:
-	// Stiffness is in kg cm /s^2
+	// Stiffness is in kg /s^2
 	static constexpr FSolverReal MinStiffness = (FSolverReal)0; // We're not checking against MinStiffness (except when it's constant and == 0)
 	static constexpr FSolverReal MaxStiffness = (FSolverReal)1e9;
 	static constexpr FSolverReal MinDampingRatio = (FSolverReal)0.;
