Chaos Cloth Asset:

diff --git a/Engine/Plugins/ChaosClothAssetEditor/Source/ChaosClothAssetEditor/Private/ChaosClothAsset/SClothEditorRestSpaceViewportToolBar.cpp b/Engine/Plugins/ChaosClothAssetEditor/Source/ChaosClothAssetEditor/Private/ChaosClothAsset/SClothEditorRestSpaceViewportToolBar.cpp
new file mode 100644
--- /dev/null
+++ b/Engine/Plugins/ChaosClothAssetEditor/Source/ChaosClothAssetEditor/Private/ChaosClothAsset/SClothEditorRestSpaceViewportToolBar.cpp
@@ -0,0 +1,316 @@
+// Copyright Epic Games, Inc. All Rights Reserved.
+
+#include "ChaosClothAsset/SClothEditorRestSpaceViewportToolBar.h"
+#include "ChaosClothAsset/SClothEditorRestSpaceViewport.h"
+#include "ChaosClothAsset/ClothEditorRestSpaceViewportClient.h"
+#include "ChaosClothAsset/ClothEditorCommands.h"
+#include "EditorViewportCommands.h"
+#include "Styling/AppStyle.h"
+#include "Framework/MultiBox/MultiBoxBuilder.h"
+#include "Widgets/Layout/SBorder.h"
+#include "Widgets/SBoxPanel.h"
+#include "Editor/UnrealEd/Public/SEditorViewportToolBarMenu.h"
+#include "SEditorViewportViewMenu.h"
+#include "EditorModeManager.h"
+#include "ChaosClothAsset/ClothEditorMode.h"
+#include "ToolMenus.h"
+#include "Widgets/Input/SSpinBox.h"
+#include "Widgets/Text/SRichTextBlock.h"
+
+#define LOCTEXT_NAMESPACE "SChaosClothAssetEditorRestSpaceViewportToolBar"
+
+void SChaosClothAssetEditorRestSpaceViewportToolBar::Construct(const FArguments& InArgs, TSharedPtr<SChaosClothAssetEditorRestSpaceViewport> InChaosClothAssetEditorViewport)
+{
+	RestSpaceViewportClient = InArgs._RestSpaceViewportClient;
+
+	SCommonEditorViewportToolbarBase::Construct(SCommonEditorViewportToolbarBase::FArguments(), InChaosClothAssetEditorViewport);
+
+	ChaosClothAssetEditorRestSpaceViewportPtr = InChaosClothAssetEditorViewport;
+	CommandList = InArgs._CommandList;
+
+	RegisterViewModeMenuContent();
+
+	const FMargin ToolbarSlotPadding(4.0f, 1.0f);
+	TSharedPtr<SHorizontalBox> MainBoxPtr;
+
+	ChildSlot
+	[
+		SNew( SBorder )
+		.BorderImage(FAppStyle::Get().GetBrush("EditorViewportToolBar.Background"))
+		.Cursor(EMouseCursor::Default)
+		[
+			SNew(SVerticalBox)
+				+ SVerticalBox::Slot()
+				.AutoHeight()
+			[
+				SAssignNew( MainBoxPtr, SHorizontalBox )
+			]
+			+ SVerticalBox::Slot()
+			.Padding(FMargin(4.0f, 3.0f, 0.0f, 0.0f))
+			[
+				// Display text (e.g., item being previewed)
+				SNew(SRichTextBlock)
+					.DecoratorStyleSet(&FAppStyle::Get())
+					.Text(this, &SChaosClothAssetEditorRestSpaceViewportToolBar::GetDisplayString)
+					.TextStyle(&FAppStyle::Get().GetWidgetStyle<FTextBlockStyle>("AnimViewport.MessageText"))
+			]
+		]
+	];
+
+	MainBoxPtr->AddSlot()
+		.Padding(ToolbarSlotPadding)
+		.AutoWidth()
+		[
+			MakeOptionsMenu()
+		];
+
+	MainBoxPtr->AddSlot()
+		.Padding(ToolbarSlotPadding)
+		.HAlign(HAlign_Left)
+		[
+			MakeDisplayToolBar(InArgs._Extenders)
+		];
+
+	MainBoxPtr->AddSlot()
+		.Padding(ToolbarSlotPadding)
+		.HAlign(HAlign_Right)
+		[
+			MakeToolBar(InArgs._Extenders)
+		];
+}
+
+TSharedRef<SWidget> SChaosClothAssetEditorRestSpaceViewportToolBar::GenerateClothRestSpaceViewportOptionsMenu()
+{
+	using namespace UE::Chaos::ClothAsset;
+
+	GetInfoProvider().OnFloatingButtonClicked();
+	TSharedRef<SEditorViewport> ViewportRef = GetInfoProvider().GetViewportWidget();
+
+	const bool bIsPerspective = GetViewportClient().GetViewportType() == LVT_Perspective;
+	const bool bInShouldCloseWindowAfterMenuSelection = true;
+
+	FMenuBuilder OptionsMenuBuilder(bInShouldCloseWindowAfterMenuSelection, ViewportRef->GetCommandList());
+	{
+		OptionsMenuBuilder.BeginSection("ClothEditorViewportOptions", LOCTEXT("OptionsMenuHeader", "Viewport Options"));
+		{
+			if (bIsPerspective)
+			{
+				OptionsMenuBuilder.AddWidget(GenerateFOVMenu(), LOCTEXT("FOVAngle", "Field of View (H)"));
+
+				OptionsMenuBuilder.AddSubMenu(
+					LOCTEXT("CameraSpeedSettings", "Camera Speed Settings"),
+					LOCTEXT("CameraSpeedSettingsToolTip", "Adjust camera speed settings"),
+					FNewMenuDelegate::CreateLambda([this](FMenuBuilder& MenuBuilder)
+				{
+					MenuBuilder.AddWidget(GenerateCameraSpeedSettingsMenu(), FText());
+				}
+				));
+			}
+			OptionsMenuBuilder.AddWidget(GenerateLightMenu(), LOCTEXT("LightIntensity", "Render Light Intensity"));
+
+			OptionsMenuBuilder.AddMenuEntry(FChaosClothAssetEditorCommands::Get().ToggleConstructionViewSeams, 
+				NAME_None, 
+				LOCTEXT("ShowSeamsLabel", "Show Seams"), 
+				LOCTEXT("ShowSeamsTooltip", "Display seam information (not available for non-manifold meshes)"));
+
+			OptionsMenuBuilder.AddMenuEntry(FChaosClothAssetEditorCommands::Get().ToggleConstructionViewSeamsCollapse, 
+				NAME_None, 
+				LOCTEXT("SeamsCollapseLabel", "Collapse Seam Lines"), 
+				LOCTEXT("SeamsCollapseTooltip", "Display a single line connecting each seam, rather than all stitches"));
+
+			OptionsMenuBuilder.AddMenuEntry(FChaosClothAssetEditorCommands::Get().TogglePatternColor,
+				NAME_None,
+				LOCTEXT("ColorPatternsLabel", "Color Patterns"),
+				LOCTEXT("ColorPatternsTooltip", "Display each Pattern in a different color"));
+
+			OptionsMenuBuilder.AddMenuEntry(FChaosClothAssetEditorCommands::Get().ToggleMeshStats,
+				NAME_None,
+				LOCTEXT("ToggleMeshStatsLabel", "Mesh Stats"),
+				LOCTEXT("ToggleMeshStatsTooltip", "Show mesh stats in the viewport"));
+
+		}
+		OptionsMenuBuilder.EndSection();
+
+		ExtendOptionsMenu(OptionsMenuBuilder);
+	}
+
+	return OptionsMenuBuilder.MakeWidget();
+}
+
+TSharedRef<SWidget> SChaosClothAssetEditorRestSpaceViewportToolBar::GenerateLightMenu()
+{
+	constexpr float IntensityMin = 0.f;
+	constexpr float IntensityMax = 20.f;
+	
+	return 
+		SNew(SBox)
+		.HAlign(HAlign_Right)
+		[
+			SNew(SBox)
+			.Padding(FMargin(4.0f, 0.0f, 0.0f, 0.0f))
+			.WidthOverride(100.0f)
+			[
+				SNew(SBorder)
+				.BorderImage(FAppStyle::Get().GetBrush("Menu.WidgetBorder"))
+				.Padding(FMargin(1.0f))
+				[
+					SNew(SSpinBox<float>)
+					.Style(&FAppStyle::Get(), "Menu.SpinBox")
+					.Font(FAppStyle::GetFontStyle(TEXT("MenuItem.Font")))
+					.MinValue(IntensityMin)
+					.MaxValue(IntensityMax)
+					.Value(this, &SChaosClothAssetEditorRestSpaceViewportToolBar::GetCameraPointLightIntensity)
+					.OnValueChanged(this, &SChaosClothAssetEditorRestSpaceViewportToolBar::CameraPointLightIntensityChanged)
+					.IsEnabled(this, &SChaosClothAssetEditorRestSpaceViewportToolBar::IsRenderModeEnabled)
+				]
+			]
+		];
+}
+
+
+float SChaosClothAssetEditorRestSpaceViewportToolBar::GetCameraPointLightIntensity() const
+{
+	return RestSpaceViewportClient->GetCameraPointLightIntensity();
+}
+
+void SChaosClothAssetEditorRestSpaceViewportToolBar::CameraPointLightIntensityChanged(float NewValue)
+{
+	RestSpaceViewportClient->SetCameraPointLightIntensity(NewValue);
+}
+
+bool SChaosClothAssetEditorRestSpaceViewportToolBar::IsRenderModeEnabled() const
+{
+	return RestSpaceViewportClient->GetConstructionViewMode() == UE::Chaos::ClothAsset::EClothPatternVertexType::Render;
+}
+
+TSharedRef<SWidget> SChaosClothAssetEditorRestSpaceViewportToolBar::MakeOptionsMenu()
+{
+	return SNew(SEditorViewportToolbarMenu)
+		.ParentToolBar(SharedThis(this))
+		.Cursor(EMouseCursor::Default)
+		.Image("EditorViewportToolBar.OptionsDropdown")
+		.OnGetMenuContent(this, &SChaosClothAssetEditorRestSpaceViewportToolBar::GenerateClothRestSpaceViewportOptionsMenu);
+}
+
+TSharedRef<SWidget> SChaosClothAssetEditorRestSpaceViewportToolBar::MakeDisplayToolBar(const TSharedPtr<FExtender> InExtenders)
+{
+	const TSharedRef<SEditorViewport> ViewportRef = StaticCastSharedPtr<SEditorViewport>(ChaosClothAssetEditorRestSpaceViewportPtr.Pin()).ToSharedRef();
+
+	return SNew(SEditorViewportViewMenu, ViewportRef, SharedThis(this))
+		.Cursor(EMouseCursor::Default)
+		.MenuExtenders(InExtenders);
+}
+
+TSharedRef<SWidget> SChaosClothAssetEditorRestSpaceViewportToolBar::MakeToolBar(const TSharedPtr<FExtender> InExtenders)
+{
+	using namespace UE::Chaos::ClothAsset;
+
+	// The following is modeled after portions of STransformViewportToolBar, which gets 
+	// used in SCommonEditorViewportToolbarBase.
+
+	// The buttons are hooked up to actual functions via command bindings in SChaosClothAssetEditorRestSpaceViewport::BindCommands(),
+	// and the toolbar gets built in SChaosClothAssetEditorRestSpaceViewport::MakeViewportToolbar().
+
+	FSlimHorizontalToolBarBuilder ToolbarBuilder(CommandList, FMultiBoxCustomization::None, InExtenders);
+
+	// Use a custom style
+	const FName ToolBarStyle = "EditorViewportToolBar";
+	ToolbarBuilder.SetStyle(&FAppStyle::Get(), ToolBarStyle);
+	ToolbarBuilder.SetLabelVisibility(EVisibility::Collapsed);
+
+	ToolbarBuilder.BeginSection("View Controls");
+	ToolbarBuilder.BeginBlockGroup();
+	{
+		ToolbarBuilder.AddToolBarButton(FChaosClothAssetEditorCommands::Get().ToggleConstructionViewWireframe);
+	}
+	ToolbarBuilder.EndBlockGroup();
+
+	ToolbarBuilder.BeginBlockGroup();
+	{
+		// View mode selector (2D/3D/Render)
+		ToolbarBuilder.AddWidget(
+			SNew(SEditorViewportToolbarMenu)
+			.ParentToolBar(SharedThis(this))
+			.Cursor(EMouseCursor::Default)
+			.Label(this, &SChaosClothAssetEditorRestSpaceViewportToolBar::GetViewModeMenuLabel)
+			.LabelIcon(this, &SChaosClothAssetEditorRestSpaceViewportToolBar::GetViewModeMenuLabelIcon)
+			.OnGetMenuContent(this, &SChaosClothAssetEditorRestSpaceViewportToolBar::GenerateViewModeMenuContent)
+		);
+	}
+	ToolbarBuilder.EndBlockGroup();
+	ToolbarBuilder.EndSection(); // View Controls
+
+	return ToolbarBuilder.MakeWidget();
+}
+
+FText SChaosClothAssetEditorRestSpaceViewportToolBar::GetDisplayString() const
+{
+	if (const FEditorModeTools* const EditorModeTools = RestSpaceViewportClient->GetModeTools())
+	{
+		if (UChaosClothAssetEditorMode* const ClothEdMode = Cast<UChaosClothAssetEditorMode>(EditorModeTools->GetActiveScriptableMode(UChaosClothAssetEditorMode::EM_ChaosClothAssetEditorModeId)))
+		{
+			if (ClothEdMode->IsMeshStatsActive())
+			{
+				const int TriangleCount = ClothEdMode->GetConstructionViewTriangleCount();
+				const int VertexCount = ClothEdMode->GetConstructionViewVertexCount();
+				const FText MeshStats = FText::Format(LOCTEXT("RestSpaceMeshStats", "Tris: {0}, Verts: {1}"), TriangleCount, VertexCount);
+				return MeshStats;
+			}
+		}
+	}
+	return FText();
+}
+
+FText SChaosClothAssetEditorRestSpaceViewportToolBar::GetViewModeMenuLabel() const
+{
+	FText Label = LOCTEXT("ConstructionViewMenuTitle_Default", "View");
+
+	TSharedPtr< SChaosClothAssetEditorRestSpaceViewport > PinnedViewport = ChaosClothAssetEditorRestSpaceViewportPtr.Pin();
+	if (PinnedViewport.IsValid())
+	{
+		const TSharedPtr<FEditorViewportClient> ViewportClient = PinnedViewport->GetViewportClient();
+		check(ViewportClient.IsValid());
+		const FEditorModeTools* const EditorModeTools = ViewportClient->GetModeTools();
+		UChaosClothAssetEditorMode* const ClothEdMode = Cast<UChaosClothAssetEditorMode>(EditorModeTools->GetActiveScriptableMode(UChaosClothAssetEditorMode::EM_ChaosClothAssetEditorModeId));
+		if (ClothEdMode)
+		{
+			switch (ClothEdMode->GetConstructionViewMode())
+			{
+			case UE::Chaos::ClothAsset::EClothPatternVertexType::Sim2D:
+				Label = LOCTEXT("ConstructionViewMenuTitle_Sim2D", "2D Sim");
+				break;
+			case UE::Chaos::ClothAsset::EClothPatternVertexType::Sim3D:
+				Label = LOCTEXT("ConstructionViewMenuTitle_Sim3D", "3D Sim");
+				break;
+			case UE::Chaos::ClothAsset::EClothPatternVertexType::Render:
+				Label = LOCTEXT("ConstructionViewMenuTitle_Render", "Render");
+				break;
+			}
+		}
+	}
+
+	return Label;
+}
+
+const FSlateBrush* SChaosClothAssetEditorRestSpaceViewportToolBar::GetViewModeMenuLabelIcon() const
+{
+	return FStyleDefaults::GetNoBrush();
+}
+
+TSharedRef<SWidget> SChaosClothAssetEditorRestSpaceViewportToolBar::GenerateViewModeMenuContent() const
+{
+	return UToolMenus::Get()->GenerateWidget("ChaosClothAssetEditor.ConstructionViewModeMenu", FToolMenuContext(CommandList));
+}
+
+void SChaosClothAssetEditorRestSpaceViewportToolBar::RegisterViewModeMenuContent()
+{
+	UToolMenu* const Menu = UToolMenus::Get()->RegisterMenu("ChaosClothAssetEditor.ConstructionViewModeMenu");
+
+	FToolMenuSection& Section = Menu->FindOrAddSection("ConstructionViewModeMenuSection");
+	Section.AddMenuEntry(UE::Chaos::ClothAsset::FChaosClothAssetEditorCommands::Get().SetConstructionMode2D);
+	Section.AddMenuEntry(UE::Chaos::ClothAsset::FChaosClothAssetEditorCommands::Get().SetConstructionMode3D);
+	Section.AddMenuEntry(UE::Chaos::ClothAsset::FChaosClothAssetEditorCommands::Get().SetConstructionModeRender);
+}
+
+
+#undef LOCTEXT_NAMESPACE
diff --git a/Engine/Plugins/ChaosClothAssetEditor/Source/ChaosClothAssetEditorTools/Public/ChaosClothAsset/ClothEditorContextObject.h b/Engine/Plugins/ChaosClothAssetEditor/Source/ChaosClothAssetEditorTools/Public/ChaosClothAsset/ClothEditorContextObject.h
new file mode 100644
--- /dev/null
+++ b/Engine/Plugins/ChaosClothAssetEditor/Source/ChaosClothAssetEditorTools/Public/ChaosClothAsset/ClothEditorContextObject.h
@@ -0,0 +1,66 @@
+// Copyright Epic Games, Inc. All Rights Reserved.
+
+#pragma once
+
+#include "UObject/Object.h"
+#include "Dataflow/DataflowEdNode.h"
+#include "Dataflow/DataflowGraphEditor.h"
+#include "ClothEditorContextObject.generated.h"
+
+class UDataflow;
+class UEdGraphNode;
+
+namespace UE::Chaos::ClothAsset
+{
+enum class EClothPatternVertexType : uint8;
+}
+struct FManagedArrayCollection;
+
+
+UCLASS()
+class CHAOSCLOTHASSETEDITORTOOLS_API UClothEditorContextObject : public UObject
+{
+	GENERATED_BODY()
+
+public:
+
+	void Init(TWeakPtr<SDataflowGraphEditor> DataflowGraphEditor, UE::Chaos::ClothAsset::EClothPatternVertexType InConstructionViewMode, TWeakPtr<FManagedArrayCollection> SelectedClothCollection, TWeakPtr<FManagedArrayCollection> SelectedInputClothCollection = nullptr);
+
+	/**
+	* Get a single selected node of the specified type. Return nullptr if the specified node is not selected, or if multiple nodes are selected
+	*/
+	template<typename NodeType>
+	NodeType* GetSingleSelectedNodeOfType() const
+	{
+		const TSharedPtr<SDataflowGraphEditor> GraphEditor = DataflowGraphEditor.Pin();
+
+		if (UEdGraphNode* const SingleSelectedNode = GraphEditor->GetSingleSelectedNode())
+		{
+			UDataflowEdNode* const SelectedDataflowEdNode = CastChecked<UDataflowEdNode>(SingleSelectedNode);
+			if (const TSharedPtr<FDataflowNode> DataflowNode = SelectedDataflowEdNode->GetDataflowNode())	// If the user deletes a node from the graph, the UDataflowEdNode might outlast the FDataflowNode
+			{
+				if (NodeType* const NodeTypeNode = DataflowNode->AsType<NodeType>())
+				{
+					return NodeTypeNode;
+				}
+			}
+		}
+
+		return nullptr;
+	}
+
+	void SetClothCollection(UE::Chaos::ClothAsset::EClothPatternVertexType ViewMode, TWeakPtr<FManagedArrayCollection> ClothCollection, TWeakPtr<FManagedArrayCollection> InputClothCollection = nullptr);
+
+	const TWeakPtr<const FManagedArrayCollection> GetSelectedClothCollection() const { return SelectedClothCollection; }
+	const TWeakPtr<const FManagedArrayCollection> GetSelectedInputClothCollection() const { return SelectedInputClothCollection; }
+	UE::Chaos::ClothAsset::EClothPatternVertexType GetConstructionViewMode() const { return ConstructionViewMode; }
+private:
+
+	TWeakPtr<SDataflowGraphEditor> DataflowGraphEditor;
+
+	UE::Chaos::ClothAsset::EClothPatternVertexType ConstructionViewMode;
+	TWeakPtr<const FManagedArrayCollection> SelectedClothCollection;
+	TWeakPtr<const FManagedArrayCollection> SelectedInputClothCollection;
+};
+
+
diff --git a/Engine/Plugins/Enterprise/DatasmithImporter/Source/DatasmithImporter/Private/Utility/DatasmithTextureResize.cpp b/Engine/Plugins/Enterprise/DatasmithImporter/Source/DatasmithImporter/Private/Utility/DatasmithTextureResize.cpp
--- a/Engine/Plugins/Enterprise/DatasmithImporter/Source/DatasmithImporter/Private/Utility/DatasmithTextureResize.cpp
+++ b/Engine/Plugins/Enterprise/DatasmithImporter/Source/DatasmithImporter/Private/Utility/DatasmithTextureResize.cpp
@@ -104,18 +104,46 @@ namespace
 
 		return x;
 	}
-}
 
-class FFreeImageWrapper
-{
-public:
-	static bool IsValid() { return FreeImageDllHandle != nullptr; }
+	class FFreeImageWrapper
+	{
+	public:
+		static bool IsValid() { return FreeImageDllHandle != nullptr; }
 
-	static void FreeImage_Initialise(); // Loads and inits FreeImage on first call
+		static void FreeImage_Initialise(); // Loads and inits FreeImage on first call
 
-private:
-	static void* FreeImageDllHandle; // Lazy init on first use, never release for now
-};
+	private:
+		static void* FreeImageDllHandle; // Lazy init on first use, never release for now
+	};
+
+	void* FFreeImageWrapper::FreeImageDllHandle = nullptr;
+
+	void FFreeImageWrapper::FreeImage_Initialise()
+	{
+		if (FreeImageDllHandle != nullptr)
+		{
+			return;
+		}
+
+		// Push/PopDllDirectory are soooooo not threadsafe!
+		// Must load library in main thread before doing parallel processing
+		check(IsInGameThread());
+
+		if (FreeImageDllHandle == nullptr)
+		{
+			FString FreeImageDir = FPaths::Combine(FPaths::EngineDir(), TEXT("Binaries/ThirdParty/FreeImage"), FPlatformProcess::GetBinariesSubdirectory());
+			FString FreeImageLibDir = FPaths::Combine(FreeImageDir, TEXT(FREEIMAGE_LIB_FILENAME));
+			FPlatformProcess::PushDllDirectory(*FreeImageDir);
+			FreeImageDllHandle = FPlatformProcess::GetDllHandle(*FreeImageLibDir);
+			FPlatformProcess::PopDllDirectory(*FreeImageDir);
+		}
+
+		if (FreeImageDllHandle)
+		{
+			::FreeImage_Initialise();
+		}
+	}
+}
 
 class DatasmithTextureResizeInternal
 {
@@ -663,34 +691,6 @@ void DatasmithTextureResizeInternal::GetBitmapPixelInfo(FIBITMAP* Bitmap, int32&
 	}
 }
 
-void* FFreeImageWrapper::FreeImageDllHandle = nullptr;
-
-void FFreeImageWrapper::FreeImage_Initialise()
-{
-	if ( FreeImageDllHandle != nullptr )
-	{
-		return;
-	}
-
-	// Push/PopDllDirectory are soooooo not threadsafe!
-	// Must load library in main thread before doing parallel processing
-	check(IsInGameThread());
-
-	if ( FreeImageDllHandle == nullptr )
-	{
-		FString FreeImageDir = FPaths::Combine( FPaths::EngineDir(), TEXT("Binaries/ThirdParty/FreeImage"), FPlatformProcess::GetBinariesSubdirectory() );
-		FString FreeImageLibDir = FPaths::Combine( FreeImageDir, TEXT( FREEIMAGE_LIB_FILENAME ));
-		FPlatformProcess::PushDllDirectory( *FreeImageDir );
-		FreeImageDllHandle = FPlatformProcess::GetDllHandle( *FreeImageLibDir );
-		FPlatformProcess::PopDllDirectory( *FreeImageDir );
-	}
-
-	if ( FreeImageDllHandle )
-	{
-		::FreeImage_Initialise();
-	}
-}
-
 #if PLATFORM_WINDOWS
 #include "Windows/HideWindowsPlatformTypes.h"
 #endif
