Move cloth editor tool headers to private. Rename ClothEditorToolBuilder to ClothEditorToolBuilders and centralize ToolBuilder classes for the cloth editor

diff --git a/Engine/Plugins/ChaosClothAssetEditor/Source/ChaosClothAssetEditor/Private/ChaosClothAsset/ClothEditorCommands.cpp b/Engine/Plugins/ChaosClothAssetEditor/Source/ChaosClothAssetEditor/Private/ChaosClothAsset/ClothEditorCommands.cpp
--- a/Engine/Plugins/ChaosClothAssetEditor/Source/ChaosClothAssetEditor/Private/ChaosClothAsset/ClothEditorCommands.cpp
+++ b/Engine/Plugins/ChaosClothAssetEditor/Source/ChaosClothAssetEditor/Private/ChaosClothAsset/ClothEditorCommands.cpp
@@ -2,8 +2,7 @@
 
 #include "ChaosClothAsset/ClothEditorCommands.h"
 #include "ChaosClothAsset/ClothEditorStyle.h"
-#include "ChaosClothAsset/ClothWeightMapPaintTool.h"
-#include "ChaosClothAsset/ClothMeshSelectionTool.h"
+#include "ChaosClothAsset/ClothEditorToolBuilders.h"
 
 #define LOCTEXT_NAMESPACE "FChaosClothAssetEditorCommands"
 
@@ -74,7 +73,7 @@ void FChaosClothAssetEditorCommands::RegisterCommands()
 
 void FChaosClothAssetEditorCommands::GetToolDefaultObjectList(TArray<UInteractiveTool*>& ToolCDOs)
 {
-	ToolCDOs.Add(GetMutableDefault<UClothEditorWeightMapPaintTool>());
+	GetClothEditorToolDefaultObjectList(ToolCDOs);
 }
 
 void FChaosClothAssetEditorCommands::UpdateToolCommandBinding(UInteractiveTool* Tool, TSharedPtr<FUICommandList> UICommandList, bool bUnbind)
diff --git a/Engine/Plugins/ChaosClothAssetEditor/Source/ChaosClothAssetEditor/Private/ChaosClothAsset/ClothEditorMode.cpp b/Engine/Plugins/ChaosClothAssetEditor/Source/ChaosClothAssetEditor/Private/ChaosClothAsset/ClothEditorMode.cpp
--- a/Engine/Plugins/ChaosClothAssetEditor/Source/ChaosClothAssetEditor/Private/ChaosClothAsset/ClothEditorMode.cpp
+++ b/Engine/Plugins/ChaosClothAssetEditor/Source/ChaosClothAssetEditor/Private/ChaosClothAsset/ClothEditorMode.cpp
@@ -10,6 +10,7 @@
 #include "ChaosClothAsset/ClothPatternToDynamicMesh.h"
 #include "ChaosClothAsset/ClothEditorPreviewScene.h"
 #include "ChaosClothAsset/ClothEditorContextObject.h"
+#include "ChaosClothAsset/ClothEditorToolBuilders.h"
 #include "ChaosClothAsset/AddWeightMapNode.h"
 #include "ChaosClothAsset/TransferSkinWeightsNode.h"
 #include "ChaosClothAsset/SelectionNode.h"
@@ -36,9 +37,6 @@
 #include "RemeshMeshTool.h"
 #include "AttributeEditorTool.h"
 #include "MeshAttributePaintTool.h"
-#include "ChaosClothAsset/ClothWeightMapPaintTool.h"
-#include "ChaosClothAsset/ClothTransferSkinWeightsTool.h"
-#include "ChaosClothAsset/ClothMeshSelectionTool.h"
 #include "DynamicMesh/DynamicMeshAttributeSet.h"
 #include "DynamicMesh/DynamicVertexSkinWeightsAttribute.h"
 #include "DynamicMeshEditor.h"
@@ -63,9 +61,9 @@
 #include "MaterialDomain.h"
 #include "Dataflow/DataflowSNode.h"
 #include "Components/SkeletalMeshComponent.h"
-#include "ChaosClothAsset/ClothEditorToolBuilder.h"
 #include "ChaosClothAsset/ClothGeometryTools.h"
 #include "DynamicMesh/NonManifoldMappingSupport.h"
+#include "DynamicMesh/MeshNormals.h"
 
 #include UE_INLINE_GENERATED_CPP_BY_NAME(ClothEditorMode)
 
diff --git a/Engine/Plugins/ChaosClothAssetEditor/Source/ChaosClothAssetEditorTools/Private/ChaosClothAsset/ClothEditorToolBuilders.cpp b/Engine/Plugins/ChaosClothAssetEditor/Source/ChaosClothAssetEditorTools/Private/ChaosClothAsset/ClothEditorToolBuilders.cpp
new file mode 100644
--- /dev/null
+++ b/Engine/Plugins/ChaosClothAssetEditor/Source/ChaosClothAssetEditorTools/Private/ChaosClothAsset/ClothEditorToolBuilders.cpp
@@ -0,0 +1,105 @@
+// Copyright Epic Games, Inc. All Rights Reserved.
+
+#include "ChaosClothAsset/ClothEditorToolBuilders.h"
+#include "ChaosClothAsset/ClothPatternVertexType.h"
+#include "TargetInterfaces/PrimitiveComponentBackedTarget.h"
+#include "ChaosClothAsset/ClothEditorContextObject.h"
+#include "ToolContextInterfaces.h"
+#include "ToolTargetManager.h"
+#include "ContextObjectStore.h"
+
+// Tools
+#include "ClothMeshSelectionTool.h"
+#include "ClothTransferSkinWeightsTool.h"
+#include "ClothWeightMapPaintTool.h"
+
+
+// ------------------- Weight Map Paint Tool -------------------
+
+void UClothEditorWeightMapPaintToolBuilder::GetSupportedViewModes(TArray<UE::Chaos::ClothAsset::EClothPatternVertexType>& Modes) const
+{
+	Modes.Add(UE::Chaos::ClothAsset::EClothPatternVertexType::Sim3D);
+	Modes.Add(UE::Chaos::ClothAsset::EClothPatternVertexType::Sim2D);
+	Modes.Add(UE::Chaos::ClothAsset::EClothPatternVertexType::Render);
+}
+
+UMeshSurfacePointTool* UClothEditorWeightMapPaintToolBuilder::CreateNewTool(const FToolBuilderState& SceneState) const
+{
+	UClothEditorWeightMapPaintTool* PaintTool = NewObject<UClothEditorWeightMapPaintTool>(SceneState.ToolManager);
+	PaintTool->SetWorld(SceneState.World);
+
+	if (UClothEditorContextObject* ContextObject = SceneState.ToolManager->GetContextObjectStore()->FindContext<UClothEditorContextObject>())
+	{
+		PaintTool->SetClothEditorContextObject(ContextObject);
+	}
+
+	return PaintTool;
+}
+
+// ------------------- Selection Tool -------------------
+
+void UClothMeshSelectionToolBuilder::GetSupportedViewModes(TArray<UE::Chaos::ClothAsset::EClothPatternVertexType>& Modes) const
+{
+	Modes.Add(UE::Chaos::ClothAsset::EClothPatternVertexType::Sim3D);
+	Modes.Add(UE::Chaos::ClothAsset::EClothPatternVertexType::Sim2D);
+	Modes.Add(UE::Chaos::ClothAsset::EClothPatternVertexType::Render);
+}
+
+const FToolTargetTypeRequirements& UClothMeshSelectionToolBuilder::GetTargetRequirements() const
+{
+	static FToolTargetTypeRequirements TypeRequirements(UPrimitiveComponentBackedTarget::StaticClass());
+	return TypeRequirements;
+}
+
+bool UClothMeshSelectionToolBuilder::CanBuildTool(const FToolBuilderState& SceneState) const
+{
+	return (SceneState.TargetManager->CountSelectedAndTargetable(SceneState, GetTargetRequirements()) == 1);
+}
+
+UInteractiveTool* UClothMeshSelectionToolBuilder::BuildTool(const FToolBuilderState& SceneState) const
+{
+	UClothMeshSelectionTool* const NewTool = NewObject<UClothMeshSelectionTool>(SceneState.ToolManager);
+
+	UToolTarget* const Target = SceneState.TargetManager->BuildFirstSelectedTargetable(SceneState, GetTargetRequirements());
+	NewTool->SetTarget(Target);
+	NewTool->SetWorld(SceneState.World);
+
+	if (UClothEditorContextObject* const ContextObject = SceneState.ToolManager->GetContextObjectStore()->FindContext<UClothEditorContextObject>())
+	{
+		NewTool->SetClothEditorContextObject(ContextObject);
+	}
+
+	return NewTool;
+}
+
+
+// ------------------- Skin Weight Transfer Tool -------------------
+
+void UClothTransferSkinWeightsToolBuilder::GetSupportedViewModes(TArray<UE::Chaos::ClothAsset::EClothPatternVertexType>& Modes) const
+{
+	Modes.Add(UE::Chaos::ClothAsset::EClothPatternVertexType::Sim3D);
+}
+
+USingleSelectionMeshEditingTool* UClothTransferSkinWeightsToolBuilder::CreateNewTool(const FToolBuilderState& SceneState) const
+{
+	UClothTransferSkinWeightsTool* NewTool = NewObject<UClothTransferSkinWeightsTool>(SceneState.ToolManager);
+
+	if (UClothEditorContextObject* ContextObject = SceneState.ToolManager->GetContextObjectStore()->FindContext<UClothEditorContextObject>())
+	{
+		NewTool->SetClothEditorContextObject(ContextObject);
+	}
+
+	return NewTool;
+}
+
+
+
+namespace UE::Chaos::ClothAsset
+{
+	void GetClothEditorToolDefaultObjectList(TArray<UInteractiveTool*>& ToolCDOs)
+	{
+		ToolCDOs.Add(GetMutableDefault<UClothEditorWeightMapPaintTool>());
+		ToolCDOs.Add(GetMutableDefault<UClothTransferSkinWeightsTool>());
+		ToolCDOs.Add(GetMutableDefault<UClothMeshSelectionTool>());
+	}
+}
diff --git a/Engine/Plugins/ChaosClothAssetEditor/Source/ChaosClothAssetEditorTools/Private/ChaosClothAsset/ClothMeshSelectionTool.cpp b/Engine/Plugins/ChaosClothAssetEditor/Source/ChaosClothAssetEditorTools/Private/ChaosClothAsset/ClothMeshSelectionTool.cpp
--- a/Engine/Plugins/ChaosClothAssetEditor/Source/ChaosClothAssetEditorTools/Private/ChaosClothAsset/ClothMeshSelectionTool.cpp
+++ b/Engine/Plugins/ChaosClothAssetEditor/Source/ChaosClothAssetEditorTools/Private/ChaosClothAsset/ClothMeshSelectionTool.cpp
@@ -10,8 +10,6 @@
 #include "ChaosClothAsset/WeightedValue.h"
 #include "InteractiveToolManager.h"
 #include "PreviewMesh.h"
-#include "TargetInterfaces/PrimitiveComponentBackedTarget.h"
-#include "ToolTargetManager.h"
 #include "Selection/PolygonSelectionMechanic.h"
 #include "ModelingToolTargetUtil.h"
 #include "GroupTopology.h"
@@ -30,43 +28,6 @@
 #define LOCTEXT_NAMESPACE "ClothMeshSelectionTool"
 
 
-// ------------------- Builder -------------------
-
-void UClothMeshSelectionToolBuilder::GetSupportedViewModes(TArray<UE::Chaos::ClothAsset::EClothPatternVertexType>& Modes) const
-{
-	Modes.Add(UE::Chaos::ClothAsset::EClothPatternVertexType::Sim3D);
-	Modes.Add(UE::Chaos::ClothAsset::EClothPatternVertexType::Sim2D);
-	Modes.Add(UE::Chaos::ClothAsset::EClothPatternVertexType::Render);
-}
-
-const FToolTargetTypeRequirements& UClothMeshSelectionToolBuilder::GetTargetRequirements() const
-{
-	static FToolTargetTypeRequirements TypeRequirements(
-		UPrimitiveComponentBackedTarget::StaticClass());
-	return TypeRequirements;
-}
-
-bool UClothMeshSelectionToolBuilder::CanBuildTool(const FToolBuilderState& SceneState) const
-{
-	return (SceneState.TargetManager->CountSelectedAndTargetable(SceneState, GetTargetRequirements()) == 1);
-}
-
-UInteractiveTool* UClothMeshSelectionToolBuilder::BuildTool(const FToolBuilderState& SceneState) const
-{
-	UClothMeshSelectionTool* const NewTool = NewObject<UClothMeshSelectionTool>(SceneState.ToolManager);
-
-	UToolTarget* const Target = SceneState.TargetManager->BuildFirstSelectedTargetable(SceneState, GetTargetRequirements());
-	NewTool->SetTarget(Target);
-	NewTool->SetWorld(SceneState.World);
-
-	if (UClothEditorContextObject* const ContextObject = SceneState.ToolManager->GetContextObjectStore()->FindContext<UClothEditorContextObject>())
-	{
-		NewTool->SetClothEditorContextObject(ContextObject);
-	}
-
-	return NewTool;
-}
-
 // ------------------- Actions ----------------------
 void UClothMeshSelectionToolActions::PostAction(EClothMeshSelectionToolActions Action)
 {
diff --git a/Engine/Plugins/ChaosClothAssetEditor/Source/ChaosClothAssetEditorTools/Private/ChaosClothAsset/ClothMeshSelectionTool.h b/Engine/Plugins/ChaosClothAssetEditor/Source/ChaosClothAssetEditorTools/Private/ChaosClothAsset/ClothMeshSelectionTool.h
new file mode 100644
--- /dev/null
+++ b/Engine/Plugins/ChaosClothAssetEditor/Source/ChaosClothAssetEditorTools/Private/ChaosClothAsset/ClothMeshSelectionTool.h
@@ -0,0 +1,195 @@
+// Copyright Epic Games, Inc. All Rights Reserved.
+
+#pragma once
+
+#include "BaseTools/SingleSelectionMeshEditingTool.h"
+#include "Selection/PolygonSelectionMechanic.h"
+#include "ClothMeshSelectionTool.generated.h"
+
+class UPolygonSelectionMechanic;
+class UClothEditorContextObject;
+class UPreviewMesh;
+
+namespace UE::Geometry
+{
+	class FGroupTopology;
+	struct FGroupTopologySelection;
+}
+
+UENUM()
+enum class EClothMeshSelectionToolActions
+{
+	NoAction,
+
+	ImportFromCollection,
+	ImportSecondaryFromCollection,
+	TogglePrimarySecondary,
+
+	GrowSelection,
+	ShrinkSelection,
+	FloodSelection
+};
+
+
+UCLASS()
+class UClothMeshSelectionMechanic : public UPolygonSelectionMechanic
+{
+	GENERATED_BODY()
+
+private:
+
+	virtual bool UpdateSelection(const FRay& WorldRay, FVector3d& LocalHitPositionOut, FVector3d& LocalHitNormalOut) override;
+};
+
+
+UCLASS()
+class CHAOSCLOTHASSETEDITORTOOLS_API UClothMeshSelectionToolActions :  public UInteractiveToolPropertySet
+{
+	GENERATED_BODY()
+
+public:
+
+	TWeakObjectPtr<UClothMeshSelectionTool> ParentTool;
+
+	void Initialize(UClothMeshSelectionTool* ParentToolIn) { ParentTool = ParentToolIn; }
+
+	void PostAction(EClothMeshSelectionToolActions Action);
+
+	UFUNCTION(CallInEditor, Category = Operations)
+	void ImportFromCollection()
+	{
+		PostAction(EClothMeshSelectionToolActions::ImportFromCollection);
+	}
+
+	UFUNCTION(CallInEditor, Category = Operations)
+	void ImportSecondaryFromCollection()
+	{
+		PostAction(EClothMeshSelectionToolActions::ImportSecondaryFromCollection);
+	}
+
+	UFUNCTION(CallInEditor, Category = Operations)
+	void TogglePrimarySecondary()
+	{
+		PostAction(EClothMeshSelectionToolActions::TogglePrimarySecondary);
+	}
+
+	UFUNCTION(CallInEditor, Category = Selection)
+	void GrowSelection()
+	{
+		PostAction(EClothMeshSelectionToolActions::GrowSelection);
+	}
+
+	UFUNCTION(CallInEditor, Category = Selection)
+	void ShrinkSelection()
+	{
+		PostAction(EClothMeshSelectionToolActions::ShrinkSelection);
+	}
+
+	UFUNCTION(CallInEditor, Category = Selection)
+	void FloodSelection()
+	{
+		PostAction(EClothMeshSelectionToolActions::FloodSelection);
+	}
+
+};
+
+UCLASS()
+class CHAOSCLOTHASSETEDITORTOOLS_API UClothMeshSelectionToolProperties : public UInteractiveToolPropertySet
+{
+	GENERATED_BODY()
+
+public:
+
+	UPROPERTY(EditAnywhere, Transient, Category = Name, meta = (DisplayName = "Name", TransientToolProperty))
+	FString Name;
+
+	/**
+	* Whether the user is editing the primary or secondary selection set. Stored as a property here so that the tool will remember what the
+	* user was doing the last time the tool shut down.
+	*/
+	UPROPERTY()
+	bool bSecondarySelection = false;
+
+	UPROPERTY(EditAnywhere, Category = Visualization, meta = (DisplayName = "Show Vertices"))
+	bool bShowVertices = false;
+	
+	UPROPERTY(EditAnywhere, Category = Visualization, meta = (DisplayName = "Show Edges"))
+	bool bShowEdges = false;
+
+private:
+
+	virtual void PostEditChangeProperty(FPropertyChangedEvent& PropertyChangedEvent) override;
+
+};
+
+UCLASS()
+class CHAOSCLOTHASSETEDITORTOOLS_API UClothMeshSelectionTool : public USingleSelectionMeshEditingTool
+{
+	GENERATED_BODY()
+
+private:
+
+	friend class UClothMeshSelectionToolBuilder;
+
+	virtual void Setup() override;
+	virtual void OnShutdown(EToolShutdownType ShutdownType) override;
+	virtual void Render(IToolsContextRenderAPI* RenderAPI) override;
+	virtual void DrawHUD(FCanvas* Canvas, IToolsContextRenderAPI* RenderAPI) override;
+
+	virtual void OnTick(float DeltaTime) override;
+
+	virtual bool HasCancel() const override { return true; }
+	virtual bool HasAccept() const override { return true; }
+	virtual bool CanAccept() const override;
+
+	// IInteractiveToolCameraFocusAPI implementation
+	virtual FBox GetWorldSpaceFocusBox() override;
+
+	void SetClothEditorContextObject(TObjectPtr<UClothEditorContextObject> InClothEditorContextObject);
+	bool GetSelectedNodeInfo(FString& OutMapName, UE::Geometry::FGroupTopologySelection& OutSelection);
+	void UpdateSelectedNode();
+
+	UPROPERTY()
+	TObjectPtr<UClothMeshSelectionToolProperties> ToolProperties;
+
+	UPROPERTY()
+	TObjectPtr<UPreviewMesh> PreviewMesh = nullptr;
+
+	UPROPERTY()
+	TObjectPtr<UPolygonSelectionMechanic> SelectionMechanic = nullptr;
+
+	UPROPERTY()
+	TObjectPtr<UClothEditorContextObject> ClothEditorContextObject = nullptr;
+
+	TUniquePtr<UE::Geometry::FGroupTopology> Topology;
+
+	bool bAnyChangeMade = false;
+
+	bool bHasNonManifoldMapping = false;
+	TArray<int32> DynamicMeshToSelection;
+	TArray<TArray<int32>> SelectionToDynamicMesh;
+
+	//
+	// Action support
+	//
+
+public:
+	virtual void RequestAction(EClothMeshSelectionToolActions ActionType);
+
+	UPROPERTY()
+	TObjectPtr<UClothMeshSelectionToolActions> ActionsProps;
+
+private:
+	bool bHavePendingAction = false;
+	EClothMeshSelectionToolActions PendingAction;
+	virtual void ApplyAction(EClothMeshSelectionToolActions ActionType);
+
+	void ImportFromCollection(bool bImportFromSecondarySet);
+	void TogglePrimarySecondaryAction();
+	void UpdatePrimarySecondaryMessage();
+
+	void GrowSelection();
+	void ShrinkSelection();
+	void FloodSelection();
+};
+
diff --git a/Engine/Plugins/ChaosClothAssetEditor/Source/ChaosClothAssetEditorTools/Private/ChaosClothAsset/ClothTransferSkinWeightsTool.cpp b/Engine/Plugins/ChaosClothAssetEditor/Source/ChaosClothAssetEditorTools/Private/ChaosClothAsset/ClothTransferSkinWeightsTool.cpp
--- a/Engine/Plugins/ChaosClothAssetEditor/Source/ChaosClothAssetEditorTools/Private/ChaosClothAsset/ClothTransferSkinWeightsTool.cpp
+++ b/Engine/Plugins/ChaosClothAssetEditor/Source/ChaosClothAssetEditorTools/Private/ChaosClothAsset/ClothTransferSkinWeightsTool.cpp
@@ -42,25 +42,6 @@ namespace UE::Chaos::ClothAsset::Private
 	};
 }
 
-// ------------------- Builder -------------------
-
-void UClothTransferSkinWeightsToolBuilder::GetSupportedViewModes(TArray<UE::Chaos::ClothAsset::EClothPatternVertexType>& Modes) const
-{
-	Modes.Add(UE::Chaos::ClothAsset::EClothPatternVertexType::Sim3D);
-}
-
-USingleSelectionMeshEditingTool* UClothTransferSkinWeightsToolBuilder::CreateNewTool(const FToolBuilderState& SceneState) const
-{
-	UClothTransferSkinWeightsTool* NewTool = NewObject<UClothTransferSkinWeightsTool>(SceneState.ToolManager);
-
-	if (UClothEditorContextObject* ContextObject = SceneState.ToolManager->GetContextObjectStore()->FindContext<UClothEditorContextObject>())
-	{
-		NewTool->SetClothEditorContextObject(ContextObject);
-	}
-
-	return NewTool;
-}
-
 
 // ------------------- Tool -------------------
 
diff --git a/Engine/Plugins/ChaosClothAssetEditor/Source/ChaosClothAssetEditorTools/Private/ChaosClothAsset/ClothTransferSkinWeightsTool.h b/Engine/Plugins/ChaosClothAssetEditor/Source/ChaosClothAssetEditorTools/Private/ChaosClothAsset/ClothTransferSkinWeightsTool.h
new file mode 100644
--- /dev/null
+++ b/Engine/Plugins/ChaosClothAssetEditor/Source/ChaosClothAssetEditorTools/Private/ChaosClothAsset/ClothTransferSkinWeightsTool.h
@@ -0,0 +1,105 @@
+// Copyright Epic Games, Inc. All Rights Reserved.
+
+#pragma once
+
+#include "BaseTools/SingleSelectionMeshEditingTool.h"
+#include "ModelingOperators.h"
+#include "Transforms/TransformGizmoDataBinder.h"
+#include "ClothTransferSkinWeightsTool.generated.h"
+
+
+class UClothTransferSkinWeightsTool;
+class USkeletalMesh;
+class UClothEditorContextObject;
+class UTransformProxy;
+class UCombinedTransformGizmo;
+class UMeshOpPreviewWithBackgroundCompute;
+class AInternalToolFrameworkActor;
+class USkeletalMeshComponent;
+class FTransformGizmoDataBinder;
+struct FChaosClothAssetTransferSkinWeightsNode;
+
+UCLASS()
+class CHAOSCLOTHASSETEDITORTOOLS_API UClothTransferSkinWeightsToolProperties : public UInteractiveToolPropertySet
+{
+	GENERATED_BODY()
+public:
+
+	UPROPERTY(EditAnywhere, Category = Source)
+	TObjectPtr<USkeletalMesh> SourceMesh;
+
+	UPROPERTY(EditAnywhere, Category = "Transform", meta = (DisplayName = "Location", EditCondition = "SourceMesh != nullptr"))
+	FVector3d SourceMeshTranslation = FVector3d::ZeroVector;
+
+	UPROPERTY(EditAnywhere, Category = "Transform", meta = (DisplayName = "Rotation", EditCondition = "SourceMesh != nullptr"))
+	FVector3d SourceMeshRotation = FVector3d::ZeroVector;
+
+	UPROPERTY(EditAnywhere, Category = "Transform", meta = (DisplayName = "Scale", AllowPreserveRatio, EditCondition = "SourceMesh != nullptr"))
+	FVector3d SourceMeshScale = FVector3d::OneVector;
+
+	UPROPERTY(EditAnywhere, Category = Source)
+	bool bHideSourceMesh = false;
+};
+
+UCLASS()
+class CHAOSCLOTHASSETEDITORTOOLS_API UClothTransferSkinWeightsTool : public USingleSelectionMeshEditingTool, public UE::Geometry::IDynamicMeshOperatorFactory
+{
+	GENERATED_BODY()
+
+private:
+
+	friend class UClothTransferSkinWeightsToolBuilder;
+
+	// UInteractiveTool
+	virtual void Setup() override;
+	virtual void Shutdown(EToolShutdownType ShutdownType) override;
+	virtual bool HasAccept() const override { return true; }
+	virtual bool HasCancel() const override { return true; }
+	virtual bool CanAccept() const override;
+	virtual void OnTick(float DeltaTime) override;
+
+	// IDynamicMeshOperatorFactory
+	virtual TUniquePtr<UE::Geometry::FDynamicMeshOperator> MakeNewOperator() override;
+	
+	void SetClothEditorContextObject(TObjectPtr<UClothEditorContextObject> InClothEditorContextObject);
+
+	FTransform TransformFromProperties() const;
+	void SetSRTPropertiesFromTransform(const FTransform& Transform) const;
+
+	void UpdateSourceMesh(TObjectPtr<USkeletalMesh> Mesh);
+
+	void OpFinishedCallback(const UE::Geometry::FDynamicMeshOperator* Op);
+
+	void PreviewMeshUpdatedCallback(UMeshOpPreviewWithBackgroundCompute* Preview);
+
+
+	UPROPERTY(Transient)
+	TObjectPtr<UClothTransferSkinWeightsToolProperties> ToolProperties;
+
+	UPROPERTY(Transient)
+	TObjectPtr<UClothEditorContextObject> ClothEditorContextObject;
+
+	UPROPERTY(Transient)
+	TObjectPtr<UMeshOpPreviewWithBackgroundCompute> TargetClothPreview;
+
+	UPROPERTY(Transient)
+	TObjectPtr<AInternalToolFrameworkActor> SourceMeshParentActor;
+
+	UPROPERTY(Transient)
+	TObjectPtr<USkeletalMeshComponent> SourceMeshComponent;
+
+	// Source mesh transform gizmo support
+	UPROPERTY(Transient)
+	TObjectPtr<UTransformProxy> SourceMeshTransformProxy;
+
+	UPROPERTY(Transient)
+	TObjectPtr<UCombinedTransformGizmo> SourceMeshTransformGizmo;
+
+	TSharedPtr<FTransformGizmoDataBinder> DataBinder;
+
+	FChaosClothAssetTransferSkinWeightsNode* TransferSkinWeightsNode = nullptr;
+
+	bool bHasOpFailedWarning = false;
+};
+
+
diff --git a/Engine/Plugins/ChaosClothAssetEditor/Source/ChaosClothAssetEditorTools/Private/ChaosClothAsset/ClothWeightMapPaintTool.cpp b/Engine/Plugins/ChaosClothAssetEditor/Source/ChaosClothAssetEditorTools/Private/ChaosClothAsset/ClothWeightMapPaintTool.cpp
--- a/Engine/Plugins/ChaosClothAssetEditor/Source/ChaosClothAssetEditorTools/Private/ChaosClothAsset/ClothWeightMapPaintTool.cpp
+++ b/Engine/Plugins/ChaosClothAssetEditor/Source/ChaosClothAssetEditorTools/Private/ChaosClothAsset/ClothWeightMapPaintTool.cpp
@@ -65,31 +65,6 @@ namespace UE::Chaos::ClothAsset::Private
 }
 
 
-/*
- * ToolBuilder
- */
-
-void UClothEditorWeightMapPaintToolBuilder::GetSupportedViewModes(TArray<UE::Chaos::ClothAsset::EClothPatternVertexType>& Modes) const
-{
-	Modes.Add(UE::Chaos::ClothAsset::EClothPatternVertexType::Sim3D);
-	Modes.Add(UE::Chaos::ClothAsset::EClothPatternVertexType::Sim2D);
-	Modes.Add(UE::Chaos::ClothAsset::EClothPatternVertexType::Render);
-}
-
-UMeshSurfacePointTool* UClothEditorWeightMapPaintToolBuilder::CreateNewTool(const FToolBuilderState& SceneState) const
-{
-	UClothEditorWeightMapPaintTool* PaintTool = NewObject<UClothEditorWeightMapPaintTool>(SceneState.ToolManager);
-	PaintTool->SetWorld(SceneState.World);
-
-	if (UClothEditorContextObject* ContextObject = SceneState.ToolManager->GetContextObjectStore()->FindContext<UClothEditorContextObject>())
-	{
-		PaintTool->SetClothEditorContextObject(ContextObject);
-	}
-
-	return PaintTool;
-}
-
-
 /*
  * ToolActions
  */
diff --git a/Engine/Plugins/ChaosClothAssetEditor/Source/ChaosClothAssetEditorTools/Private/ChaosClothAsset/ClothWeightMapPaintTool.h b/Engine/Plugins/ChaosClothAssetEditor/Source/ChaosClothAssetEditorTools/Private/ChaosClothAsset/ClothWeightMapPaintTool.h
new file mode 100644
--- /dev/null
+++ b/Engine/Plugins/ChaosClothAssetEditor/Source/ChaosClothAssetEditorTools/Private/ChaosClothAsset/ClothWeightMapPaintTool.h
@@ -0,0 +1,561 @@
+// Copyright Epic Games, Inc. All Rights Reserved.
+
+#pragma once
+
+#include "CoreMinimal.h"
+#include "HAL/Platform.h"
+#include "UObject/NoExportTypes.h"
+#include "BaseTools/BaseBrushTool.h"
+#include "BaseTools/MeshSurfacePointMeshEditingTool.h"
+#include "Components/DynamicMeshComponent.h"
+#include "PropertySets/WeightMapSetProperties.h"
+#include "Mechanics/PolyLassoMarqueeMechanic.h"
+
+#include "Sculpting/MeshSculptToolBase.h"
+#include "Sculpting/MeshBrushOpBase.h"
+
+#include "DynamicMesh/DynamicMeshAABBTree3.h"
+#include "DynamicMesh/DynamicMeshOctree3.h"
+#include "DynamicMesh/DynamicMeshAttributeSet.h"
+#include "DynamicMesh/DynamicMeshChangeTracker.h"
+#include "DynamicMesh/MeshNormals.h"
+#include "TransformTypes.h"
+#include "ToolDataVisualizer.h"
+#include "GroupTopology.h"
+#include "ClothWeightMapPaintTool.generated.h"
+
+class UMeshElementsVisualizer;
+class UWeightMapEraseBrushOpProps;
+class UWeightMapPaintBrushOpProps;
+class UWeightMapSmoothBrushOpProps;
+class UClothEditorContextObject;
+class UPolygonSelectionMechanic;
+struct FChaosClothAssetAddWeightMapNode;
+
+DECLARE_STATS_GROUP(TEXT("WeightMapPaintTool"), STATGROUP_WeightMapPaintTool, STATCAT_Advanced);
+DECLARE_CYCLE_STAT(TEXT("WeightMapPaintTool_UpdateROI"), WeightMapPaintTool_UpdateROI, STATGROUP_WeightMapPaintTool);
+DECLARE_CYCLE_STAT(TEXT("WeightMapPaintTool_ApplyStamp"), WeightMapPaintToolApplyStamp, STATGROUP_WeightMapPaintTool );
+DECLARE_CYCLE_STAT(TEXT("WeightMapPaintTool_Tick"), WeightMapPaintToolTick, STATGROUP_WeightMapPaintTool);
+DECLARE_CYCLE_STAT(TEXT("WeightMapPaintTool_Tick_ApplyStampBlock"), WeightMapPaintTool_Tick_ApplyStampBlock, STATGROUP_WeightMapPaintTool);
+DECLARE_CYCLE_STAT(TEXT("WeightMapPaintTool_Tick_ApplyStamp_Remove"), WeightMapPaintTool_Tick_ApplyStamp_Remove, STATGROUP_WeightMapPaintTool);
+DECLARE_CYCLE_STAT(TEXT("WeightMapPaintTool_Tick_ApplyStamp_Insert"), WeightMapPaintTool_Tick_ApplyStamp_Insert, STATGROUP_WeightMapPaintTool);
+DECLARE_CYCLE_STAT(TEXT("WeightMapPaintTool_Tick_NormalsBlock"), WeightMapPaintTool_Tick_NormalsBlock, STATGROUP_WeightMapPaintTool);
+DECLARE_CYCLE_STAT(TEXT("WeightMapPaintTool_Tick_UpdateMeshBlock"), WeightMapPaintTool_Tick_UpdateMeshBlock, STATGROUP_WeightMapPaintTool);
+DECLARE_CYCLE_STAT(TEXT("WeightMapPaintTool_Tick_UpdateTargetBlock"), WeightMapPaintTool_Tick_UpdateTargetBlock, STATGROUP_WeightMapPaintTool);
+DECLARE_CYCLE_STAT(TEXT("WeightMapPaintTool_Normals_Collect"), WeightMapPaintTool_Normals_Collect, STATGROUP_WeightMapPaintTool);
+DECLARE_CYCLE_STAT(TEXT("WeightMapPaintTool_Normals_Compute"), WeightMapPaintTool_Normals_Compute, STATGROUP_WeightMapPaintTool);
+
+
+UENUM()
+enum class EClothEditorWeightMapDisplayType : uint8
+{
+	BlackAndWhite,
+	WhiteAndRed,
+	LastValue UMETA(Hidden)
+};
+
+/** Mesh Sculpting Brush Types */
+UENUM()
+enum class EClothEditorWeightMapPaintInteractionType : uint8
+{
+	Brush,
+	Fill,
+	PolyLasso,
+	Gradient,
+	HideTriangles,
+	LastValue UMETA(Hidden)
+};
+
+
+
+
+
+
+
+/** Mesh Sculpting Brush Types */
+UENUM()
+enum class EClothEditorWeightMapPaintBrushType : uint8
+{
+	/** Paint weights */
+	Paint UMETA(DisplayName = "Paint"),
+
+	/** Smooth existing weights */
+	Smooth UMETA(DisplayName = "Smooth"),
+
+	/** Erase weights */
+	Erase UMETA(Hidden, DisplayName = "Erase"),
+
+	LastValue UMETA(Hidden)
+};
+
+
+/** Mesh Sculpting Brush Area Types */
+UENUM()
+enum class EClothEditorWeightMapPaintBrushAreaType : uint8
+{
+	Connected,
+	Volumetric
+};
+
+/** Mesh Sculpting Brush Types */
+UENUM()
+enum class EClothEditorWeightMapPaintVisibilityType : uint8
+{
+	None,
+	Unoccluded
+};
+
+
+
+// TODO: Look at EditConditions for all these properties. Which ones make sense for which SubToolType?
+
+UCLASS()
+class CHAOSCLOTHASSETEDITORTOOLS_API UClothEditorWeightMapPaintBrushFilterProperties : public UInteractiveToolPropertySet
+{
+	GENERATED_BODY()
+
+public:
+
+	UPROPERTY(EditAnywhere, Category = Display)
+	EClothEditorWeightMapDisplayType ColorMap = EClothEditorWeightMapDisplayType::BlackAndWhite;
+
+	UPROPERTY(EditAnywhere, Category = ActionType, meta = (DisplayName = "Action"))
+	EClothEditorWeightMapPaintInteractionType SubToolType = EClothEditorWeightMapPaintInteractionType::Brush;
+
+	UPROPERTY(EditAnywhere, Category = ActionType, meta = (DisplayName = "Brush Mode", 
+		HideEditConditionToggle, EditConditionHides, EditCondition = "SubToolType == EClothEditorWeightMapPaintInteractionType::Brush"))
+	EClothEditorWeightMapPaintBrushType PrimaryBrushType = EClothEditorWeightMapPaintBrushType::Paint;
+
+	/** Relative size of brush */
+	UPROPERTY(EditAnywhere, Category = Brush, meta = (DisplayName = "Brush Size", UIMin = "0.0", UIMax = "1.0", ClampMin = "0.0", ClampMax = "10.0", 
+		HideEditConditionToggle, EditConditionHides, EditCondition = "SubToolType == EClothEditorWeightMapPaintInteractionType::Brush"))
+	float BrushSize = 0.25f;
+
+	/** Allow the Brush to hit the back-side of the mesh */
+	UPROPERTY(EditAnywhere, Category = Brush, meta = (HideEditConditionToggle, EditConditionHides, EditCondition = "SubToolType == EClothEditorWeightMapPaintInteractionType::Brush"))
+	bool bHitBackFaces = true;
+
+	/** The new value to paint on the mesh */
+	UPROPERTY(EditAnywhere, Category = Brush, meta = (UIMin = 0, ClampMin = 0, UIMax = 1, ClampMax = 1,
+		HideEditConditionToggle, EditConditionHides, EditCondition = 
+		"(SubToolType == EClothEditorWeightMapPaintInteractionType::Brush && PrimaryBrushType == EClothEditorWeightMapPaintBrushType::Paint) || SubToolType == EClothEditorWeightMapPaintInteractionType::Fill || SubToolType == EClothEditorWeightMapPaintInteractionType::PolyLasso"))
+	double AttributeValue = 1;
+
+	/** How quickly each brush stroke will drive mesh values towards the desired value */
+	UPROPERTY(EditAnywhere, Category = Brush, meta = (UIMin = 0, ClampMin = 0, UIMax = 1, ClampMax = 1,
+		HideEditConditionToggle, EditConditionHides, EditCondition = 
+		"SubToolType == EClothEditorWeightMapPaintInteractionType::Brush || SubToolType == EClothEditorWeightMapPaintInteractionType::Fill"))
+	double Strength = 1.0;
+
+	/** The Gradient upper limit value */
+	UPROPERTY(EditAnywhere, Category = Gradient, meta = (UIMin = 0, ClampMin = 0, UIMax = 1, ClampMax = 1,
+		HideEditConditionToggle, EditConditionHides, EditCondition = "SubToolType == EClothEditorWeightMapPaintInteractionType::Gradient"))
+	double GradientHighValue = 1.0;
+
+	/** The Gradient lower limit value */
+	UPROPERTY(EditAnywhere, Category = Gradient, meta = (UIMin = 0, ClampMin = 0, UIMax = 1, ClampMax = 1,
+		HideEditConditionToggle, EditConditionHides, EditCondition = "SubToolType == EClothEditorWeightMapPaintInteractionType::Gradient"))
+	double GradientLowValue = 0.0;
+
+
+	/** The Region affected by the current operation will be bounded by edge angles larger than this threshold */
+	UPROPERTY(EditAnywhere, Category = Filters, meta = (UIMin = "0.0", ClampMin = "0.0", UIMax = "180.0", ClampMax = "180.0",
+		HideEditConditionToggle, EditConditionHides, EditCondition = 
+		"SubToolType == EClothEditorWeightMapPaintInteractionType::Brush || SubToolType == EClothEditorWeightMapPaintInteractionType::Fill"))
+	float AngleThreshold = 180.0f;
+
+	/** The Region affected by the current operation will be bounded by UV borders/seams */
+	UPROPERTY(EditAnywhere, Category = Filters, meta = (HideEditConditionToggle, EditConditionHides, EditCondition = 
+		"SubToolType == EClothEditorWeightMapPaintInteractionType::Brush || SubToolType == EClothEditorWeightMapPaintInteractionType::Fill"))
+	bool bUVSeams = false;
+
+	/** The Region affected by the current operation will be bounded by Hard Normal edges/seams */
+	UPROPERTY(EditAnywhere, Category = Filters, meta = (HideEditConditionToggle, EditConditionHides, EditCondition = 
+		"SubToolType == EClothEditorWeightMapPaintInteractionType::Brush || SubToolType == EClothEditorWeightMapPaintInteractionType::Fill"))
+	bool bNormalSeams = false;
+
+	/** Control which triangles can be affected by the current operation based on visibility. Applied after all other filters. */
+	UPROPERTY(EditAnywhere, Category = Filters, meta = (HideEditConditionToggle, EditConditionHides, EditCondition =
+		"SubToolType == EClothEditorWeightMapPaintInteractionType::Brush || SubToolType == EClothEditorWeightMapPaintInteractionType::Fill || SubToolType == EClothEditorWeightMapPaintInteractionType::PolyLasso"))
+	EClothEditorWeightMapPaintVisibilityType VisibilityFilter = EClothEditorWeightMapPaintVisibilityType::None;
+
+	/** The weight value at the brush indicator */
+	UPROPERTY(VisibleAnywhere, Transient, Category = Query, meta = (NoResetToDefault, 
+		HideEditConditionToggle, EditConditionHides, EditCondition = "SubToolType == EClothEditorWeightMapPaintInteractionType::Brush || SubToolType == EClothEditorWeightMapPaintInteractionType::Fill"))
+	double ValueAtBrush = 0;
+
+};
+
+
+
+
+
+UENUM()
+enum class EClothEditorWeightMapPaintToolActions
+{
+	NoAction,
+
+	FloodFillCurrent,
+	ClearAll,
+	Invert,
+	Multiply,
+	ClearHiddenTriangles,
+};
+
+
+UCLASS()
+class CHAOSCLOTHASSETEDITORTOOLS_API UClothEditorMeshWeightMapPaintToolActions : public UInteractiveToolPropertySet
+{
+	GENERATED_BODY()
+
+public:
+
+	TWeakObjectPtr<UClothEditorWeightMapPaintTool> ParentTool;
+
+	void Initialize(UClothEditorWeightMapPaintTool* ParentToolIn) { ParentTool = ParentToolIn; }
+
+	void PostAction(EClothEditorWeightMapPaintToolActions Action);
+
+	UFUNCTION(CallInEditor, Category = Operations, meta = (DisplayPriority = 10))
+	void ClearAll()
+	{
+		PostAction(EClothEditorWeightMapPaintToolActions::ClearAll);
+	}
+
+	UFUNCTION(CallInEditor, Category = Operations, meta = (DisplayPriority = 12))
+	void FloodFillCurrent()
+	{
+		PostAction(EClothEditorWeightMapPaintToolActions::FloodFillCurrent);
+	}
+
+	UFUNCTION(CallInEditor, Category = Operations, meta = (DisplayPriority = 13))
+	void Invert()
+	{
+		PostAction(EClothEditorWeightMapPaintToolActions::Invert);
+	}
+
+	UFUNCTION(CallInEditor, Category = Operations, meta = (DisplayPriority = 14))
+	void Multiply()
+	{
+		PostAction(EClothEditorWeightMapPaintToolActions::Multiply);
+	}
+
+};
+
+UCLASS()
+class CHAOSCLOTHASSETEDITORTOOLS_API UClothEditorUpdateWeightMapProperties : public UInteractiveToolPropertySet
+{
+	GENERATED_BODY()
+
+public:
+
+	UPROPERTY(EditAnywhere, Category = UpdateNode, meta = (DisplayName = "Name"))
+	FString Name;
+
+private:
+
+	virtual void PostEditChangeProperty(struct FPropertyChangedEvent& PropertyChangedEvent) override;
+};
+
+/// Property panel for controlling hiding triangles
+UCLASS()
+class CHAOSCLOTHASSETEDITORTOOLS_API UClothEditorMeshWeightMapPaintToolShowHideProperties : public UInteractiveToolPropertySet
+{
+	GENERATED_BODY()
+
+public:
+
+	TWeakObjectPtr<UClothEditorWeightMapPaintTool> ParentTool;
+	void Initialize(UClothEditorWeightMapPaintTool* ParentToolIn) { ParentTool = ParentToolIn; }
+	void PostAction(EClothEditorWeightMapPaintToolActions Action);
+
+	/** Toggles whether each pattern is shown or hidden */
+	UPROPERTY(EditAnywhere, EditFixedSize, Transient, Category = TriangleVisibility, meta = (DisplayPriority = 2));
+	TMap<int32, bool> ShowPatterns;
+
+	/** Unhide all triangles */
+	UFUNCTION(CallInEditor, Category = TriangleVisibility, meta = (DisplayPriority = 1))
+	void ShowAll()
+	{
+		PostAction(EClothEditorWeightMapPaintToolActions::ClearHiddenTriangles);
+	}
+};
+
+
+/**
+ * Mesh Element Paint Tool Class
+ */
+UCLASS()
+class CHAOSCLOTHASSETEDITORTOOLS_API UClothEditorWeightMapPaintTool : public UMeshSculptToolBase
+{
+	GENERATED_BODY()
+
+public:
+	virtual void RegisterActions(FInteractiveToolActionSet& ActionSet) override;
+
+	virtual void Setup() override;
+	virtual void Shutdown(EToolShutdownType ShutdownType) override;
+	virtual void DrawHUD(FCanvas* Canvas, IToolsContextRenderAPI* RenderAPI) override;
+	virtual void Render(IToolsContextRenderAPI* RenderAPI) override;
+	virtual void UpdateMaterialMode(EMeshEditingMaterialModes MaterialMode) override;
+
+	virtual void UpdateStampPendingState() override;
+
+	virtual void OnTick(float DeltaTime) override;
+
+	virtual bool HasCancel() const override { return true; }
+	virtual bool HasAccept() const override { return true; }
+	virtual bool CanAccept() const override;
+
+	virtual bool OnUpdateHover(const FInputDeviceRay& DevicePos) override;
+
+	virtual void OnPropertyModified(UObject* PropertySet, FProperty* Property) override;
+
+	bool IsInBrushSubMode() const;
+
+	virtual void CommitResult(UBaseDynamicMeshComponent* Component, bool bModifiedTopology) override;
+
+	void SetClothEditorContextObject(TObjectPtr<UClothEditorContextObject> InClothEditorContextObject);
+
+public:
+
+
+	/** Filters on paint brush */
+	UPROPERTY()
+	TObjectPtr<UClothEditorWeightMapPaintBrushFilterProperties> FilterProperties;
+
+
+private:
+	UPROPERTY()
+	TObjectPtr<UWeightMapPaintBrushOpProps> PaintBrushOpProperties;
+
+	UPROPERTY()
+	TObjectPtr<UWeightMapSmoothBrushOpProps> SmoothBrushOpProperties;
+
+	UPROPERTY()
+	TObjectPtr<UWeightMapEraseBrushOpProps> EraseBrushOpProperties;
+
+	void FloodFillCurrentWeightAction();
+	void ClearAllWeightsAction();
+	void InvertWeightsAction();
+	void MultiplyWeightsAction();
+	void ClearHiddenAction();
+
+public:
+	void SetVerticesToWeightMap(const TSet<int32>& Vertices, double WeightValue, bool bIsErase);
+
+	bool HaveVisibilityFilter() const;
+	void ApplyVisibilityFilter(const TArray<int32>& Vertices, TArray<int32>& VisibleVertices);
+	void ApplyVisibilityFilter(TSet<int32>& Triangles, TArray<int32>& ROIBuffer, TArray<int32>& OutputBuffer);
+
+	// we override these so we can update the separate BrushSize property added for this tool
+	virtual void IncreaseBrushRadiusAction();
+	virtual void DecreaseBrushRadiusAction();
+	virtual void IncreaseBrushRadiusSmallStepAction();
+	virtual void DecreaseBrushRadiusSmallStepAction();
+
+protected:
+	// UMeshSculptToolBase API
+	virtual UBaseDynamicMeshComponent* GetSculptMeshComponent() { return DynamicMeshComponent; }
+	virtual FDynamicMesh3* GetBaseMesh() { check(false); return nullptr; }
+	virtual const FDynamicMesh3* GetBaseMesh() const { check(false); return nullptr; }
+
+	virtual int32 FindHitSculptMeshTriangle(const FRay3d& LocalRay) override;
+	virtual int32 FindHitTargetMeshTriangle(const FRay3d& LocalRay) override;
+
+	virtual void OnBeginStroke(const FRay& WorldRay) override;
+	virtual void OnEndStroke() override;
+
+	virtual TUniquePtr<FMeshSculptBrushOp>& GetActiveBrushOp();
+
+	virtual bool SharesBrushPropertiesChanges() const override { return false; }
+
+	virtual void InitializeBrushSizeRange(const UE::Geometry::FAxisAlignedBox3d& TargetBounds) override;
+
+	virtual void NextBrushModeAction() override;
+	virtual void PreviousBrushModeAction() override;
+
+	// Note: these will actually modify the brush's Attribute Value since we don't use Speed in our brush
+	virtual void IncreaseBrushSpeedAction() override;
+	virtual void DecreaseBrushSpeedAction() override;
+
+
+	// end UMeshSculptToolBase API
+
+
+
+	//
+	// Action support
+	//
+
+public:
+	virtual void RequestAction(EClothEditorWeightMapPaintToolActions ActionType);
+	
+	UPROPERTY()
+	TObjectPtr<UClothEditorMeshWeightMapPaintToolActions> ActionsProps;
+
+	UPROPERTY()
+	TObjectPtr<UClothEditorUpdateWeightMapProperties> UpdateWeightMapProperties;
+
+protected:
+	bool bHavePendingAction = false;
+	EClothEditorWeightMapPaintToolActions PendingAction;
+	virtual void ApplyAction(EClothEditorWeightMapPaintToolActions ActionType);
+
+
+
+	//
+	// Marquee Support
+	//
+public:
+	UPROPERTY()
+	TObjectPtr<UPolyLassoMarqueeMechanic> PolyLassoMechanic;
+
+protected:
+	void OnPolyLassoFinished(const FCameraPolyLasso& Lasso, bool bCanceled);
+
+
+	// 
+	// Gradient Support
+	//
+	UPROPERTY()
+	TObjectPtr<UPolygonSelectionMechanic> PolygonSelectionMechanic;
+
+	TUniquePtr<UE::Geometry::FDynamicMeshAABBTree3> MeshSpatial = nullptr;
+
+	TUniquePtr<UE::Geometry::FTriangleGroupTopology> GradientSelectionTopology = nullptr;
+
+	FToolDataVisualizer GradientSelectionRenderer;
+
+	UE::Geometry::FGroupTopologySelection LowValueGradientVertexSelection;
+	UE::Geometry::FGroupTopologySelection HighValueGradientVertexSelection;
+
+	void ComputeGradient();
+	void OnSelectionModified();
+
+
+	//
+	// Show/Hide support
+	//
+protected:
+	UPROPERTY()
+	TObjectPtr<UClothEditorMeshWeightMapPaintToolShowHideProperties> ShowHideProperties;
+
+	// Hidden triangles are not rendered, their wireframes are not rendered, and they don't block ray casts from the mouse
+	TSet<int32> HiddenTriangles;
+
+	// Pending hidden triangles are triangles that are selected by the current HideTriangles brush stroke. When the stroke finishes, these pending triangles are
+	// added to the set of hidden triangles, and this set is cleared.
+	// 
+	// Pending hidden triangles are not rendered, but their wireframes *are* rendered, and they *do* block ray casts.
+	// - Wireframe update is delayed because its a relatively expensive operation, so we do it only when the stroke finishes
+	// - Ray cast blocking is delayed because otherwise one small mouse drag would drill through multiple mesh layers
+	TSet<int32> PendingHiddenTriangles;
+
+	// Triangle index offset and number of triangles for each non-empty pattern. This is used when we want to hide/show entire patterns.
+	TArray<TPair<int32, int32>> PatternTriangleOffsetAndNum;
+
+	//
+	// Internals
+	//
+
+protected:
+
+	UPROPERTY()
+	TObjectPtr<AInternalToolFrameworkActor> PreviewMeshActor = nullptr;
+
+	UPROPERTY()
+	TObjectPtr<UDynamicMeshComponent> DynamicMeshComponent;
+
+	UPROPERTY()
+	TObjectPtr<UMeshElementsVisualizer> MeshElementsDisplay;
+
+	UPROPERTY()
+	TObjectPtr<UClothEditorContextObject> ClothEditorContextObject = nullptr;
+
+	// realtime visualization
+	void OnDynamicMeshComponentChanged(UDynamicMeshComponent* Component, const FMeshVertexChange* Change, bool bRevert);
+	FDelegateHandle OnDynamicMeshComponentChangedHandle;
+
+	UE::Geometry::FDynamicMeshWeightAttribute* ActiveWeightMap;
+	double GetCurrentWeightValue(int32 VertexId) const;
+	double GetCurrentWeightValueUnderBrush() const;
+	FVector3d CurrentBaryCentricCoords;
+	int32 GetBrushNearestVertex() const;
+
+	void GetCurrentWeightMap(TArray<float>& OutWeights) const;
+
+	void UpdateSubToolType(EClothEditorWeightMapPaintInteractionType NewType);
+
+	void UpdateBrushType(EClothEditorWeightMapPaintBrushType BrushType);
+
+	TSet<int32> AccumulatedTriangleROI;
+	bool bUndoUpdatePending = false;
+	TArray<int> NormalsBuffer;
+	void WaitForPendingUndoRedo();
+
+	TArray<int> TempROIBuffer;
+	TArray<int> VertexROI;
+	TArray<bool> VisibilityFilterBuffer;
+	TSet<int> VertexSetBuffer;
+	TSet<int> TriangleROI;
+	void UpdateROI(const FSculptBrushStamp& CurrentStamp);
+
+	EClothEditorWeightMapPaintBrushType PendingStampType = EClothEditorWeightMapPaintBrushType::Paint;
+
+	bool UpdateStampPosition(const FRay& WorldRay);
+	bool ApplyStamp();
+
+	UE::Geometry::FDynamicMeshOctree3 Octree;
+
+	bool UpdateBrushPosition(const FRay& WorldRay);
+
+	bool GetInEraseStroke()
+	{
+		// Re-use the smoothing stroke key (shift) for erase stroke in the weight paint tool
+		return GetInSmoothingStroke();
+	}
+
+
+	bool bPendingPickWeight = false;
+
+	TArray<int32> ROITriangleBuffer;
+	TArray<double> ROIWeightValueBuffer;
+	bool SyncMeshWithWeightBuffer(FDynamicMesh3* Mesh);
+	bool SyncWeightBufferWithMesh(const FDynamicMesh3* Mesh);
+
+	TUniquePtr<UE::Geometry::FDynamicMeshChangeTracker> ActiveWeightEditChangeTracker;
+	void BeginChange();
+	void EndChange();
+
+	FColor GetColorForWeightValue(double WeightValue);
+
+	TArray<FVector3d> TriNormals;
+	TArray<int32> UVSeamEdges;
+	TArray<int32> NormalSeamEdges;
+	void PrecomputeFilterData();
+
+	// DynamicMesh might be unwelded mesh, but weights are on the welded mesh.
+	bool bHaveDynamicMeshToWeightConversion = false;
+	TArray<int32> DynamicMeshToWeight; 
+	TArray<TArray<int32>> WeightToDynamicMesh;
+
+protected:
+	virtual bool ShowWorkPlane() const override { return false; }
+
+	friend class UClothEditorWeightMapPaintToolBuilder;
+
+	bool bAnyChangeMade = false;
+
+	// Node graph editor support
+
+	FChaosClothAssetAddWeightMapNode* WeightMapNodeToUpdate = nullptr;
+
+	void UpdateSelectedNode();
+
+	void UpdateVertexColorOverlay(const TSet<int>* TrianglesToUpdate = nullptr);
+
+};
+
+
+
diff --git a/Engine/Plugins/ChaosClothAssetEditor/Source/ChaosClothAssetEditorTools/Public/ChaosClothAsset/ClothEditorToolBuilder.h b/Engine/Plugins/ChaosClothAssetEditor/Source/ChaosClothAssetEditorTools/Public/ChaosClothAsset/ClothEditorToolBuilder.h
deleted file mode 100644
diff --git a/Engine/Plugins/ChaosClothAssetEditor/Source/ChaosClothAssetEditorTools/Public/ChaosClothAsset/ClothEditorToolBuilders.h b/Engine/Plugins/ChaosClothAssetEditor/Source/ChaosClothAssetEditorTools/Public/ChaosClothAsset/ClothEditorToolBuilders.h
new file mode 100644
--- /dev/null
+++ b/Engine/Plugins/ChaosClothAssetEditor/Source/ChaosClothAssetEditorTools/Public/ChaosClothAsset/ClothEditorToolBuilders.h
@@ -0,0 +1,82 @@
+// Copyright Epic Games, Inc. All Rights Reserved.
+
+#pragma once
+
+#include "Containers/Array.h"
+#include "UObject/Interface.h"
+#include "BaseTools/SingleSelectionMeshEditingTool.h"
+#include "BaseTools/MeshSurfacePointMeshEditingTool.h"
+#include "ClothEditorToolBuilders.generated.h"
+
+namespace UE::Chaos::ClothAsset
+{
+	enum class EClothPatternVertexType : uint8;
+}
+
+
+UINTERFACE(MinimalAPI)
+class UChaosClothAssetEditorToolBuilder : public UInterface
+{
+	GENERATED_BODY()
+};
+
+
+class IChaosClothAssetEditorToolBuilder
+{
+	GENERATED_BODY()
+
+public:
+
+	/** Returns all Construction View modes that this tool can operate in. The first element should be the preferred mode to switch to if necessary. */
+	virtual void GetSupportedViewModes(TArray<UE::Chaos::ClothAsset::EClothPatternVertexType>& Modes) const = 0;
+
+	/** Returns whether or not view can be set to wireframe when this tool is active.. */
+	virtual bool CanSetConstructionViewWireframeActive() const { return true; }
+};
+
+
+UCLASS()
+class CHAOSCLOTHASSETEDITORTOOLS_API UClothEditorWeightMapPaintToolBuilder : public UMeshSurfacePointMeshEditingToolBuilder, public IChaosClothAssetEditorToolBuilder
+{
+	GENERATED_BODY()
+
+private:
+	virtual void GetSupportedViewModes(TArray<UE::Chaos::ClothAsset::EClothPatternVertexType>& Modes) const override;
+	virtual UMeshSurfacePointTool* CreateNewTool(const FToolBuilderState& SceneState) const override;
+	virtual bool CanSetConstructionViewWireframeActive() const { return false; }
+};
+
+
+UCLASS()
+class CHAOSCLOTHASSETEDITORTOOLS_API UClothMeshSelectionToolBuilder : public UInteractiveToolWithToolTargetsBuilder, public IChaosClothAssetEditorToolBuilder
+{
+	GENERATED_BODY()
+
+private:
+
+	virtual void GetSupportedViewModes(TArray<UE::Chaos::ClothAsset::EClothPatternVertexType>& Modes) const override;
+	virtual bool CanBuildTool(const FToolBuilderState& SceneState) const override;
+	virtual UInteractiveTool* BuildTool(const FToolBuilderState& SceneState) const override;
+	virtual const FToolTargetTypeRequirements& GetTargetRequirements() const override;
+	virtual bool CanSetConstructionViewWireframeActive() const { return false; }
+};
+
+UCLASS()
+class CHAOSCLOTHASSETEDITORTOOLS_API UClothTransferSkinWeightsToolBuilder : public USingleSelectionMeshEditingToolBuilder, public IChaosClothAssetEditorToolBuilder
+{
+	GENERATED_BODY()
+
+private:
+
+	virtual void GetSupportedViewModes(TArray<UE::Chaos::ClothAsset::EClothPatternVertexType>& Modes) const override;
+	virtual USingleSelectionMeshEditingTool* CreateNewTool(const FToolBuilderState& SceneState) const override;
+
+};
+
+
+namespace UE::Chaos::ClothAsset
+{
+	// Provide a list of Tool default objects for us in TInteractiveToolCommands::RegisterCommands()
+	void CHAOSCLOTHASSETEDITORTOOLS_API GetClothEditorToolDefaultObjectList(TArray<UInteractiveTool*>& ToolCDOs);
+}
+
diff --git a/Engine/Plugins/ChaosClothAssetEditor/Source/ChaosClothAssetEditorTools/Public/ChaosClothAsset/ClothMeshSelectionTool.h b/Engine/Plugins/ChaosClothAssetEditor/Source/ChaosClothAssetEditorTools/Public/ChaosClothAsset/ClothMeshSelectionTool.h
deleted file mode 100644
diff --git a/Engine/Plugins/ChaosClothAssetEditor/Source/ChaosClothAssetEditorTools/Public/ChaosClothAsset/ClothTransferSkinWeightsTool.h b/Engine/Plugins/ChaosClothAssetEditor/Source/ChaosClothAssetEditorTools/Public/ChaosClothAsset/ClothTransferSkinWeightsTool.h
deleted file mode 100644
diff --git a/Engine/Plugins/ChaosClothAssetEditor/Source/ChaosClothAssetEditorTools/Public/ChaosClothAsset/ClothWeightMapPaintTool.h b/Engine/Plugins/ChaosClothAssetEditor/Source/ChaosClothAssetEditorTools/Public/ChaosClothAsset/ClothWeightMapPaintTool.h
deleted file mode 100644
