Remove ClothComponentToolTarget as it is not used in the Chaos Cloth Editor

diff --git a/Engine/Plugins/ChaosClothAssetEditor/Source/ChaosClothAssetEditor/Private/ChaosClothAsset/ClothEditorMode.cpp b/Engine/Plugins/ChaosClothAssetEditor/Source/ChaosClothAssetEditor/Private/ChaosClothAsset/ClothEditorMode.cpp
--- a/Engine/Plugins/ChaosClothAssetEditor/Source/ChaosClothAssetEditor/Private/ChaosClothAsset/ClothEditorMode.cpp
+++ b/Engine/Plugins/ChaosClothAssetEditor/Source/ChaosClothAssetEditor/Private/ChaosClothAsset/ClothEditorMode.cpp
@@ -30,7 +30,6 @@
 #include "Toolkits/BaseToolkit.h"
 #include "ToolTargets/ToolTarget.h"
 #include "ToolTargets/DynamicMeshComponentToolTarget.h"
-#include "ChaosClothAsset/ClothComponentToolTarget.h"
 #include "Engine/Selection.h"
 #include "Misc/ConfigCacheIni.h"
 #include "Modules/ModuleManager.h"
@@ -144,7 +143,6 @@ void UChaosClothAssetEditorMode::Enter()
 void UChaosClothAssetEditorMode::AddToolTargetFactories()
 {
 	GetInteractiveToolsContext()->TargetManager->AddTargetFactory(NewObject<UDynamicMeshComponentToolTargetFactory>(GetToolManager()));
-	GetInteractiveToolsContext()->TargetManager->AddTargetFactory(NewObject<UClothComponentToolTargetFactory>(GetToolManager()));
 }
 
 void UChaosClothAssetEditorMode::RegisterClothTool(TSharedPtr<FUICommandInfo> UICommand, 
@@ -420,7 +418,6 @@ void UChaosClothAssetEditorMode::SetPreviewScene(UE::Chaos::ClothAsset::FChaosCl
 
 	UEditorInteractiveToolsContext* const PreviewToolsContext = PreviewScene->GetClothPreviewEditorModeManager()->GetInteractiveToolsContext();
 	UInteractiveToolManager* const PreviewToolManager = PreviewToolsContext->ToolManager;
-	PreviewToolsContext->TargetManager->AddTargetFactory(NewObject<UClothComponentToolTargetFactory>(PreviewToolManager));
 	PreviewToolsContext->TargetManager->AddTargetFactory(NewObject<USkeletalMeshComponentToolTargetFactory>(PreviewToolManager));
 
 	PreviewToolManager->OnToolStarted.AddUObject(this, &UChaosClothAssetEditorMode::OnToolStarted);
diff --git a/Engine/Plugins/ChaosClothAssetEditor/Source/ChaosClothAssetEditorTools/Private/ChaosClothAsset/ClothComponentToolTarget.cpp b/Engine/Plugins/ChaosClothAssetEditor/Source/ChaosClothAssetEditorTools/Private/ChaosClothAsset/ClothComponentToolTarget.cpp
deleted file mode 100644
diff --git a/Engine/Plugins/ChaosClothAssetEditor/Source/ChaosClothAssetEditorTools/Public/ChaosClothAsset/ClothAssetBackedTarget.h b/Engine/Plugins/ChaosClothAssetEditor/Source/ChaosClothAssetEditorTools/Public/ChaosClothAsset/ClothAssetBackedTarget.h
deleted file mode 100644
diff --git a/Engine/Plugins/ChaosClothAssetEditor/Source/ChaosClothAssetEditorTools/Public/ChaosClothAsset/ClothComponentToolTarget.h b/Engine/Plugins/ChaosClothAssetEditor/Source/ChaosClothAssetEditorTools/Public/ChaosClothAsset/ClothComponentToolTarget.h
deleted file mode 100644
