ChaosClothEditor: fix crash when trying to remesh an empty render mesh.

diff --git a/Engine/Plugins/ChaosClothAssetEditor/Source/ChaosClothAssetDataflowNodes/Private/ChaosClothAsset/RemeshNode.cpp b/Engine/Plugins/ChaosClothAssetEditor/Source/ChaosClothAssetDataflowNodes/Private/ChaosClothAsset/RemeshNode.cpp
--- a/Engine/Plugins/ChaosClothAssetEditor/Source/ChaosClothAssetDataflowNodes/Private/ChaosClothAsset/RemeshNode.cpp
+++ b/Engine/Plugins/ChaosClothAssetEditor/Source/ChaosClothAssetDataflowNodes/Private/ChaosClothAsset/RemeshNode.cpp
@@ -1291,9 +1291,8 @@ void FChaosClothAssetRemeshNode::RemeshRenderMesh(const TSharedRef<const FManage
 
 	const bool bHasUVs = (DynamicMesh.Attributes()->PrimaryUV() != nullptr);
 
-	const double MeshArea = TMeshQueries<FDynamicMesh3>::GetVolumeArea(DynamicMesh).Y;
 	const int TargetTriangleCount = FMath::RoundToInt(static_cast<float>(TargetPercentRender) / 100.0f * static_cast<float>(InputMeshTriangleCount));
-	const double TargetEdgeLength = FRemeshMeshOp::CalculateTargetEdgeLength(nullptr, TargetTriangleCount, MeshArea);
+	const double TargetEdgeLength = FRemeshMeshOp::CalculateTargetEdgeLength(&DynamicMesh, TargetTriangleCount);
 
 	TArray<TArray<FIntVector2>> Seams;
 
