Chaos Cloth Asset - Fixed required bones when Leader's LOD isn't available on the Cloth Asset.

diff --git a/Engine/Plugins/ChaosClothAsset/Source/ChaosClothAssetEngine/Private/ChaosClothAsset/ClothComponent.cpp b/Engine/Plugins/ChaosClothAsset/Source/ChaosClothAssetEngine/Private/ChaosClothAsset/ClothComponent.cpp
--- a/Engine/Plugins/ChaosClothAsset/Source/ChaosClothAssetEngine/Private/ChaosClothAsset/ClothComponent.cpp
+++ b/Engine/Plugins/ChaosClothAsset/Source/ChaosClothAssetEngine/Private/ChaosClothAsset/ClothComponent.cpp
@@ -398,10 +398,15 @@ void UChaosClothComponent::SetSkinnedAssetAndUpdate(USkinnedAsset* InSkinnedAsse
 	}
 }
 
-void UChaosClothComponent::GetAdditionalRequiredBonesForLeader(int32 LODIndex, TArray<FBoneIndexType>& InOutRequiredBones) const
+void UChaosClothComponent::GetAdditionalRequiredBonesForLeader(int32 LeaderLODIndex, TArray<FBoneIndexType>& InOutRequiredBones) const
 {
 	if (const FSkeletalMeshRenderData* const SkeletalMeshRenderData = GetSkeletalMeshRenderData())
 	{
+		const int32 MinLODIndex = ComputeMinLOD();
+		const int32 MaxLODIndex = FMath::Max(GetNumLODs() - 1, MinLODIndex);
+
+		const int32 LODIndex = FMath::Clamp(LeaderLODIndex, MinLODIndex, MaxLODIndex);
+
 		if (SkeletalMeshRenderData->LODRenderData.IsValidIndex(LODIndex))
 		{
 			// Gather the follower's bones
diff --git a/Engine/Plugins/ChaosClothAssetEditor/Source/ChaosClothAssetEditor/Private/ChaosClothAsset/ClothEditorPreviewScene.cpp b/Engine/Plugins/ChaosClothAssetEditor/Source/ChaosClothAssetEditor/Private/ChaosClothAsset/ClothEditorPreviewScene.cpp
--- a/Engine/Plugins/ChaosClothAssetEditor/Source/ChaosClothAssetEditor/Private/ChaosClothAsset/ClothEditorPreviewScene.cpp
+++ b/Engine/Plugins/ChaosClothAssetEditor/Source/ChaosClothAssetEditor/Private/ChaosClothAsset/ClothEditorPreviewScene.cpp
@@ -82,27 +82,6 @@ namespace UE::Chaos::ClothAsset::Private
 					UE_LOG(LogChaosClothAssetEditor, Warning, TEXT("%s"), *MissingBonesMsg.ToString());
 					bAnyMissingBones = true;
 				}
-
-				// These should be the bones that are required by the follower cloth component
-				const int32 ClothLODIndex = FMath::Min(ClothComponent.GetNumLODs() - 1, LODIndex);
-				if (ensure(ClothRenderData->LODRenderData.IsValidIndex(ClothLODIndex)))
-				{
-					CheckRequiredBones(ClothRenderData->LODRenderData[ClothLODIndex].ActiveBoneIndices,
-						SKMFillComponentSpaceTransformsRequiredBones, MissingBones);
-
-					if (MissingBones.Num())
-					{
-						FText MissingBonesMsg = FText::FormatOrdered(LOCTEXT("MissingClothBones", "SkeletalMesh \"{0}\" (LOD {1}) will not update the following bones required by \"{2}\" (LOD {3}): "),
-							FText::FromString(SkeletalMeshComponent->GetSkinnedAsset()->GetName()),
-							FText::AsNumber(LODIndex), FText::FromString(ClothAsset->GetName()), FText::AsNumber(ClothLODIndex));
-						for (const FBoneIndexType MissingBone : MissingBones)
-						{
-							MissingBonesMsg = FText::Format(LOCTEXT("MissingClothBoneList", "{0} {1}"), MissingBonesMsg, FText::FromName(RefSkeleton.GetBoneName(MissingBone)));
-						}
-						UE_LOG(LogChaosClothAssetEditor, Warning, TEXT("%s"), *MissingBonesMsg.ToString());
-						bAnyMissingBones = true;
-					}
-				}
 			}
 
 			if (bAnyMissingBones)
