ChaosClothAsset: Disallow merging cloth collections with different reference skeletons.

Also: when merging cloth collections, use other's lod data if our lod data is empty. And make LogAndToastWarning log level Warning rather than Display.

diff --git a/Engine/Plugins/ChaosClothAsset/Source/ChaosClothAsset/Private/ChaosClothAsset/CollectionClothFacade.cpp b/Engine/Plugins/ChaosClothAsset/Source/ChaosClothAsset/Private/ChaosClothAsset/CollectionClothFacade.cpp
--- a/Engine/Plugins/ChaosClothAsset/Source/ChaosClothAsset/Private/ChaosClothAsset/CollectionClothFacade.cpp
+++ b/Engine/Plugins/ChaosClothAsset/Source/ChaosClothAsset/Private/ChaosClothAsset/CollectionClothFacade.cpp
@@ -540,10 +540,6 @@ namespace UE::Chaos::ClothAsset
 	void FCollectionClothFacade::Initialize(const FCollectionClothConstFacade& Other)
 	{
 		Reset();
-
-		// LODs Group
-		SetPhysicsAssetPathName(Other.GetPhysicsAssetPathName());
-		SetSkeletalMeshPathName(Other.GetSkeletalMeshPathName());
 		
 		// Solvers Group
 		if (Other.IsValid(EClothCollectionOptionalSchemas::Solvers) && Other.HasSolverElement())
@@ -567,7 +563,15 @@ namespace UE::Chaos::ClothAsset
 #endif
 
 		// LODs Group 
-		// Just keep original data.
+		// Keep original data unless our data is empty. Then take Other's data.
+		if (GetPhysicsAssetPathName().IsEmpty())
+		{
+			SetPhysicsAssetPathName(Other.GetPhysicsAssetPathName());
+		}
+		if (GetSkeletalMeshPathName().IsEmpty())
+		{
+			SetSkeletalMeshPathName(Other.GetSkeletalMeshPathName());
+		}
 
 		// Very important order of operations to ensure indices don't get messed up:
 		// 1) Append 3D Vertices, but don't set 2D Lookups or SeamStitch Lookups because those indices don't exist yet.
diff --git a/Engine/Plugins/ChaosClothAssetEditor/Source/ChaosClothAssetDataflowNodes/Private/ChaosClothAsset/ClothDataflowTools.cpp b/Engine/Plugins/ChaosClothAssetEditor/Source/ChaosClothAssetDataflowNodes/Private/ChaosClothAsset/ClothDataflowTools.cpp
--- a/Engine/Plugins/ChaosClothAssetEditor/Source/ChaosClothAssetDataflowNodes/Private/ChaosClothAsset/ClothDataflowTools.cpp
+++ b/Engine/Plugins/ChaosClothAssetEditor/Source/ChaosClothAssetDataflowNodes/Private/ChaosClothAsset/ClothDataflowTools.cpp
@@ -451,7 +451,7 @@ namespace UE::Chaos::ClothAsset
 		NotificationInfo.ExpireDuration = 5.0f;
 		FSlateNotificationManager::Get().AddNotification(NotificationInfo);
 
-		UE_LOG(LogChaosClothAssetDataflowNodes, Display, TEXT("%s"), *Text.ToString());
+		UE_LOG(LogChaosClothAssetDataflowNodes, Warning, TEXT("%s"), *Text.ToString());
 	}
 
 	void FClothDataflowTools::MakeCollectionName(FString& InOutString)
diff --git a/Engine/Plugins/ChaosClothAssetEditor/Source/ChaosClothAssetDataflowNodes/Private/ChaosClothAsset/MergeClothCollectionsNode.cpp b/Engine/Plugins/ChaosClothAssetEditor/Source/ChaosClothAssetDataflowNodes/Private/ChaosClothAsset/MergeClothCollectionsNode.cpp
--- a/Engine/Plugins/ChaosClothAssetEditor/Source/ChaosClothAssetDataflowNodes/Private/ChaosClothAsset/MergeClothCollectionsNode.cpp
+++ b/Engine/Plugins/ChaosClothAssetEditor/Source/ChaosClothAssetDataflowNodes/Private/ChaosClothAsset/MergeClothCollectionsNode.cpp
@@ -6,6 +6,7 @@
 #include "ChaosClothAsset/ClothDataflowTools.h"
 #include "Chaos/CollectionPropertyFacade.h"
 #include "Dataflow/DataflowInputOutput.h"
+#include "Engine/SkeletalMesh.h"
 
 #include UE_INLINE_GENERATED_CPP_BY_NAME(MergeClothCollectionsNode)
 
@@ -17,11 +18,11 @@ namespace UE::Chaos::ClothAsset::Private
 	{
 		using namespace UE::Chaos::ClothAsset;
 
-		static const FText Headline = LOCTEXT("DiffertentWeightMapNamesHeadline", "Different weight map names.");
+		static const FText Headline = LOCTEXT("DifferentWeightMapNamesHeadline", "Different weight map names.");
 
 		const FText Details = FText::Format(
 			LOCTEXT(
-				"DiffertentWeightMapNamesDetails",
+				"DifferentWeightMapNamesDetails",
 				"Two identical Cloth Collection properties '{0}' are being merged but have different weight map names '{1}' and '{2}'. The weight map named '{3}' will be used in the resulting merge."),
 			FText::FromString(PropertyName),
 			FText::FromString(OutWeightMapName),
@@ -247,6 +248,85 @@ namespace UE::Chaos::ClothAsset::Private
 			}
 		}
 	}
+
+	static bool AreSkeletalMeshesCompatible(const FDataflowNode& DataflowNode, const FCollectionClothConstFacade& Cloth1, const FCollectionClothConstFacade& Cloth2)
+	{
+		/** Disallow merging cloth facades with incompatible ref skeletons. */
+		const FString& SkeletalMeshPathName1 = Cloth1.GetSkeletalMeshPathName();
+		const FString& SkeletalMeshPathName2 = Cloth2.GetSkeletalMeshPathName();
+		if (SkeletalMeshPathName1.IsEmpty() || SkeletalMeshPathName2.IsEmpty() || SkeletalMeshPathName1 == SkeletalMeshPathName2)
+		{
+			return true;
+		}
+
+		static const FText ErrorHeadline = LOCTEXT("IncompatibleSkeletalMeshesHeadline", "Incompatible Skeletal Meshes.");
+		const USkeletalMesh* const SkeletalMesh1 = LoadObject<USkeletalMesh>(nullptr, *SkeletalMeshPathName1, nullptr, LOAD_None, nullptr);
+		const USkeletalMesh* const SkeletalMesh2 = LoadObject<USkeletalMesh>(nullptr, *SkeletalMeshPathName2, nullptr, LOAD_None, nullptr);
+		if (!SkeletalMesh1 || !SkeletalMesh2)
+		{
+			const FText Details = FText::Format(
+				LOCTEXT(
+					"IncompatibleSkeletalMeshesLoadFailureDetails",
+					"Cloth collections failed to merge due to failing to load SkeletalMesh \"{0}\" to check compatibility."),
+				!SkeletalMesh1 ? FText::FromString(SkeletalMeshPathName1) : FText::FromString(SkeletalMeshPathName2));
+
+			FClothDataflowTools::LogAndToastWarning(DataflowNode, ErrorHeadline, Details);
+			return false;
+		}
+
+		const FReferenceSkeleton& RefSkeleton1 = SkeletalMesh1->GetRefSkeleton();
+		const FReferenceSkeleton& RefSkeleton2 = SkeletalMesh2->GetRefSkeleton();
+		if (RefSkeleton1.GetNum() != RefSkeleton2.GetNum())
+		{
+			const FText Details = FText::Format(
+				LOCTEXT(
+					"IncompatibleSkeletalMeshesNumBonesDetails",
+					"Cloth collections failed to merge due to incompatible Skeletal Meshes, \"{0}\" and \"{1}\". RefSkeleton Bone counts {2} != {3}."),
+				FText::FromString(SkeletalMeshPathName1),
+				FText::FromString(SkeletalMeshPathName2),
+				RefSkeleton1.GetNum(),
+				RefSkeleton2.GetNum());
+
+			FClothDataflowTools::LogAndToastWarning(DataflowNode, ErrorHeadline, Details);
+			return false;
+		}
+
+		const TArray<FMeshBoneInfo>& RefBoneInfo1 = RefSkeleton1.GetRefBoneInfo();
+		const TArray<FMeshBoneInfo>& RefBoneInfo2 = RefSkeleton2.GetRefBoneInfo();
+		const TArray<FTransform>& RefBonePose1 = RefSkeleton1.GetRefBonePose();
+		const TArray<FTransform>& RefBonePose2 = RefSkeleton2.GetRefBonePose();
+		for (int32 BoneIndex = 0; BoneIndex < RefSkeleton1.GetNum(); ++BoneIndex)
+		{
+			if (!(RefBoneInfo1[BoneIndex] == RefBoneInfo2[BoneIndex]))
+			{
+				const FText Details = FText::Format(
+					LOCTEXT(
+						"IncompatibleSkeletalMeshesRefBoneInfoDetails",
+						"Cloth collections failed to merge due to incompatible Skeletal Meshes, \"{0}\" and \"{1}\". RefBoneInfos are mismatched at BoneIndex {2}."),
+					FText::FromString(SkeletalMeshPathName1),
+					FText::FromString(SkeletalMeshPathName2),
+					BoneIndex);
+
+				FClothDataflowTools::LogAndToastWarning(DataflowNode, ErrorHeadline, Details);
+				return false;
+			}
+			if (!RefBonePose1[BoneIndex].Equals(RefBonePose2[BoneIndex]))
+			{
+				const FText Details = FText::Format(
+					LOCTEXT(
+						"IncompatibleSkeletalMeshesRefBonePoseDetails",
+						"Cloth collections failed to merge due to incompatible Skeletal Meshes, \"{0}\" and \"{1}\". RefBonePoses are mismatched at BoneIndex {2}."),
+					FText::FromString(SkeletalMeshPathName1),
+					FText::FromString(SkeletalMeshPathName2),
+					BoneIndex);
+
+				FClothDataflowTools::LogAndToastWarning(DataflowNode, ErrorHeadline, Details);
+				return false;
+			}
+		}
+
+		return true;
+	}
 }
 
 FChaosClothAssetMergeClothCollectionsNode::FChaosClothAssetMergeClothCollectionsNode(const Dataflow::FNodeParameters& InParam, FGuid InGuid)
@@ -308,7 +388,7 @@ void FChaosClothAssetMergeClothCollectionsNode::Evaluate(Dataflow::FContext& Con
 
 			// Append cloth
 			const FCollectionClothConstFacade OtherClothFacade(OtherClothCollection);
-			if (OtherClothFacade.IsValid())
+			if (OtherClothFacade.IsValid() && Private::AreSkeletalMeshesCompatible(*this, ClothFacade, OtherClothFacade))
 			{
 				ClothFacade.Append(OtherClothFacade);
 				bAreAnyValid = true;
