Chaos Cloth Asset - Fixed the cloth asset update that wasn't getting triggered when adding user attributes, such as for the RecomputeTangents attribute.

diff --git a/Engine/Plugins/ChaosClothAsset/Source/ChaosClothAsset/Private/ChaosClothAsset/ClothCollection.cpp b/Engine/Plugins/ChaosClothAsset/Source/ChaosClothAsset/Private/ChaosClothAsset/ClothCollection.cpp
--- a/Engine/Plugins/ChaosClothAsset/Source/ChaosClothAsset/Private/ChaosClothAsset/ClothCollection.cpp
+++ b/Engine/Plugins/ChaosClothAsset/Source/ChaosClothAsset/Private/ChaosClothAsset/ClothCollection.cpp
@@ -676,14 +676,16 @@ namespace UE::Chaos::ClothAsset
 
 		TArray<FName> UserDefinedAttributeNames;
 
-		const TArray<FName>& FixedAttributeNames = FixedAttributeNamesMap.FindChecked(GroupName);  // Also checks that the group name is a recognized group name 
-		const TArray<FName> AttributeNames = ManagedArrayCollection->AttributeNames(GroupName);
-
-		for (const FName& AttributeName : AttributeNames)
+		if (const TArray<FName>* const FixedAttributeNames = FixedAttributeNamesMap.Find(GroupName))  // Also checks that the group name is a recognized group name
 		{
-			if (!FixedAttributeNames.Contains(AttributeName) && ManagedArrayCollection->FindAttributeTyped<T>(AttributeName, GroupName))
+			const TArray<FName> AttributeNames = ManagedArrayCollection->AttributeNames(GroupName);
+
+			for (const FName& AttributeName : AttributeNames)
 			{
-				UserDefinedAttributeNames.Add(AttributeName);
+				if (!FixedAttributeNames->Contains(AttributeName) && ManagedArrayCollection->FindAttributeTyped<T>(AttributeName, GroupName))
+				{
+					UserDefinedAttributeNames.Add(AttributeName);
+				}
 			}
 		}
 
diff --git a/Engine/Plugins/ChaosClothAsset/Source/ChaosClothAsset/Private/ChaosClothAsset/CollectionClothFacade.cpp b/Engine/Plugins/ChaosClothAsset/Source/ChaosClothAsset/Private/ChaosClothAsset/CollectionClothFacade.cpp
--- a/Engine/Plugins/ChaosClothAsset/Source/ChaosClothAsset/Private/ChaosClothAsset/CollectionClothFacade.cpp
+++ b/Engine/Plugins/ChaosClothAsset/Source/ChaosClothAsset/Private/ChaosClothAsset/CollectionClothFacade.cpp
@@ -468,6 +468,11 @@ namespace UE::Chaos::ClothAsset
 	template CHAOSCLOTHASSET_API TConstArrayView<float> FCollectionClothConstFacade::GetUserDefinedAttribute<float>(const FName& Name, const FName& GroupName) const;
 	template CHAOSCLOTHASSET_API TConstArrayView<FVector3f> FCollectionClothConstFacade::GetUserDefinedAttribute<FVector3f>(const FName& Name, const FName& GroupName) const;
 
+	bool FCollectionClothConstFacade::IsValidClothCollectionGroupName(const FName& GroupName)
+	{
+		return FClothCollection::IsValidClothCollectionGroupName(GroupName);
+	}
+
 	void FCollectionClothConstFacade::BuildSimulationMesh(TArray<FVector3f>& Positions, TArray<FVector3f>& Normals, TArray<uint32>& Indices, TArray<FVector2f>& PatternsPositions, TArray<uint32>& PatternsIndices,
 		TArray<uint32>& PatternToWeldedIndices, TArray<TArray<int32>>* OptionalWeldedToPatternIndices) const
 	{
@@ -1124,11 +1129,6 @@ namespace UE::Chaos::ClothAsset
 	template CHAOSCLOTHASSET_API TArrayView<float> FCollectionClothFacade::GetUserDefinedAttribute<float>(const FName& Name, const FName& GroupName);
 	template CHAOSCLOTHASSET_API TArrayView<FVector3f> FCollectionClothFacade::GetUserDefinedAttribute<FVector3f>(const FName& Name, const FName& GroupName);
 
-	bool FCollectionClothFacade::IsValidClothCollectionGroupName(const FName& GroupName)
-	{
-		return FClothCollection::IsValidClothCollectionGroupName(GroupName);
-	}
-
 	void FCollectionClothFacade::SetDefaults()
 	{
 		GetClothCollection()->SetNumElements(1, ClothCollectionGroup::Lods);
diff --git a/Engine/Plugins/ChaosClothAssetEditor/Source/ChaosClothAssetDataflowNodes/Private/ChaosClothAsset/AttributeNode.cpp b/Engine/Plugins/ChaosClothAssetEditor/Source/ChaosClothAssetDataflowNodes/Private/ChaosClothAsset/AttributeNode.cpp
--- a/Engine/Plugins/ChaosClothAssetEditor/Source/ChaosClothAssetDataflowNodes/Private/ChaosClothAsset/AttributeNode.cpp
+++ b/Engine/Plugins/ChaosClothAssetEditor/Source/ChaosClothAssetDataflowNodes/Private/ChaosClothAsset/AttributeNode.cpp
@@ -87,11 +87,23 @@ void FChaosClothAssetAttributeNode::Evaluate(Dataflow::FContext& Context, const
 
 void FChaosClothAssetAttributeNode::OnSelected(Dataflow::FContext& Context)
 {
+	using namespace UE::Chaos::ClothAsset;
+
 	// Re-evaluate the input collection
-	const FManagedArrayCollection& SelectionCollection = GetValue<FManagedArrayCollection>(Context, &Collection);
+	FManagedArrayCollection InCollection = GetValue<FManagedArrayCollection>(Context, &Collection);
+	const TSharedRef<FManagedArrayCollection> ClothCollection = MakeShared<FManagedArrayCollection>(MoveTemp(InCollection));
+	FCollectionClothFacade Cloth(ClothCollection);
 
 	// Update the list of used group for the UI customization
-	CachedCollectionGroupNames = SelectionCollection.GroupNames();
+	const TArray<FName> GroupNames = ClothCollection->GroupNames();
+	CachedCollectionGroupNames.Reset(GroupNames.Num());
+	for (const FName& GroupName : GroupNames)
+	{
+		if (Cloth.IsValidClothCollectionGroupName(GroupName))  // Restrict to the cloth facade groups
+		{
+			CachedCollectionGroupNames.Emplace(GroupName);
+		}
+	}
 }
 
 void FChaosClothAssetAttributeNode::OnDeselected()
diff --git a/Engine/Plugins/ChaosClothAssetEditor/Source/ChaosClothAssetDataflowNodes/Private/ChaosClothAsset/SelectionNode.cpp b/Engine/Plugins/ChaosClothAssetEditor/Source/ChaosClothAssetDataflowNodes/Private/ChaosClothAsset/SelectionNode.cpp
--- a/Engine/Plugins/ChaosClothAssetEditor/Source/ChaosClothAssetDataflowNodes/Private/ChaosClothAsset/SelectionNode.cpp
+++ b/Engine/Plugins/ChaosClothAssetEditor/Source/ChaosClothAssetDataflowNodes/Private/ChaosClothAsset/SelectionNode.cpp
@@ -2,6 +2,7 @@
 
 #include "ChaosClothAsset/SelectionNode.h"
 #include "ChaosClothAsset/ClothCollectionGroup.h"
+#include "ChaosClothAsset/CollectionClothFacade.h"
 #include "ChaosClothAsset/ClothDataflowTools.h"
 #include "ChaosClothAsset/CollectionClothSelectionFacade.h"
 #include "Dataflow/DataflowInputOutput.h"
@@ -102,11 +103,23 @@ void FChaosClothAssetSelectionNode::Evaluate(Dataflow::FContext& Context, const
 
 void FChaosClothAssetSelectionNode::OnSelected(Dataflow::FContext& Context)
 {
+	using namespace UE::Chaos::ClothAsset;
+
 	// Re-evaluate the input collection
-	const FManagedArrayCollection& SelectionCollection = GetValue<FManagedArrayCollection>(Context, &Collection);
+	FManagedArrayCollection InCollection = GetValue<FManagedArrayCollection>(Context, &Collection);
+	const TSharedRef<FManagedArrayCollection> ClothCollection = MakeShared<FManagedArrayCollection>(MoveTemp(InCollection));
+	FCollectionClothFacade Cloth(ClothCollection);
 
 	// Update the list of used group for the UI customization
-	CachedCollectionGroupNames = SelectionCollection.GroupNames();
+	const TArray<FName> GroupNames = ClothCollection->GroupNames();
+	CachedCollectionGroupNames.Reset(GroupNames.Num());
+	for (const FName& GroupName : GroupNames)
+	{
+		if (Cloth.IsValidClothCollectionGroupName(GroupName))  // Restrict to the cloth facade groups
+		{
+			CachedCollectionGroupNames.Emplace(GroupName);
+		}
+	}
 }
 
 void FChaosClothAssetSelectionNode::OnDeselected()
diff --git a/Engine/Plugins/ChaosClothAssetEditor/Source/ChaosClothAssetDataflowNodes/Private/ChaosClothAsset/TerminalNode.cpp b/Engine/Plugins/ChaosClothAssetEditor/Source/ChaosClothAssetDataflowNodes/Private/ChaosClothAsset/TerminalNode.cpp
--- a/Engine/Plugins/ChaosClothAssetEditor/Source/ChaosClothAssetDataflowNodes/Private/ChaosClothAsset/TerminalNode.cpp
+++ b/Engine/Plugins/ChaosClothAssetEditor/Source/ChaosClothAssetDataflowNodes/Private/ChaosClothAsset/TerminalNode.cpp
@@ -2,12 +2,12 @@
 
 #include "ChaosClothAsset/TerminalNode.h"
 #include "ChaosClothAsset/ClothAsset.h"
+#include "ChaosClothAsset/ClothCollectionGroup.h"
 #include "ChaosClothAsset/ClothDataflowTools.h"
 #include "ChaosClothAsset/ClothGeometryTools.h"
 #include "ChaosClothAsset/ClothLODTransitionDataCache.h"
 #include "ChaosClothAsset/CollectionClothFacade.h"
 #include "ChaosClothAsset/CollectionClothSelectionFacade.h"
-#include "ChaosClothAsset/ClothCollectionGroup.h"
 #include "Animation/Skeleton.h"
 #include "Chaos/CollectionPropertyFacade.h"
 #include "Dataflow/DataflowInputOutput.h"
@@ -31,7 +31,14 @@ namespace UE::Chaos::ClothAsset::Private
 			if (Cloth.HasValidData())
 			{
 				Checksum = Cloth.CalculateTypeHash(bIncludeWeightMapsTrue, Checksum);
-				Checksum = Cloth.CalculateUserDefinedAttributesTypeHash<int32>(ClothCollectionGroup::SimFaces, Checksum);
+
+				const TArray<FName> GroupNames = ClothCollection->GroupNames();
+				for (const FName& GroupName : GroupNames)
+				{
+					Checksum = Cloth.CalculateUserDefinedAttributesTypeHash<int32>(GroupName, Checksum);
+					Checksum = Cloth.CalculateUserDefinedAttributesTypeHash<float>(GroupName, Checksum);
+					Checksum = Cloth.CalculateUserDefinedAttributesTypeHash<FVector3f>(GroupName, Checksum);
+				}
 			}
 			FCollectionClothSelectionConstFacade Selection(ClothCollection);
 			if (Selection.IsValid())
diff --git a/Engine/Plugins/ChaosClothAssetEditor/Source/ChaosClothAssetDataflowNodes/Public/ChaosClothAsset/SimulationVelocityScaleConfigNode.h b/Engine/Plugins/ChaosClothAssetEditor/Source/ChaosClothAssetDataflowNodes/Public/ChaosClothAsset/SimulationVelocityScaleConfigNode.h
--- a/Engine/Plugins/ChaosClothAssetEditor/Source/ChaosClothAssetDataflowNodes/Public/ChaosClothAsset/SimulationVelocityScaleConfigNode.h
+++ b/Engine/Plugins/ChaosClothAssetEditor/Source/ChaosClothAssetDataflowNodes/Public/ChaosClothAsset/SimulationVelocityScaleConfigNode.h
@@ -20,7 +20,7 @@ public:
 	 * (the closest bone to the root on which the cloth section has been skinned, or the root itself if the cloth isn't skinned).
 	 * This value will be clamped by "Max Velocity Scale". A velocity scale of > 1 will amplify the velocities from the reference bone.
 	 */
-	UPROPERTY(EditAnywhere, Category = "Animation Properties", meta = (UIMin = "0", UIMax = "1", ClampMin = "0", ClampMax = "100"))
+	UPROPERTY(EditAnywhere, Category = "Velocity Scale Properties", meta = (UIMin = "0", UIMax = "1", ClampMin = "0", ClampMax = "100"))
 	FVector3f LinearVelocityScale = { 0.75f, 0.75f, 0.75f };
 
 	/**
@@ -28,13 +28,13 @@ public:
 	 * (the closest bone to the root on which the cloth section has been skinned, or the root itself if the cloth isn't skinned).
 	 * This value will be clamped by "Max Velocity Scale". A velocity scale of > 1 will amplify the velocities from the reference bone.
 	 */
-	UPROPERTY(EditAnywhere, Category = "Animation Properties", meta = (UIMin = "0", UIMax = "1", ClampMin = "0", ClampMax = "100"))
+	UPROPERTY(EditAnywhere, Category = "Velocity Scale Properties", meta = (UIMin = "0", UIMax = "1", ClampMin = "0", ClampMax = "100"))
 	float AngularVelocityScale = 0.75f;
 
 	/**
 	 * Clamp on Linear and Angular Velocity Scale. The final velocity scale (e.g., including contributions from blueprints) will be clamped to this value.
 	 */
-	UPROPERTY(EditAnywhere, Category = "Animation Properties", meta = (UIMin = "0", UIMax = "1", ClampMin = "0", ClampMax = "100"))
+	UPROPERTY(EditAnywhere, Category = "Velocity Scale Properties", meta = (UIMin = "0", UIMax = "1", ClampMin = "0", ClampMax = "100"))
 	float MaxVelocityScale = 1.f;
 
 	/**
@@ -44,7 +44,7 @@ public:
 	 * case the cloth is simulated with full world space angular velocities and subjected to the true physical world inertial forces.
 	 * Values range from 0 to 2, with 0 showing no centrifugal effect, 1 full centrifugal effect, and 2 an overdriven centrifugal effect.
 	 */
-	UPROPERTY(EditAnywhere, Category = "Animation Properties", meta = (UIMin = "0", UIMax = "1", ClampMin = "0", ClampMax = "2"))
+	UPROPERTY(EditAnywhere, Category = "Velocity Scale Properties", meta = (UIMin = "0", UIMax = "1", ClampMin = "0", ClampMax = "2"))
 	float FictitiousAngularScale = 1.f;
 
 	FChaosClothAssetSimulationVelocityScaleConfigNode(const Dataflow::FNodeParameters& InParam, FGuid InGuid = FGuid::NewGuid());
