Chaos Cloth Asset: USDImport set BucklingStiffness = BendingStiffness when Buckling Ratio == 0.

diff --git a/Engine/Plugins/ChaosClothAssetEditor/Source/ChaosClothAssetDataflowNodes/Private/ChaosClothAsset/SimulationBendingConfigNode.cpp b/Engine/Plugins/ChaosClothAssetEditor/Source/ChaosClothAssetDataflowNodes/Private/ChaosClothAsset/SimulationBendingConfigNode.cpp
--- a/Engine/Plugins/ChaosClothAssetEditor/Source/ChaosClothAssetDataflowNodes/Private/ChaosClothAsset/SimulationBendingConfigNode.cpp
+++ b/Engine/Plugins/ChaosClothAssetEditor/Source/ChaosClothAssetDataflowNodes/Private/ChaosClothAsset/SimulationBendingConfigNode.cpp
@@ -103,7 +103,7 @@ void FChaosClothAssetSimulationBendingConfigNode::AddProperties(FPropertyHelper&
 			PropertyHelper.SetFabricPropertyWeighted(FName(TEXT("XPBDAnisoBucklingStiffnessWarp")), BucklingStiffnessWarp, [](
 				const UE::Chaos::ClothAsset::FCollectionClothFabricFacade& FabricFacade)-> float
 			 {
-				 return FabricFacade.GetBucklingStiffness().Warp;
+				 return FabricFacade.GetBucklingRatio() > UE_SMALL_NUMBER ? FabricFacade.GetBucklingStiffness().Warp : FabricFacade.GetBendingStiffness().Warp;
 			 }, {
 				FName(TEXT("BucklingStiffness")),    
 				FName(TEXT("XPBDBucklingStiffness"))});
@@ -111,13 +111,13 @@ void FChaosClothAssetSimulationBendingConfigNode::AddProperties(FPropertyHelper&
 			PropertyHelper.SetFabricPropertyWeighted(FName(TEXT("XPBDAnisoBucklingStiffnessWeft")), BucklingStiffnessWeft, [](
 				const UE::Chaos::ClothAsset::FCollectionClothFabricFacade& FabricFacade)-> float
 			{
-				return FabricFacade.GetBucklingStiffness().Weft;
+				return FabricFacade.GetBucklingRatio() > UE_SMALL_NUMBER ? FabricFacade.GetBucklingStiffness().Weft : FabricFacade.GetBendingStiffness().Weft;
 			},{});
 			
 			PropertyHelper.SetFabricPropertyWeighted(FName(TEXT("XPBDAnisoBucklingStiffnessBias")), BucklingStiffnessBias, [](
 				const UE::Chaos::ClothAsset::FCollectionClothFabricFacade& FabricFacade)-> float
 			{
-				return FabricFacade.GetBucklingStiffness().Bias;
+				return FabricFacade.GetBucklingRatio() > UE_SMALL_NUMBER ? FabricFacade.GetBucklingStiffness().Bias : FabricFacade.GetBendingStiffness().Bias;
 			},{});
 
 			PropertyHelper.SetFabricPropertyWeighted(FName(TEXT("XPBDAnisoBendingDamping")), BendingAnisoDamping, [](
diff --git a/Engine/Plugins/ChaosClothAssetEditor/Source/ChaosClothAssetDataflowNodes/Private/ChaosClothAsset/USDImportNode.cpp b/Engine/Plugins/ChaosClothAssetEditor/Source/ChaosClothAssetDataflowNodes/Private/ChaosClothAsset/USDImportNode.cpp
--- a/Engine/Plugins/ChaosClothAssetEditor/Source/ChaosClothAssetDataflowNodes/Private/ChaosClothAsset/USDImportNode.cpp
+++ b/Engine/Plugins/ChaosClothAssetEditor/Source/ChaosClothAssetDataflowNodes/Private/ChaosClothAsset/USDImportNode.cpp
@@ -793,14 +793,17 @@ bool FChaosClothAssetUSDImportNode::ImportFromFile(const FString& UsdFilePath, c
 			
 			FCollectionClothFabricFacade::FAnisotropicData StretchStiffness(
 				StretchWeft*StretchShearScaling, StretchWarp*StretchShearScaling, 0.5f * (ShearLeft+ShearRight)*StretchShearScaling);
+
+			// Only scalar value used in the solver right now
+			const float BucklingRatio = (BucklingRatioWeft + BucklingRatioWarp +
+				0.5f * (BucklingRatioBiasLeft + BucklingRatioBiasRight)) / 3.0f;
 			
-			FCollectionClothFabricFacade::FAnisotropicData BucklingStiffness(
+			FCollectionClothFabricFacade::FAnisotropicData BucklingStiffness = 
+				BucklingRatio < UE_SMALL_NUMBER ? BendingStiffness : 
+				FCollectionClothFabricFacade::FAnisotropicData(
 				BendingStiffness.Weft * BucklingStiffnessWeft, BendingStiffness.Warp * BucklingStiffnessWarp,
 				BendingStiffness.Bias * 0.5f * (BucklingStiffnessBiasLeft+BucklingStiffnessBiasRight));
 
-			// Only scalar value used in the solver right now
-			const float BucklingRatio = (BucklingRatioWeft + BucklingRatioWarp +
-				0.5f * (BucklingRatioBiasLeft+BucklingRatioBiasRight)) / 3.0f;
 
 			Fabric.Initialize(BendingStiffness, BucklingRatio, BucklingStiffness, StretchStiffness,
 				Density * DensityScaling, Friction, Damping, 0.0f, 0, Thickness * ThicknessScaling);
