Cloth Editor: fix crash when reloading Dataflow asset

diff --git a/Engine/Plugins/ChaosClothAssetEditor/Source/ChaosClothAssetEditor/Private/ChaosClothAsset/ClothEditorMode.cpp b/Engine/Plugins/ChaosClothAssetEditor/Source/ChaosClothAssetEditor/Private/ChaosClothAsset/ClothEditorMode.cpp
--- a/Engine/Plugins/ChaosClothAssetEditor/Source/ChaosClothAssetEditor/Private/ChaosClothAsset/ClothEditorMode.cpp
+++ b/Engine/Plugins/ChaosClothAssetEditor/Source/ChaosClothAssetEditor/Private/ChaosClothAsset/ClothEditorMode.cpp
@@ -1415,9 +1415,9 @@ void UChaosClothAssetEditorMode::DeleteContextObject()
 
 void UChaosClothAssetEditorMode::SetDataflowGraphEditor(TSharedPtr<SDataflowGraphEditor> InGraphEditor)
 {
+	DataflowGraphEditor = InGraphEditor;
 	if (InGraphEditor)
 	{
-		DataflowGraphEditor = InGraphEditor;
 		InitializeContextObject();
 	}
 	else
@@ -1495,11 +1495,13 @@ UEdGraphNode* UChaosClothAssetEditorMode::CreateNewNode(const FName& NewNodeType
 		return nullptr;
 	}
 
+	checkf(DataflowGraph.IsValid(), TEXT("Dataflow pointer is invalid in UChaosClothAssetEditorMode"));
+
 	const TSharedPtr<FAssetSchemaAction_Dataflow_CreateNode_DataflowEdNode> NodeAction =
-		FAssetSchemaAction_Dataflow_CreateNode_DataflowEdNode::CreateAction(DataflowGraph, NewNodeTypeName);
+		FAssetSchemaAction_Dataflow_CreateNode_DataflowEdNode::CreateAction(DataflowGraph.Get(), NewNodeTypeName);
 	constexpr UEdGraphPin* FromPin = nullptr;
 	constexpr bool bSelectNewNode = true;
-	UEdGraphNode* const NewEdNode = NodeAction->PerformAction(DataflowGraph, FromPin, PinnedDataflowGraphEditor->GetPasteLocation(), bSelectNewNode);
+	UEdGraphNode* const NewEdNode = NodeAction->PerformAction(DataflowGraph.Get(), FromPin, PinnedDataflowGraphEditor->GetPasteLocation(), bSelectNewNode);
 
 	return NewEdNode;
 }
@@ -1507,6 +1509,8 @@ UEdGraphNode* UChaosClothAssetEditorMode::CreateNewNode(const FName& NewNodeType
 
 UEdGraphNode* UChaosClothAssetEditorMode::CreateAndConnectNewNode(const FName& NewNodeTypeName, UEdGraphNode& UpstreamNode, const FName& ConnectionTypeName, const FName& NewNodeConnectionName)
 {
+	checkf(DataflowGraph.IsValid(), TEXT("Dataflow pointer is invalid in UChaosClothAssetEditorMode"));
+
 	// First find the specified output of the upstream node, plus any pins it's connected to
 
 	UEdGraphPin* UpstreamNodeOutputPin = nullptr;
diff --git a/Engine/Plugins/ChaosClothAssetEditor/Source/ChaosClothAssetEditor/Private/ChaosClothAsset/ClothEditorMode.h b/Engine/Plugins/ChaosClothAssetEditor/Source/ChaosClothAssetEditor/Private/ChaosClothAsset/ClothEditorMode.h
--- a/Engine/Plugins/ChaosClothAssetEditor/Source/ChaosClothAssetEditor/Private/ChaosClothAsset/ClothEditorMode.h
+++ b/Engine/Plugins/ChaosClothAssetEditor/Source/ChaosClothAssetEditor/Private/ChaosClothAsset/ClothEditorMode.h
@@ -292,8 +292,7 @@ private:
 	UPROPERTY()
 	TObjectPtr<UDataflowComponent> DataflowComponent = nullptr;
 
-	UPROPERTY()
-	TObjectPtr<UDataflow> DataflowGraph = nullptr;
+	TWeakObjectPtr<UDataflow> DataflowGraph = nullptr;
 
 	TWeakPtr<SDataflowGraphEditor> DataflowGraphEditor;
 
diff --git a/Engine/Plugins/ChaosClothAssetEditor/Source/ChaosClothAssetEditor/Private/ChaosClothAsset/ClothEditorToolkit.cpp b/Engine/Plugins/ChaosClothAssetEditor/Source/ChaosClothAssetEditor/Private/ChaosClothAsset/ClothEditorToolkit.cpp
--- a/Engine/Plugins/ChaosClothAssetEditor/Source/ChaosClothAssetEditor/Private/ChaosClothAsset/ClothEditorToolkit.cpp
+++ b/Engine/Plugins/ChaosClothAssetEditor/Source/ChaosClothAssetEditor/Private/ChaosClothAsset/ClothEditorToolkit.cpp
@@ -42,6 +42,7 @@
 #include "DynamicMesh/DynamicMesh3.h"
 #include "Toolkits/AssetEditorToolkitMenuContext.h"
 #include "FileHelpers.h"
+#include "UObject/PackageReload.h"
 
 #define LOCTEXT_NAMESPACE "ChaosClothAssetEditorToolkit"
 
@@ -398,6 +399,8 @@ bool FChaosClothAssetEditorToolkit::OnRequestClose(EAssetEditorCloseReason InClo
 	}
 	SelectedDataflowNode.Reset();
 
+	FCoreUObjectDelegates::OnPackageReloaded.Remove(OnPackageReloadedDelegateHandle);
+
 	return FAssetEditorToolkit::OnRequestClose(InCloseReason);
 }
 
@@ -491,6 +494,9 @@ void FChaosClothAssetEditorToolkit::PostInitAssetEditor()
 
 	ClothMode->DataflowGraph = GetDataflow();
 	ClothMode->SetDataflowGraphEditor(GraphEditor);
+
+	// Handle Dataflow asset reload event
+	OnPackageReloadedDelegateHandle = FCoreUObjectDelegates::OnPackageReloaded.AddSP(this, &FChaosClothAssetEditorToolkit::HandlePackageReloaded);
 }
 
 void FChaosClothAssetEditorToolkit::InitToolMenuContext(FToolMenuContext& MenuContext)
@@ -985,6 +991,7 @@ void FChaosClothAssetEditorToolkit::ReinitializeGraphEditorWidget()
 	SGraphEditor::FGraphEditorEvents InEvents;
 	InEvents.OnVerifyTextCommit = FOnNodeVerifyTextCommit::CreateSP(this, &FChaosClothAssetEditorToolkit::OnNodeVerifyTitleCommit);
 	InEvents.OnTextCommitted = FOnNodeTextCommitted::CreateSP(this, &FChaosClothAssetEditorToolkit::OnNodeTitleCommitted);
+	InEvents.OnNodeSingleClicked = SGraphEditor::FOnNodeSingleClicked::CreateSP(this, &FChaosClothAssetEditorToolkit::OnNodeSingleClicked);
 
 	UChaosClothAsset* const ClothAsset = GetAsset();
 
@@ -1012,7 +1019,9 @@ void FChaosClothAssetEditorToolkit::ReinitializeGraphEditorWidget()
 	GraphEditor->Construct(Args, ClothAsset);
 
 	GraphEditor->OnSelectionChangedMulticast.RemoveAll(this);
+	GraphEditor->OnNodeDeletedMulticast.RemoveAll(this);
 	GraphEditor->OnSelectionChangedMulticast.AddSP(this, &FChaosClothAssetEditorToolkit::OnNodeSelectionChanged);
+	GraphEditor->OnNodeDeletedMulticast.AddSP(this, &FChaosClothAssetEditorToolkit::OnNodeDeleted);
 }
 
 TSharedPtr<IStructureDetailsView> FChaosClothAssetEditorToolkit::CreateNodeDetailsEditorWidget(UObject* ObjectToEdit)
@@ -1339,6 +1348,62 @@ void FChaosClothAssetEditorToolkit::InvalidateViews()
 	ViewportClient->Invalidate();
 	ClothPreviewViewportClient->Invalidate();
 }
+
+void FChaosClothAssetEditorToolkit::HandlePackageReloaded(const EPackageReloadPhase InPackageReloadPhase, FPackageReloadedEvent* InPackageReloadedEvent)
+{
+	// Handle the Dataflow asset being reloaded
+	// 
+	// UAssetEditorSubsystem::HandlePackageReloaded deals with restarting the corresponding asset editor when an asset is reloaded. However in our case we have the Dataflow Asset 
+	// open inside the Cloth Asset Editor and the AssetEditorSubsystem is not tracking that information. So we will handle it ourselves here. 
+	// 
+	// We will not restart the entire Cloth Editor in this case, but we will rebuild the Graph Editor Widget.
+	// 
+	// Note we don't need to update the Dataflow pointer in the ClothAsset as the reload process does that -- we just need to know when that pointer changes.
+
+	// TODO: Invesigate if we can somehow tell UAssetEditorSubsystem that when the DataflowAsset is reloaded we should restart the Cloth Editor
+
+	if (InPackageReloadPhase == EPackageReloadPhase::PrePackageFixup)	  // Listen for pre-reload event
+	{
+		checkf(InPackageReloadedEvent, TEXT("Expected a FPackageReloadedEvent object on PrePackageFixup phase"));
+
+		for (const TPair<UObject*, UObject*>& RepointPair : InPackageReloadedEvent->GetRepointedObjects())
+		{
+			if (RepointPair.Key == GetDataflow())
+			{
+				// Close any open tool
+				UChaosClothAssetEditorMode* const ClothMode = CastChecked<UChaosClothAssetEditorMode>(EditorModeManager->GetActiveScriptableMode(UChaosClothAssetEditorMode::EM_ChaosClothAssetEditorModeId));
+				UInteractiveToolManager* const ToolManager = ClothMode->GetToolManager();
+				if (UInteractiveTool* const ActiveTool = ToolManager->GetActiveTool(EToolSide::Left))
+				{
+					ToolManager->PostActiveToolShutdownRequest(ActiveTool, EToolShutdownType::Accept);
+				}
+
+				// Eliminate anything could be holding a reference to the Dataflow object that will be reloaded (including references to nodes)
+				SelectedDataflowNode.Reset();
+				GraphEditor.Reset();
+				GraphEditorTab.Get()->SetContent(SNew(SSpacer));
+			}
+		}
+	}
+	else if (InPackageReloadPhase == EPackageReloadPhase::PostPackageFixup)	  // Listen for post-reload event
+	{
+		checkf(InPackageReloadedEvent, TEXT("Expected a FPackageReloadedEvent object on PostPackageFixup phase"));
+
+		for (const TPair<UObject*, UObject*>& RepointPair : InPackageReloadedEvent->GetRepointedObjects())
+		{
+			if (RepointPair.Key == GetDataflow())	// On PostPackageFixup, both Key and Value point to the new Object
+			{
+				ReinitializeGraphEditorWidget();
+
+				UChaosClothAssetEditorMode* const ClothMode = CastChecked<UChaosClothAssetEditorMode>(EditorModeManager->GetActiveScriptableMode(UChaosClothAssetEditorMode::EM_ChaosClothAssetEditorModeId));
+				ClothMode->DataflowGraph = GetDataflow();
+				ClothMode->SetDataflowGraphEditor(GraphEditor);
+			}
+		}
+	}
+}
+
+
 } // namespace UE::Chaos::ClothAsset
 
 #undef LOCTEXT_NAMESPACE
diff --git a/Engine/Plugins/ChaosClothAssetEditor/Source/ChaosClothAssetEditor/Private/ChaosClothAsset/ClothEditorToolkit.h b/Engine/Plugins/ChaosClothAssetEditor/Source/ChaosClothAssetEditor/Private/ChaosClothAsset/ClothEditorToolkit.h
--- a/Engine/Plugins/ChaosClothAssetEditor/Source/ChaosClothAssetEditor/Private/ChaosClothAsset/ClothEditorToolkit.h
+++ b/Engine/Plugins/ChaosClothAssetEditor/Source/ChaosClothAssetEditor/Private/ChaosClothAsset/ClothEditorToolkit.h
@@ -139,6 +139,9 @@ private:
 	void OnNodeDeleted(const TSet<UObject*>& DeletedNodes);
 	void OnNodeSingleClicked(UObject* ClickedNode) const;
 
+	// Delegates
+	void HandlePackageReloaded(const EPackageReloadPhase InPackageReloadPhase, FPackageReloadedEvent* InPackageReloadedEvent);
+
 	/** Scene in which the 3D sim space preview meshes live. Ownership shared with AdvancedPreviewSettingsWidget*/
 	TSharedPtr<FChaosClothPreviewScene> ClothPreviewScene;
 
@@ -171,5 +174,8 @@ private:
 	static const FName NodeDetailsTabId;
 	TSharedPtr<SDockTab> NodeDetailsTab;
 	TSharedPtr<IStructureDetailsView> NodeDetailsEditor;
+
+	FDelegateHandle OnPackageReloadedDelegateHandle;
+	
 };
 } // namespace UE::Chaos::ClothAsset
