Add basic telemetry to the Cloth Editor: editor mode enter/exit and tool start/stop information

diff --git a/Engine/Plugins/ChaosClothAssetEditor/Source/ChaosClothAssetEditor/Private/ChaosClothAsset/ClothEditorMode.cpp b/Engine/Plugins/ChaosClothAssetEditor/Source/ChaosClothAssetEditor/Private/ChaosClothAsset/ClothEditorMode.cpp
--- a/Engine/Plugins/ChaosClothAssetEditor/Source/ChaosClothAssetEditor/Private/ChaosClothAsset/ClothEditorMode.cpp
+++ b/Engine/Plugins/ChaosClothAssetEditor/Source/ChaosClothAssetEditor/Private/ChaosClothAsset/ClothEditorMode.cpp
@@ -18,6 +18,7 @@
 #include "Drawing/MeshElementsVisualizer.h"
 #include "EditorViewportClient.h"
 #include "EdModeInteractiveToolsContext.h"
+#include "EngineAnalytics.h"
 #include "Framework/Commands/UICommandList.h"
 #include "InteractiveTool.h"
 #include "ModelingToolTargetUtil.h"
@@ -109,6 +110,11 @@ namespace UE::Chaos::ClothAsset::Private
 		return FLinearColor::MakeFromHSV8(Seed, 180, 140);
 	}
 
+	FString GetToolName(const UInteractiveTool& Tool)
+	{
+		const FString* const ToolName = FTextInspector::GetSourceString(Tool.GetToolInfo().ToolDisplayName);
+		return ToolName ? *ToolName : FString(TEXT("<Invalid ToolName>"));
+	}
 }
 
 
@@ -138,6 +144,52 @@ void UChaosClothAssetEditorMode::Enter()
 
 	// Register gizmo ContextObject for use inside interactive tools
 	UE::TransformGizmoUtil::RegisterTransformGizmoContextObject(GetInteractiveToolsContext());
+
+
+	//
+	// Engine Analytics
+	//
+
+	// Log mode starting
+	if (FEngineAnalytics::IsAvailable())
+	{
+		LastModeStartTimestamp = FDateTime::UtcNow();
+		TArray<FAnalyticsEventAttribute> EventAttributes;
+		EventAttributes.Add(FAnalyticsEventAttribute(TEXT("Timestamp"), LastModeStartTimestamp.ToString()));
+		FEngineAnalytics::GetProvider().RecordEvent(TEXT("Editor.Usage.ChaosClothEditor.Enter"), EventAttributes);
+
+		// Log tool starting
+		GetToolManager()->OnToolStarted.AddLambda([this](UInteractiveToolManager* Manager, UInteractiveTool* Tool)
+		{
+			if (FEngineAnalytics::IsAvailable() && Tool)
+			{
+				LastToolStartTimestamp = FDateTime::UtcNow();
+
+				TArray<FAnalyticsEventAttribute> EventAttributes;
+				EventAttributes.Add(FAnalyticsEventAttribute(TEXT("ToolName"), UE::Chaos::ClothAsset::Private::GetToolName(*Tool)));
+				EventAttributes.Add(FAnalyticsEventAttribute(TEXT("Timestamp"), LastToolStartTimestamp.ToString()));
+
+				FEngineAnalytics::GetProvider().RecordEvent(TEXT("Editor.Usage.ChaosClothEditor.ToolStarted"), EventAttributes);
+			}
+		});
+
+		// Log tool ending
+		GetToolManager()->OnToolEnded.AddLambda([this](UInteractiveToolManager* Manager, UInteractiveTool* Tool)
+		{
+			if (FEngineAnalytics::IsAvailable() && Tool)
+			{
+				const FDateTime Now = FDateTime::UtcNow();
+				const FTimespan ToolUsageDuration = Now - LastToolStartTimestamp;
+
+				TArray<FAnalyticsEventAttribute> EventAttributes;
+				EventAttributes.Add(FAnalyticsEventAttribute(TEXT("ToolName"), UE::Chaos::ClothAsset::Private::GetToolName(*Tool)));
+				EventAttributes.Add(FAnalyticsEventAttribute(TEXT("Timestamp"), Now.ToString()));
+				EventAttributes.Add(FAnalyticsEventAttribute(TEXT("Duration.Seconds"), static_cast<float>(ToolUsageDuration.GetTotalSeconds())));
+
+				FEngineAnalytics::GetProvider().RecordEvent(TEXT("Editor.Usage.ChaosClothEditor.ToolEnded"), EventAttributes);
+			}
+		});
+	}
 }
 
 void UChaosClothAssetEditorMode::AddToolTargetFactories()
@@ -406,6 +458,21 @@ void UChaosClothAssetEditorMode::Exit()
 	PropertyObjectsToTick.Empty();
 	PreviewScene = nullptr;
 
+	//
+	// Engine Analytics
+	//
+	// Log mode exit
+	if (FEngineAnalytics::IsAvailable())
+	{
+		const FTimespan ModeUsageDuration = FDateTime::UtcNow() - LastModeStartTimestamp;
+
+		TArray<FAnalyticsEventAttribute> Attributes;
+		Attributes.Add(FAnalyticsEventAttribute(TEXT("Timestamp"), FDateTime::UtcNow().ToString()));
+		Attributes.Add(FAnalyticsEventAttribute(TEXT("Duration.Seconds"), static_cast<float>(ModeUsageDuration.GetTotalSeconds())));
+
+		FEngineAnalytics::GetProvider().RecordEvent(TEXT("Editor.Usage.ChaosClothEditor.Exit"));
+	}
+
 	Super::Exit();
 }
 
diff --git a/Engine/Plugins/ChaosClothAssetEditor/Source/ChaosClothAssetEditor/Private/ChaosClothAsset/ClothEditorMode.h b/Engine/Plugins/ChaosClothAssetEditor/Source/ChaosClothAssetEditor/Private/ChaosClothAsset/ClothEditorMode.h
--- a/Engine/Plugins/ChaosClothAssetEditor/Source/ChaosClothAssetEditor/Private/ChaosClothAsset/ClothEditorMode.h
+++ b/Engine/Plugins/ChaosClothAssetEditor/Source/ChaosClothAssetEditor/Private/ChaosClothAsset/ClothEditorMode.h
@@ -304,5 +304,9 @@ private:
 
 	// Correspondence between node types and commands to launch tools
 	TMap<FName, TSharedPtr<const FUICommandInfo>> NodeTypeToToolCommandMap;
+
+	// Timestamps for telemetry
+	FDateTime LastModeStartTimestamp;
+	FDateTime LastToolStartTimestamp;
 };
 
