Check that mesh description exists before converting LOD data to Dynamic Mesh.

#rnx

diff --git a/Engine/Plugins/ChaosClothAssetEditor/Source/ChaosClothAssetDataflowNodes/Private/ChaosClothAsset/TransferSkinWeightsNode.cpp b/Engine/Plugins/ChaosClothAssetEditor/Source/ChaosClothAssetDataflowNodes/Private/ChaosClothAsset/TransferSkinWeightsNode.cpp
--- a/Engine/Plugins/ChaosClothAssetEditor/Source/ChaosClothAssetDataflowNodes/Private/ChaosClothAsset/TransferSkinWeightsNode.cpp
+++ b/Engine/Plugins/ChaosClothAssetEditor/Source/ChaosClothAssetDataflowNodes/Private/ChaosClothAsset/TransferSkinWeightsNode.cpp
@@ -748,6 +748,19 @@ void FChaosClothAssetTransferSkinWeightsNode::Evaluate(Dataflow::FContext& Conte
 				return;
 			}
 
+			if (!SkeletalMesh->HasMeshDescription(LodIndex))
+			{
+				//TODO: Once the support for converting render data to dynamic meshes is added we can support auto LODs
+				FClothDataflowTools::LogAndToastWarning(*this,
+					LOCTEXT("AutogeneratedLodIndexHeadline", "Autogenerated LOD not supported."),
+					FText::Format(
+						LOCTEXT("AutogeneratedLodIndexDetails", "Auto generated LOD index {0} for skeletal mesh {1} is not supported."),
+						LodIndex,
+						FText::FromString(SkeletalMesh.GetName())));
+				SetValue(Context, MoveTemp(*ClothCollection), &Collection);
+				return;
+			}
+
 			//
 			// Convert source Skeletal Mesh to Dynamic Mesh.
 			//
