Chaos Cloth Editor: use input collection rather than output collection on tools. Add ability to edit an existing weight map.

diff --git a/Engine/Plugins/ChaosClothAssetEditor/Source/ChaosClothAssetDataflowNodes/Private/ChaosClothAsset/AddWeightMapNode.cpp b/Engine/Plugins/ChaosClothAssetEditor/Source/ChaosClothAssetDataflowNodes/Private/ChaosClothAsset/AddWeightMapNode.cpp
--- a/Engine/Plugins/ChaosClothAssetEditor/Source/ChaosClothAssetDataflowNodes/Private/ChaosClothAsset/AddWeightMapNode.cpp
+++ b/Engine/Plugins/ChaosClothAssetEditor/Source/ChaosClothAssetDataflowNodes/Private/ChaosClothAsset/AddWeightMapNode.cpp
@@ -127,12 +127,93 @@ namespace UE::Chaos::ClothAsset::Private
 			OutTargetWeights[TargetWeightsLookup[Index]] = TargetWeights[Index];
 		}
 	}
+
+
+	void SetVertexWeights(const TConstArrayView<float> InputMap, const TArray<float>& FinalValues, EChaosClothAssetWeightMapOverrideType OverrideType, TArray<float>& SourceVertexWeights)
+	{
+		if (InputMap.IsEmpty() || OverrideType == EChaosClothAssetWeightMapOverrideType::ReplaceAll)
+		{
+			// Default input is 0, so OverrideType doesn't matter.
+			SourceVertexWeights = FinalValues;
+			return;
+		}
+
+		check(InputMap.Num() == FinalValues.Num());
+		SourceVertexWeights.SetNumUninitialized(FinalValues.Num());
+		for (int32 Index = 0; Index < FinalValues.Num(); ++Index)
+		{
+			switch (OverrideType)
+			{
+			case EChaosClothAssetWeightMapOverrideType::ReplaceChanged:
+				if (InputMap[Index] == FinalValues[Index])
+				{
+					SourceVertexWeights[Index] = FChaosClothAssetAddWeightMapNode::ReplaceChangedPassthroughValue;
+				}
+				else
+				{
+					SourceVertexWeights[Index] = FinalValues[Index];
+				}
+				break;
+			case EChaosClothAssetWeightMapOverrideType::Add:
+				SourceVertexWeights[Index] = FinalValues[Index] - InputMap[Index];
+				break;
+			default: unimplemented();
+			}
+		}
+	}
+
+	void CalculateFinalVertexWeightValues(const TConstArrayView<float> InputMap, TArrayView<float> FinalOutputMap, EChaosClothAssetWeightMapOverrideType OverrideType, const TArray<float>& SourceVertexWeights)
+	{
+		const int32 EndWeightIndex = FMath::Min(FinalOutputMap.Num(), SourceVertexWeights.Num());
+		if (InputMap.IsEmpty())
+		{
+			for (int32 Index = 0; Index < EndWeightIndex; ++Index)
+			{
+				FinalOutputMap[Index] = FMath::Clamp(SourceVertexWeights[Index] == FChaosClothAssetAddWeightMapNode::ReplaceChangedPassthroughValue ?
+					0.f : SourceVertexWeights[Index], 0.f, 1.f);
+			}
+			return;
+		}
+
+		check(InputMap.Num() == FinalOutputMap.Num());
+		for (int32 Index = 0; Index < EndWeightIndex; ++Index)
+		{
+			if(SourceVertexWeights[Index] == FChaosClothAssetAddWeightMapNode::ReplaceChangedPassthroughValue)
+			{
+				// This value is only set when OverrideType == ReplaceChanged, but it's possible the override type changed.
+				FinalOutputMap[Index] = FMath::Clamp(InputMap[Index], 0.f, 1.f);
+			}
+			else
+			{
+				switch (OverrideType)
+				{
+				case EChaosClothAssetWeightMapOverrideType::ReplaceAll:
+				case EChaosClothAssetWeightMapOverrideType::ReplaceChanged:
+					FinalOutputMap[Index] = FMath::Clamp(SourceVertexWeights[Index], 0.f, 1.f);
+					break;
+				case EChaosClothAssetWeightMapOverrideType::Add:
+					FinalOutputMap[Index] = FMath::Clamp(InputMap[Index] + SourceVertexWeights[Index], 0.f, 1.f);
+					break;
+				default: unimplemented();
+				}
+			}
+		}
+		if (InputMap.GetData() != FinalOutputMap.GetData())
+		{
+			// Fill in remaining values with InputMap
+			for (int32 Index = EndWeightIndex; Index < FinalOutputMap.Num(); ++Index)
+			{
+				FinalOutputMap[Index] = FMath::Clamp(InputMap[Index], 0.f, 1.f);
+			}
+		}
+	}
 }
 
 FChaosClothAssetAddWeightMapNode::FChaosClothAssetAddWeightMapNode(const Dataflow::FNodeParameters& InParam, FGuid InGuid)
 	: FDataflowTerminalNode(InParam, InGuid)
 {
 	RegisterInputConnection(&Collection);
+	RegisterInputConnection(&InputName.StringValue, GET_MEMBER_NAME_CHECKED(FChaosClothAssetConnectableIStringValue, StringValue));
 	RegisterInputConnection(&TransferCollection);
 	RegisterOutputConnection(&Collection, &Collection);
 	RegisterOutputConnection(&Name);
@@ -167,9 +248,9 @@ PRAGMA_ENABLE_DEPRECATION_WARNINGS
 					const TSharedRef<const FManagedArrayCollection> TransferClothCollection = MakeShared<const FManagedArrayCollection>(MoveTemp(InTransferCollection));
 					FCollectionClothConstFacade TransferClothFacade(TransferClothCollection);
 
-					const FName InName(Name);
-					const uint32 NameTypeHash = HashCombineFast(GetTypeHash(InName), (uint32)TransferType);
-					const uint32 InTransferCollectionHash = (TransferClothFacade.IsValid() && InName != NAME_None) ?
+					const FName InInputName = GetInputName(Context);
+					const uint32 NameTypeHash = HashCombineFast(GetTypeHash(InInputName), (uint32)TransferType);
+					const uint32 InTransferCollectionHash = (TransferClothFacade.IsValid() && InInputName != NAME_None) ?
 						HashCombineFast(TransferClothFacade.CalculateWeightMapTypeHash(), NameTypeHash) : 0;  // TODO: Remove after adding the function (currently shelved!)  
 				
 					if (TransferCollectionHash != InTransferCollectionHash)
@@ -178,10 +259,11 @@ PRAGMA_ENABLE_DEPRECATION_WARNINGS
 
 						if (TransferCollectionHash)
 						{
-							if (TransferClothFacade.HasWeightMap(InName))
+							if (TransferClothFacade.HasWeightMap(InInputName))
 							{
 								// Remap the weights
-								MutableThis->GetVertexWeights().SetNumZeroed(ClothFacade.GetNumSimVertices3D());
+								TArray<float> RemappedWeights;
+								RemappedWeights.SetNumZeroed(ClothFacade.GetNumSimVertices3D());
 
 								switch (TransferType)
 								{
@@ -190,24 +272,26 @@ PRAGMA_ENABLE_DEPRECATION_WARNINGS
 										TransferClothFacade.GetSimPosition2D(),
 										TransferClothFacade.GetSimIndices2D(),
 										TransferClothFacade.GetSimVertex3DLookup(),
-										TransferClothFacade.GetWeightMap(InName),
+										TransferClothFacade.GetWeightMap(InInputName),
 										ClothFacade.GetSimPosition2D(),
 										ClothFacade.GetSimIndices2D(),
 										ClothFacade.GetSimVertex3DLookup(),
-										MutableThis->GetVertexWeights());
+										RemappedWeights);
 										break;
 								case EChaosClothAssetWeightMapTransferType::Use3DSimMesh:
 									Private::TransferWeightMap(
 										TransferClothFacade.GetSimPosition3D(),
 										TransferClothFacade.GetSimIndices3D(),
-										TransferClothFacade.GetWeightMap(InName),
+										TransferClothFacade.GetWeightMap(InInputName),
 										ClothFacade.GetSimPosition3D(),
 										ClothFacade.GetSimNormal(),
 										ClothFacade.GetSimIndices3D(),
-										MutableThis->GetVertexWeights());
+										RemappedWeights);
 										break;
 								default: unimplemented();
 								}
+
+								MutableThis->SetVertexWeights(ClothFacade.GetWeightMap(InInputName), RemappedWeights);
 							}
 						}
 					}
@@ -221,9 +305,8 @@ void FChaosClothAssetAddWeightMapNode::Evaluate(Dataflow::FContext& Context, con
 {
 	using namespace UE::Chaos::ClothAsset;
 
-	auto CopyWeightsIntoClothCollection = [this](TArrayView<float>& ClothWeights, const TArray<float>& SourceVertexWeights, bool bIsSim)
+	auto CheckSourceVertexWeights = [this](TArrayView<float>& ClothWeights, const TArray<float>& SourceVertexWeights, bool bIsSim)
 	{
-		const int32 MaxWeightIndex = FMath::Min(SourceVertexWeights.Num(), ClothWeights.Num());
 		if (SourceVertexWeights.Num() > 0 && SourceVertexWeights.Num() != ClothWeights.Num())
 		{
 			FClothDataflowTools::LogAndToastWarning(*this,
@@ -233,22 +316,21 @@ void FChaosClothAssetAddWeightMapNode::Evaluate(Dataflow::FContext& Context, con
 					SourceVertexWeights.Num(),
 					ClothWeights.Num()));
 		}
-
-		for (int32 VertexID = 0; VertexID < MaxWeightIndex; ++VertexID)
-		{
-			ClothWeights[VertexID] = SourceVertexWeights[VertexID];
-		}
 	};
 
 	if (Out->IsA<FManagedArrayCollection>(&Collection))
 	{
+		// Evaluate InputName
+		const FName InInputName = GetInputName(Context);
+		InputName.StringValue_Override = GetValue<FString>(Context, &InputName.StringValue, UE::Chaos::ClothAsset::FWeightMapTools::NotOverridden);
+
 		// Evaluate in collection
 		FManagedArrayCollection InCollection = GetValue<FManagedArrayCollection>(Context, &Collection);
 		const TSharedRef<FManagedArrayCollection> ClothCollection = MakeShared<FManagedArrayCollection>(MoveTemp(InCollection));
 		FCollectionClothFacade ClothFacade(ClothCollection);
 		if (ClothFacade.IsValid())  // Can only act on the collection if it is a valid cloth collection
 		{
-			const FName InName(Name);
+			const FName InName(Name.IsEmpty() ? InInputName : FName(Name));
 
 			// Copy simulation weights into cloth collection
 
@@ -268,7 +350,8 @@ void FChaosClothAssetAddWeightMapNode::Evaluate(Dataflow::FContext& Context, con
 				else
 				{
 					constexpr bool bIsSim = true;
-					CopyWeightsIntoClothCollection(ClothSimWeights, GetVertexWeights(), bIsSim);
+					CheckSourceVertexWeights(ClothSimWeights, GetVertexWeights(), bIsSim);
+					CalculateFinalVertexWeightValues(ClothFacade.GetWeightMap(InInputName), ClothSimWeights);
 				}
 			}
 			
@@ -290,7 +373,8 @@ void FChaosClothAssetAddWeightMapNode::Evaluate(Dataflow::FContext& Context, con
 				else
 				{
 					constexpr bool bIsSim = false;
-					CopyWeightsIntoClothCollection(ClothRenderWeights, GetRenderVertexWeights(), bIsSim);
+					CheckSourceVertexWeights(ClothRenderWeights, GetRenderVertexWeights(), bIsSim);
+					CalculateFinalRenderVertexWeightValues(ClothFacade.GetUserDefinedAttribute<float>(InInputName, ClothCollectionGroup::RenderVertices), ClothRenderWeights);
 				}
 			}
 
@@ -299,8 +383,37 @@ void FChaosClothAssetAddWeightMapNode::Evaluate(Dataflow::FContext& Context, con
 	}
 	else if (Out->IsA<FString>(&Name))
 	{
-		SetValue(Context, Name, &Name);
+		FString InputNameString = GetValue<FString>(Context, &InputName.StringValue);
+		UE::Chaos::ClothAsset::FWeightMapTools::MakeWeightMapName(InputNameString);
+		SetValue(Context, InputNameString.IsEmpty() ? Name : InputNameString, &Name);
 	}
 }
 
+FName FChaosClothAssetAddWeightMapNode::GetInputName(Dataflow::FContext& Context) const
+{
+	FString InputNameString = GetValue<FString>(Context, &InputName.StringValue);
+	UE::Chaos::ClothAsset::FWeightMapTools::MakeWeightMapName(InputNameString);
+	const FName InInputName(*InputNameString);
+	return InInputName != NAME_None ? InInputName : FName(Name);
+}
+
+void FChaosClothAssetAddWeightMapNode::SetVertexWeights(const TConstArrayView<float> InputMap, const TArray<float>& FinalValues)
+{
+	UE::Chaos::ClothAsset::Private::SetVertexWeights(InputMap, FinalValues, MapOverrideType, GetVertexWeights());
+}
+
+void FChaosClothAssetAddWeightMapNode::SetRenderVertexWeights(const TConstArrayView<float> InputMap, const TArray<float>& FinalValues)
+{
+	UE::Chaos::ClothAsset::Private::SetVertexWeights(InputMap, FinalValues, MapOverrideType, GetRenderVertexWeights());
+}
+
+void FChaosClothAssetAddWeightMapNode::CalculateFinalVertexWeightValues(const TConstArrayView<float> InputMap, TArrayView<float> FinalOutputMap) const
+{
+	UE::Chaos::ClothAsset::Private::CalculateFinalVertexWeightValues(InputMap, FinalOutputMap, MapOverrideType, GetVertexWeights());
+}
+
+void FChaosClothAssetAddWeightMapNode::CalculateFinalRenderVertexWeightValues(const TConstArrayView<float> InputMap, TArrayView<float> FinalOutputMap) const
+{
+	UE::Chaos::ClothAsset::Private::CalculateFinalVertexWeightValues(InputMap, FinalOutputMap, MapOverrideType, GetRenderVertexWeights());
+}
 #undef LOCTEXT_NAMESPACE
diff --git a/Engine/Plugins/ChaosClothAssetEditor/Source/ChaosClothAssetDataflowNodes/Public/ChaosClothAsset/AddWeightMapNode.h b/Engine/Plugins/ChaosClothAssetEditor/Source/ChaosClothAssetDataflowNodes/Public/ChaosClothAsset/AddWeightMapNode.h
--- a/Engine/Plugins/ChaosClothAssetEditor/Source/ChaosClothAssetDataflowNodes/Public/ChaosClothAsset/AddWeightMapNode.h
+++ b/Engine/Plugins/ChaosClothAssetEditor/Source/ChaosClothAssetDataflowNodes/Public/ChaosClothAsset/AddWeightMapNode.h
@@ -4,6 +4,7 @@
 
 #include "Dataflow/DataflowTerminalNode.h"
 #include "GeometryCollection/ManagedArrayCollection.h"
+#include "ChaosClothAsset/ConnectableValue.h"
 #include "AddWeightMapNode.generated.h"
 
 UENUM()
@@ -24,6 +25,18 @@ enum class EChaosClothAssetWeightMapMeshType : uint8
 	Both
 };
 
+/** How the map stored on the AddWeightMapNode should be applied to an existing map. If no map exists, it is treated as zero.*/
+UENUM()
+enum class EChaosClothAssetWeightMapOverrideType : uint8
+{
+	/** The full map is stored and reapplied.*/
+	ReplaceAll,
+	/** Only changed values are stored and reapplied. */
+	ReplaceChanged,
+	/** Add values. */
+	Add
+};
+
 
 /** Painted weight map attributes node. */
 PRAGMA_DISABLE_DEPRECATION_WARNINGS  // For deprecated VertexWeights in copy constructor (Clang)
@@ -34,16 +47,17 @@ struct FChaosClothAssetAddWeightMapNode : public FDataflowTerminalNode
 	DATAFLOW_NODE_DEFINE_INTERNAL(FChaosClothAssetAddWeightMapNode, "AddWeightMap", "Cloth", "Cloth Add Weight Map")
 
 public:
+	static constexpr float ReplaceChangedPassthroughValue = UE_BIG_NUMBER;
 
 	UPROPERTY(Meta = (DataflowInput, DataflowOutput, DataflowPassthrough = "Collection"))
 	FManagedArrayCollection Collection;
 
 	/**
 	 * The collection used to transfer weight map from.
-	 * Connecting a collection containing a weight map sharing the same name as the one
-	 * created with this node will transfer the weights to the input collection vertices.
+	 * Connecting a collection containing a weight map with Input Name (or Name if Input Name is empty) 
+	 * will transfer the weights to the input collection vertices.
 	 * Note this operation only happens once when the TransferCollection is first connected, or updated.
-	 * Changing the Name or the TransferType will also redo the transfer operation.
+	 * Changing the InputName or the TransferType will also redo the transfer operation.
 	 */
 	UPROPERTY(Meta = (DataflowInput))
 	FManagedArrayCollection TransferCollection;
@@ -52,6 +66,15 @@ public:
 	UPROPERTY(EditAnywhere, Category = "Add Weight Map", Meta = (DataflowOutput))
 	FString Name;
 
+	/** The name to populate this map from and override based on Map Override Type. Name will be used if Input Name is empty.*/
+	UPROPERTY(EditAnywhere, Category = "Add Weight Map")
+	FChaosClothAssetConnectableIStringValue InputName = {TEXT("")};
+
+	/** How to apply this node's weight values onto existing maps. Changing this value will change the output map. 
+	 *  To change how the node's stored weights are calculated, change the equivalent value on the Weight Map Paint Tool context.*/
+	UPROPERTY(EditAnywhere, Category = "Add Weight Map")
+	EChaosClothAssetWeightMapOverrideType MapOverrideType = EChaosClothAssetWeightMapOverrideType::ReplaceAll;
+
 	/**
 	 * The type of transfer used to transfer the weight map when a TransferCollection is connected.
 	 * This property is disabled when no TransferCollection input has been conencted.
@@ -80,6 +103,16 @@ private:
 	UPROPERTY()
 	TArray<float> RenderVertexWeights;
 
+	// These methods are exported for UClothEditorWeightMapPaintTool which lives in a different module.
+	FName CHAOSCLOTHASSETDATAFLOWNODES_API GetInputName(Dataflow::FContext& Context) const;
+
+	void CHAOSCLOTHASSETDATAFLOWNODES_API SetVertexWeights(const TConstArrayView<float> InputMap, const TArray<float>& FinalValues);
+	void CHAOSCLOTHASSETDATAFLOWNODES_API SetRenderVertexWeights(const TConstArrayView<float> InputMap, const TArray<float>& FinalValues);
+
+	// Input and FinalOutputMap can be the same array, but should not otherwise be interleaved.
+	void CHAOSCLOTHASSETDATAFLOWNODES_API CalculateFinalVertexWeightValues(const TConstArrayView<float> InputMap, TArrayView<float> FinalOutputMap) const;
+	void CHAOSCLOTHASSETDATAFLOWNODES_API CalculateFinalRenderVertexWeightValues(const TConstArrayView<float> InputMap, TArrayView<float> FinalOutputMap) const;
+
 	//~ Begin FDataflowNode interface
 	virtual void SetAssetValue(TObjectPtr<UObject> Asset, Dataflow::FContext& Context) const override;
 	virtual void Evaluate(Dataflow::FContext& Context, const FDataflowOutput* Out) const override;
diff --git a/Engine/Plugins/ChaosClothAssetEditor/Source/ChaosClothAssetDataflowNodes/Public/ChaosClothAsset/ConnectableValue.h b/Engine/Plugins/ChaosClothAssetEditor/Source/ChaosClothAssetDataflowNodes/Public/ChaosClothAsset/ConnectableValue.h
--- a/Engine/Plugins/ChaosClothAssetEditor/Source/ChaosClothAssetDataflowNodes/Public/ChaosClothAsset/ConnectableValue.h
+++ b/Engine/Plugins/ChaosClothAssetEditor/Source/ChaosClothAssetDataflowNodes/Public/ChaosClothAsset/ConnectableValue.h
@@ -12,7 +12,7 @@ struct FChaosClothAssetConnectableStringValue
 
 	/** The value for this property. */
 	UPROPERTY(EditAnywhere, Category = "Value")
-	FString StringValue = TEXT("StringValue");
+	FString StringValue;
 
 	/** The string value override value for when the StringValue has a connection that replaces the provided string value. */
 	UPROPERTY(VisibleAnywhere, Category = "Value", Transient)
@@ -26,7 +26,7 @@ struct FChaosClothAssetConnectableIStringValue
 
 	/** The value for this property. */
 	UPROPERTY(EditAnywhere, Category = "Value", Meta = (DataflowInput))
-	FString StringValue = TEXT("StringValue");
+	FString StringValue;
 
 	/**
 	  * Whether the property could import fabrics datas or not
@@ -52,7 +52,7 @@ struct FChaosClothAssetConnectableIOStringValue
 
 	/** The value for this property. */
 	UPROPERTY(EditAnywhere, Category = "Value", Meta = (DataflowInput, DataflowOutput, DataflowPassthrough = "StringValue"))
-	FString StringValue = TEXT("StringValue");
+	FString StringValue;
 
 	/** The string value override value for when the StringValue has a connection that replaces the provided string value. */
 	UPROPERTY(VisibleAnywhere, Category = "Value", Transient)
diff --git a/Engine/Plugins/ChaosClothAssetEditor/Source/ChaosClothAssetEditor/Private/ChaosClothAsset/ClothEditorMode.cpp b/Engine/Plugins/ChaosClothAssetEditor/Source/ChaosClothAssetEditor/Private/ChaosClothAsset/ClothEditorMode.cpp
--- a/Engine/Plugins/ChaosClothAssetEditor/Source/ChaosClothAssetEditor/Private/ChaosClothAsset/ClothEditorMode.cpp
+++ b/Engine/Plugins/ChaosClothAssetEditor/Source/ChaosClothAssetEditor/Private/ChaosClothAsset/ClothEditorMode.cpp
@@ -231,6 +231,7 @@ void UChaosClothAssetEditorMode::RegisterClothTool(TSharedPtr<FUICommandInfo> UI
 			TArray<UE::Chaos::ClothAsset::EClothPatternVertexType> SupportedModes;
 			ClothToolBuilder->GetSupportedViewModes(SupportedModes);
 
+			bDynamicMeshUseInputCollection = true;
 			if (SupportedModes.Num() > 0 && !SupportedModes.Contains(this->GetConstructionViewMode()))
 			{
 				if (!bShouldRestoreSavedConstructionViewMode)
@@ -243,6 +244,21 @@ void UChaosClothAssetEditorMode::RegisterClothTool(TSharedPtr<FUICommandInfo> UI
 				// switch to the preferred view mode for the tool that's about to start
 				this->SetConstructionViewMode(SupportedModes[0]);
 			}
+			else
+			{
+				bool bCurrentDynamicMeshIsInput = false;
+				if (const UEditorInteractiveToolsContext* const RestSpaceToolsContext = GetInteractiveToolsContext())
+				{
+					if (const UClothEditorContextObject* const EditorContextObject = RestSpaceToolsContext->ContextObjectStore->FindContext<UClothEditorContextObject>())
+					{
+						bCurrentDynamicMeshIsInput = EditorContextObject->IsUsingInputCollection();
+					}
+				}
+				if (!bCurrentDynamicMeshIsInput || bDynamicMeshComponentInitDeferred)
+				{
+					ReinitializeDynamicMeshComponents();
+				}
+			}
 
 			// Check if we need to disable wireframe mode before starting tool.
 			const bool bCanSetWireframeActive = ClothToolBuilder->CanSetConstructionViewWireframeActive();
@@ -380,6 +396,7 @@ void UChaosClothAssetEditorMode::OnToolEnded(UInteractiveToolManager* Manager, U
 		bShouldRestoreConstructionViewSeams = false;
 	}
 
+	bDynamicMeshUseInputCollection = false;
 	if (bShouldRestoreSavedConstructionViewMode)
 	{
 		SetConstructionViewMode(SavedConstructionViewMode);
@@ -527,14 +544,22 @@ bool UChaosClothAssetEditorMode::IsComponentSelected(const UPrimitiveComponent*
 }
 
 
-void UChaosClothAssetEditorMode::SetSelectedClothCollection(TSharedPtr<FManagedArrayCollection> Collection, TSharedPtr<FManagedArrayCollection> InputCollection)
+void UChaosClothAssetEditorMode::SetSelectedClothCollection(TSharedPtr<FManagedArrayCollection> Collection, TSharedPtr<FManagedArrayCollection> InputCollection, bool bDeferDynamicMeshInitForTool)
 {
 	SelectedClothCollection = Collection;
 	SelectedInputClothCollection = InputCollection;
-	ReinitializeDynamicMeshComponents();
 
-	// The first time we get a valid mesh, refocus the camera on it
-	FirstTimeFocusRestSpaceViewport();
+	if (bDeferDynamicMeshInitForTool)
+	{
+		bDynamicMeshComponentInitDeferred = true;
+	}
+	else
+	{
+		ReinitializeDynamicMeshComponents();
+
+		// The first time we get a valid mesh, refocus the camera on it
+		FirstTimeFocusRestSpaceViewport();
+	}
 }
 
 TSharedPtr<FManagedArrayCollection> UChaosClothAssetEditorMode::GetClothCollection()
@@ -789,8 +814,9 @@ void UChaosClothAssetEditorMode::ReinitializeDynamicMeshComponents()
 	DynamicMeshComponentParentActor = nullptr;
 	WireframeDraw = nullptr;
 	ClothSeamDraw = nullptr;
+	bDynamicMeshComponentInitDeferred = false;
 
-	TSharedPtr<FManagedArrayCollection> Collection = GetClothCollection();
+	TSharedPtr<FManagedArrayCollection> Collection = bDynamicMeshUseInputCollection ? GetInputClothCollection() : GetClothCollection();
 	if (!Collection)
 	{
 		return;
@@ -928,7 +954,7 @@ void UChaosClothAssetEditorMode::ReinitializeDynamicMeshComponents()
 	UClothEditorContextObject* EditorContextObject = RestSpaceToolsContext->ContextObjectStore->FindContext<UClothEditorContextObject>();
 	if (ensure(EditorContextObject))
 	{
-		EditorContextObject->SetClothCollection(ConstructionViewMode, Collection, GetInputClothCollection());
+		EditorContextObject->SetClothCollection(ConstructionViewMode, Collection, bDynamicMeshUseInputCollection);
 	}
 }
 
@@ -1185,6 +1211,10 @@ void UChaosClothAssetEditorMode::ModeTick(float DeltaTime)
 		NodeTypeForPendingToolStart = FName();
 	}
 
+	if (bDynamicMeshComponentInitDeferred)
+	{
+		ReinitializeDynamicMeshComponents();
+	}
 
 	if (PreviewScene->GetWorld())
 	{
@@ -1464,8 +1494,8 @@ void UChaosClothAssetEditorMode::InitializeContextObject()
 		EditorContextObject = NewObject<UClothEditorContextObject>();
 		RestSpaceToolsContext->ContextObjectStore->AddContextObject(EditorContextObject);
 	}
-
-	EditorContextObject->Init(DataflowGraphEditor, ConstructionViewMode, SelectedClothCollection, SelectedInputClothCollection);
+	constexpr bool bUsingInputCollection = false;
+	EditorContextObject->Init(DataflowGraphEditor, DataflowContext, ConstructionViewMode, SelectedClothCollection, bUsingInputCollection);
 
 	check(EditorContextObject);
 
@@ -1480,6 +1510,16 @@ void UChaosClothAssetEditorMode::DeleteContextObject()
 	}
 }
 
+void UChaosClothAssetEditorMode::SetDataflowContext(TWeakPtr<Dataflow::FEngineContext> InDataflowContext)
+{
+	DataflowContext = InDataflowContext;
+	UEditorInteractiveToolsContext* const RestSpaceToolsContext = GetInteractiveToolsContext();
+	if (UClothEditorContextObject* ClothEditorContextObject = RestSpaceToolsContext->ContextObjectStore->FindContext<UClothEditorContextObject>())
+	{
+		ClothEditorContextObject->SetDataflowContext(DataflowContext);
+	}	
+}
+
 void UChaosClothAssetEditorMode::SetDataflowGraphEditor(TSharedPtr<SDataflowGraphEditor> InGraphEditor)
 {
 	DataflowGraphEditor = InGraphEditor;
diff --git a/Engine/Plugins/ChaosClothAssetEditor/Source/ChaosClothAssetEditor/Private/ChaosClothAsset/ClothEditorMode.h b/Engine/Plugins/ChaosClothAssetEditor/Source/ChaosClothAssetEditor/Private/ChaosClothAsset/ClothEditorMode.h
--- a/Engine/Plugins/ChaosClothAssetEditor/Source/ChaosClothAssetEditor/Private/ChaosClothAsset/ClothEditorMode.h
+++ b/Engine/Plugins/ChaosClothAssetEditor/Source/ChaosClothAssetEditor/Private/ChaosClothAsset/ClothEditorMode.h
@@ -4,6 +4,7 @@
 
 #include "BaseCharacterFXEditorMode.h"
 #include "GeometryBase.h"
+#include "Dataflow/DataflowObjectInterface.h"
 #include "Delegates/IDelegateInstance.h"
 #include "ChaosClothAsset/ClothPatternVertexType.h"
 #include "ClothEditorMode.generated.h"
@@ -29,6 +30,7 @@ class UDataflow;
 class UDataflowComponent;
 class SDataflowGraphEditor;
 struct FManagedArrayCollection;
+
 namespace UE::Chaos::ClothAsset
 {
 class FChaosClothPreviewScene;
@@ -177,12 +179,13 @@ private:
 	void FirstTimeFocusRestSpaceViewport();
 
 	// intended to be called by the toolkit when selected node in the Dataflow graph changes
-	void SetSelectedClothCollection(TSharedPtr<FManagedArrayCollection> Collection, TSharedPtr<FManagedArrayCollection> InputCollection = nullptr);
+	void SetSelectedClothCollection(TSharedPtr<FManagedArrayCollection> Collection, TSharedPtr<FManagedArrayCollection> InputCollection = nullptr, bool bDeferDynamicMeshInitForTool = false);
 
 	// gets the currently selected cloth collection, as specified by the toolkit
 	TSharedPtr<FManagedArrayCollection> GetClothCollection();
 	TSharedPtr<FManagedArrayCollection> GetInputClothCollection();
 
+	void SetDataflowContext(TWeakPtr<Dataflow::FEngineContext> InDataflowContext);
 	void SetDataflowGraphEditor(TSharedPtr<SDataflowGraphEditor> InGraphEditor);
 	
 	void StartToolForSelectedNode(const UObject* SelectedNode);
@@ -250,6 +253,11 @@ private:
 
 	TWeakPtr<UE::Chaos::ClothAsset::FChaosClothEditorRestSpaceViewportClient, ESPMode::ThreadSafe> RestSpaceViewportClient;
 
+	// The dynamic mesh component needs to be reinitialized on next tick.
+	bool bDynamicMeshComponentInitDeferred = false;
+	// Use the input collection to build the dynamic mesh component (used by tools).
+	bool bDynamicMeshUseInputCollection = false;
+
 	// The first time we get a valid mesh, refocus the camera on it
 	bool bFirstValid2DMesh = true;
 	bool bFirstValid3DMesh = true;
@@ -298,7 +306,7 @@ private:
 
 	UPROPERTY()
 	TObjectPtr<UEditorInteractiveToolsContext> ActiveToolsContext = nullptr;
-
+	TWeakPtr<Dataflow::FEngineContext> DataflowContext;
 	TSharedPtr<FManagedArrayCollection> SelectedClothCollection = nullptr;
 	TSharedPtr<FManagedArrayCollection> SelectedInputClothCollection = nullptr;
 
diff --git a/Engine/Plugins/ChaosClothAssetEditor/Source/ChaosClothAssetEditor/Private/ChaosClothAsset/ClothEditorToolkit.cpp b/Engine/Plugins/ChaosClothAssetEditor/Source/ChaosClothAssetEditor/Private/ChaosClothAsset/ClothEditorToolkit.cpp
--- a/Engine/Plugins/ChaosClothAssetEditor/Source/ChaosClothAssetEditor/Private/ChaosClothAsset/ClothEditorToolkit.cpp
+++ b/Engine/Plugins/ChaosClothAssetEditor/Source/ChaosClothAssetEditor/Private/ChaosClothAsset/ClothEditorToolkit.cpp
@@ -494,6 +494,8 @@ void FChaosClothAssetEditorToolkit::PostInitAssetEditor()
 
 	ClothMode->DataflowGraph = GetDataflow();
 	ClothMode->SetDataflowGraphEditor(GraphEditor);
+	DataflowContext = TSharedPtr<Dataflow::FEngineContext>(new Dataflow::FClothAssetDataflowContext(GetAsset(), GetDataflow(), Dataflow::FTimestamp::Invalid));
+	ClothMode->SetDataflowContext(DataflowContext);
 
 	// Handle Dataflow asset reload event
 	OnPackageReloadedDelegateHandle = FCoreUObjectDelegates::OnPackageReloaded.AddSP(this, &FChaosClothAssetEditorToolkit::HandlePackageReloaded);
@@ -855,11 +857,7 @@ void FChaosClothAssetEditorToolkit::EvaluateNode(FDataflowNode* Node, bool bForc
 		}
 
 		// Set evaluation context
-		if (!DataflowContext)
-		{
-			DataflowContext = TSharedPtr<Dataflow::FEngineContext>(new Dataflow::FClothAssetDataflowContext(ClothAsset, Dataflow, Dataflow::FTimestamp::Invalid));
-			bForceOperation = true;
-		}
+		check(DataflowContext);
 		if (bForceOperation)
 		{
 			LastDataflowNodeTimestamp = Dataflow::FTimestamp::Invalid;
@@ -1258,7 +1256,10 @@ void FChaosClothAssetEditorToolkit::OnNodeSelectionChanged(const TSet<UObject*>&
 			}
 
 			// Update the Construction viewport with the newly selected node's Collection
-			ClothMode->SetSelectedClothCollection(Collection, InputCollection);
+			// Defer updating the dynamic mesh component when single selecting since tools want the dynamic mesh collection 
+			//   to be built from the InputCollection rather than the (Output)Collection.
+			const bool bDeferDynamicMeshInitForTool = NewSelection.Num() == 1;
+			ClothMode->SetSelectedClothCollection(Collection, InputCollection, bDeferDynamicMeshInitForTool);
 		}
 
 		if (Outliner)
diff --git a/Engine/Plugins/ChaosClothAssetEditor/Source/ChaosClothAssetEditorTools/Private/ChaosClothAsset/ClothEditorContextObject.cpp b/Engine/Plugins/ChaosClothAssetEditor/Source/ChaosClothAssetEditorTools/Private/ChaosClothAsset/ClothEditorContextObject.cpp
--- a/Engine/Plugins/ChaosClothAssetEditor/Source/ChaosClothAssetEditorTools/Private/ChaosClothAsset/ClothEditorContextObject.cpp
+++ b/Engine/Plugins/ChaosClothAssetEditor/Source/ChaosClothAssetEditorTools/Private/ChaosClothAsset/ClothEditorContextObject.cpp
@@ -8,17 +8,18 @@
 
 #include UE_INLINE_GENERATED_CPP_BY_NAME(ClothEditorContextObject)
 
-void UClothEditorContextObject::Init(TWeakPtr<SDataflowGraphEditor> InDataflowGraphEditor,UE::Chaos::ClothAsset::EClothPatternVertexType InConstructionViewMode, TWeakPtr<FManagedArrayCollection> InSelectedClothCollection, TWeakPtr<FManagedArrayCollection> InSelectedInputClothCollection)
+void UClothEditorContextObject::Init(TWeakPtr<SDataflowGraphEditor> InDataflowGraphEditor, TWeakPtr<Dataflow::FEngineContext> InDataflowContext, UE::Chaos::ClothAsset::EClothPatternVertexType InConstructionViewMode, TWeakPtr<FManagedArrayCollection> InSelectedClothCollection, bool bInUsingInputCollection)
 {
 	DataflowGraphEditor = InDataflowGraphEditor;
+	DataflowContext = InDataflowContext;
 	ConstructionViewMode = InConstructionViewMode;
 	SelectedClothCollection = InSelectedClothCollection;
-	SelectedInputClothCollection = InSelectedInputClothCollection;
+	bUsingInputCollection = bInUsingInputCollection;
 }
 
-void UClothEditorContextObject::SetClothCollection(UE::Chaos::ClothAsset::EClothPatternVertexType ViewMode, TWeakPtr<FManagedArrayCollection> ClothCollection, TWeakPtr<FManagedArrayCollection> InputClothCollection)
+void UClothEditorContextObject::SetClothCollection(UE::Chaos::ClothAsset::EClothPatternVertexType ViewMode, TWeakPtr<FManagedArrayCollection> ClothCollection, bool bInUsingInputCollection)
 {
 	ConstructionViewMode = ViewMode;
 	SelectedClothCollection = ClothCollection;
-	SelectedInputClothCollection = InputClothCollection;
+	bUsingInputCollection = bInUsingInputCollection;
 }
\ No newline at end of file
diff --git a/Engine/Plugins/ChaosClothAssetEditor/Source/ChaosClothAssetEditorTools/Private/ChaosClothAsset/ClothMeshSelectionTool.cpp b/Engine/Plugins/ChaosClothAssetEditor/Source/ChaosClothAssetEditorTools/Private/ChaosClothAsset/ClothMeshSelectionTool.cpp
--- a/Engine/Plugins/ChaosClothAssetEditor/Source/ChaosClothAssetEditorTools/Private/ChaosClothAsset/ClothMeshSelectionTool.cpp
+++ b/Engine/Plugins/ChaosClothAssetEditor/Source/ChaosClothAssetEditorTools/Private/ChaosClothAsset/ClothMeshSelectionTool.cpp
@@ -255,6 +255,7 @@ void UClothMeshSelectionTool::Setup()
 	// Setup non-manifold mapping if necessary
 	if (ClothEditorContextObject)
 	{
+		ensure(ClothEditorContextObject->IsUsingInputCollection());
 		if (const TSharedPtr<const FManagedArrayCollection> ClothCollection = ClothEditorContextObject->GetSelectedClothCollection().Pin())
 		{
 			PreviewMesh->ProcessMesh([this, ClothCollection](const FDynamicMesh3& Mesh)
@@ -620,8 +621,10 @@ void UClothMeshSelectionTool::ApplyAction(EClothMeshSelectionToolActions ActionT
 
 void UClothMeshSelectionTool::ImportFromCollection(bool bImportFromSecondarySet)
 {
-	if (const TSharedPtr<const FManagedArrayCollection> ClothCollection = ClothEditorContextObject->GetSelectedInputClothCollection().Pin())
+	if (const TSharedPtr<const FManagedArrayCollection> ClothCollection = ClothEditorContextObject->GetSelectedClothCollection().Pin())
 	{
+		ensure(ClothEditorContextObject->IsUsingInputCollection());
+
 		using namespace UE::Chaos::ClothAsset;
 		const FCollectionClothSelectionConstFacade SelectionFacade(ClothCollection.ToSharedRef());
 		if (SelectionFacade.IsValid())
diff --git a/Engine/Plugins/ChaosClothAssetEditor/Source/ChaosClothAssetEditorTools/Private/ChaosClothAsset/ClothWeightMapPaintTool.cpp b/Engine/Plugins/ChaosClothAssetEditor/Source/ChaosClothAssetEditorTools/Private/ChaosClothAsset/ClothWeightMapPaintTool.cpp
--- a/Engine/Plugins/ChaosClothAssetEditor/Source/ChaosClothAssetEditorTools/Private/ChaosClothAsset/ClothWeightMapPaintTool.cpp
+++ b/Engine/Plugins/ChaosClothAssetEditor/Source/ChaosClothAssetEditorTools/Private/ChaosClothAsset/ClothWeightMapPaintTool.cpp
@@ -31,6 +31,7 @@
 
 #include "CanvasTypes.h"
 #include "CanvasItem.h"
+#include "ChaosClothAsset/ClothCollectionGroup.h"
 #include "ChaosClothAsset/ClothEditorContextObject.h"
 #include "ChaosClothAsset/ClothPatternVertexType.h"
 #include "ChaosClothAsset/CollectionClothFacade.h"
@@ -195,6 +196,7 @@ void UClothEditorWeightMapPaintTool::Setup()
 
 	UpdateWeightMapProperties = NewObject<UClothEditorUpdateWeightMapProperties>(this);
 	UpdateWeightMapProperties->Name = WeightMapNodeToUpdate->Name;
+	UpdateWeightMapProperties->MapOverrideType = WeightMapNodeToUpdate->MapOverrideType;
 
 	UpdateWeightMapProperties->WatchProperty(WeightMapNodeToUpdate->Name, [this](const FString& NewName)
 	{
@@ -356,9 +358,12 @@ void UClothEditorWeightMapPaintTool::Setup()
 	// configure panels
 	UpdateSubToolType(FilterProperties->SubToolType);
 
-	// Setup DynamicMeshToWeight conversion
+	const bool bIsRenderMode = ClothEditorContextObject && (ClothEditorContextObject->GetConstructionViewMode() == UE::Chaos::ClothAsset::EClothPatternVertexType::Render);
+	// Setup DynamicMeshToWeight conversion and get Input weight map (if it exists)
+	InputWeightMap = TConstArrayView<float>();
 	if (ClothEditorContextObject)
 	{
+		ensure(ClothEditorContextObject->IsUsingInputCollection());
 		if (TSharedPtr<const FManagedArrayCollection> ClothCollection = ClothEditorContextObject->GetSelectedClothCollection().Pin())
 		{
 			using namespace UE::Chaos::ClothAsset;
@@ -394,7 +399,6 @@ void UClothEditorWeightMapPaintTool::Setup()
 				WeightToDynamicMesh = Cloth.GetSimVertex2DLookup();
 			}
 
-			const bool bIsRenderMode = ClothEditorContextObject->GetConstructionViewMode() == EClothPatternVertexType::Render;
 			const int32 NumPatterns = bIsRenderMode ? Cloth.GetNumRenderPatterns() : Cloth.GetNumSimPatterns();		
 			
 			PatternTriangleOffsetAndNum.SetNum(NumPatterns);
@@ -429,7 +433,20 @@ void UClothEditorWeightMapPaintTool::Setup()
 			{
 				ShowHideProperties->ShowPatterns.Add({ PatternID, false });
 			}
-		
+
+			// Find the map if it exists.
+			if (TSharedPtr<Dataflow::FEngineContext> DataflowContext = ClothEditorContextObject->GetDataflowContext().Pin())
+			{
+				const FName InputName = WeightMapNodeToUpdate->GetInputName(*DataflowContext);
+				if (bIsRenderMode)
+				{
+					InputWeightMap = Cloth.GetUserDefinedAttribute<float>(InputName, ClothCollectionGroup::RenderVertices);
+				}
+				else
+				{
+					InputWeightMap = Cloth.GetWeightMap(InputName);
+				}
+			}
 		}
 	}
 
@@ -444,34 +461,38 @@ void UClothEditorWeightMapPaintTool::Setup()
 	ActiveWeightMap->SetName(FName("PaintLayer"));
 
 	// Copy weights from selected node to the preview mesh
-	const bool bIsRenderMode = (ClothEditorContextObject->GetConstructionViewMode() == UE::Chaos::ClothAsset::EClothPatternVertexType::Render);
-	const TArray<float>& CurrentWeights = bIsRenderMode ? WeightMapNodeToUpdate->GetRenderVertexWeights() : WeightMapNodeToUpdate->GetVertexWeights();
+	const int32 NumExpectedWeights = bHaveDynamicMeshToWeightConversion ? WeightToDynamicMesh.Num() : Mesh->MaxVertexID();
+	TArray<float> CurrentWeights;
+	CurrentWeights.SetNumZeroed(NumExpectedWeights);
+	if (bIsRenderMode)
+	{
+		WeightMapNodeToUpdate->CalculateFinalRenderVertexWeightValues(InputWeightMap, TArrayView<float>(CurrentWeights));
+	}
+	else
+	{
+		WeightMapNodeToUpdate->CalculateFinalVertexWeightValues(InputWeightMap, TArrayView<float>(CurrentWeights));
+	}
 	
 	if (bHaveDynamicMeshToWeightConversion)
 	{
-		if (WeightToDynamicMesh.Num() == CurrentWeights.Num())	// Only copy node weights if they match the number of mesh vertices
+		for (int32 WeightID = 0; WeightID < CurrentWeights.Num(); ++WeightID)
 		{
-			for (int32 WeightID = 0; WeightID < CurrentWeights.Num(); ++WeightID)
+			for (const int32 VertexID : WeightToDynamicMesh[WeightID])
 			{
-				for (const int32 VertexID : WeightToDynamicMesh[WeightID])
-				{
-					ActiveWeightMap->SetValue(VertexID, &CurrentWeights[WeightID]);
-				}
+				ActiveWeightMap->SetValue(VertexID, &CurrentWeights[WeightID]);
 			}
 		}
 	}
 	else
 	{
-		if (Mesh->MaxVertexID() == CurrentWeights.Num())	// Only copy node weights if they match the number of mesh vertices
+		for (int32 VertexID = 0; VertexID < CurrentWeights.Num(); ++VertexID)
 		{
-			for (int32 VertexID = 0; VertexID < CurrentWeights.Num(); ++VertexID)
-			{
-				ActiveWeightMap->SetValue(VertexID, &CurrentWeights[VertexID]);
-			}
+			ActiveWeightMap->SetValue(VertexID, &CurrentWeights[VertexID]);
 		}
 	}
 
 	UpdateWeightMapProperties->Name = WeightMapNodeToUpdate->Name;
+	UpdateWeightMapProperties->MapOverrideType = WeightMapNodeToUpdate->MapOverrideType;
 	SetToolPropertySourceEnabled(UpdateWeightMapProperties, true);
 
 	DynamicMeshComponent->EnableSecondaryTriangleBuffers(
@@ -1756,11 +1777,11 @@ void UClothEditorWeightMapPaintTool::OnTick(float DeltaTime)
 
 bool UClothEditorWeightMapPaintTool::CanAccept() const
 {
-	return bAnyChangeMade || UpdateWeightMapProperties->Name != WeightMapNodeToUpdate->Name;
+	return bAnyChangeMade || 
+		UpdateWeightMapProperties->Name != WeightMapNodeToUpdate->Name || 
+		UpdateWeightMapProperties->MapOverrideType != WeightMapNodeToUpdate->MapOverrideType;
 }
 
-
-
 FColor UClothEditorWeightMapPaintTool::GetColorForWeightValue(double WeightValue)
 {
 	FColor MaxColor = LinearColors::White3b();
@@ -1953,25 +1974,41 @@ void UClothEditorWeightMapPaintTool::UpdateSelectedNode()
 	GetCurrentWeightMap(CurrentWeights);
 
 	check(WeightMapNodeToUpdate);
-	const bool bIsRenderMode = (ClothEditorContextObject->GetConstructionViewMode() == UE::Chaos::ClothAsset::EClothPatternVertexType::Render);
+	WeightMapNodeToUpdate->MapOverrideType = UpdateWeightMapProperties->MapOverrideType;
+	WeightMapNodeToUpdate->Name = UpdateWeightMapProperties->Name;
 
-	TArray<float>& NodeWeights = bIsRenderMode ? WeightMapNodeToUpdate->GetRenderVertexWeights() : WeightMapNodeToUpdate->GetVertexWeights();
+	const bool bIsRenderMode = (ClothEditorContextObject->GetConstructionViewMode() == UE::Chaos::ClothAsset::EClothPatternVertexType::Render);
 
 	if (bHaveDynamicMeshToWeightConversion)
 	{
+		TArray<float> NodeWeights;
 		NodeWeights.Init(0.f, WeightToDynamicMesh.Num());
 		for (int32 DynamicMeshIdx = 0; DynamicMeshIdx < CurrentWeights.Num(); ++DynamicMeshIdx)
 		{
 			NodeWeights[DynamicMeshToWeight[DynamicMeshIdx]] = CurrentWeights[DynamicMeshIdx];
 		}
+
+		if (bIsRenderMode)
+		{
+			WeightMapNodeToUpdate->SetRenderVertexWeights(InputWeightMap, NodeWeights);
+		}
+		else
+		{
+			WeightMapNodeToUpdate->SetVertexWeights(InputWeightMap, NodeWeights);
+		}
 	}
 	else
 	{
-		NodeWeights = CurrentWeights;
+		if (bIsRenderMode)
+		{
+			WeightMapNodeToUpdate->SetRenderVertexWeights(InputWeightMap, CurrentWeights);
+		}
+		else
+		{
+			WeightMapNodeToUpdate->SetVertexWeights(InputWeightMap, CurrentWeights);
+		}
 	}
 	
-	WeightMapNodeToUpdate->Name = UpdateWeightMapProperties->Name;
-
 	WeightMapNodeToUpdate->Invalidate();
 }
 
diff --git a/Engine/Plugins/ChaosClothAssetEditor/Source/ChaosClothAssetEditorTools/Private/ChaosClothAsset/ClothWeightMapPaintTool.h b/Engine/Plugins/ChaosClothAssetEditor/Source/ChaosClothAssetEditorTools/Private/ChaosClothAsset/ClothWeightMapPaintTool.h
--- a/Engine/Plugins/ChaosClothAssetEditor/Source/ChaosClothAssetEditorTools/Private/ChaosClothAsset/ClothWeightMapPaintTool.h
+++ b/Engine/Plugins/ChaosClothAssetEditor/Source/ChaosClothAssetEditorTools/Private/ChaosClothAsset/ClothWeightMapPaintTool.h
@@ -31,6 +31,7 @@ class UWeightMapSmoothBrushOpProps;
 class UClothEditorContextObject;
 class UPolygonSelectionMechanic;
 struct FChaosClothAssetAddWeightMapNode;
+enum class EChaosClothAssetWeightMapOverrideType : uint8;
 
 DECLARE_STATS_GROUP(TEXT("WeightMapPaintTool"), STATGROUP_WeightMapPaintTool, STATCAT_Advanced);
 DECLARE_CYCLE_STAT(TEXT("WeightMapPaintTool_UpdateROI"), WeightMapPaintTool_UpdateROI, STATGROUP_WeightMapPaintTool);
@@ -252,6 +253,9 @@ public:
 	UPROPERTY(EditAnywhere, Category = UpdateNode, meta = (DisplayName = "Name"))
 	FString Name;
 
+	UPROPERTY(EditAnywhere, Category = UpdateNode)
+	EChaosClothAssetWeightMapOverrideType MapOverrideType;
+
 private:
 
 	virtual void PostEditChangeProperty(struct FPropertyChangedEvent& PropertyChangedEvent) override;
@@ -477,6 +481,8 @@ protected:
 	void OnDynamicMeshComponentChanged(UDynamicMeshComponent* Component, const FMeshVertexChange* Change, bool bRevert);
 	FDelegateHandle OnDynamicMeshComponentChangedHandle;
 
+	TConstArrayView<float> InputWeightMap;
+
 	UE::Geometry::FDynamicMeshWeightAttribute* ActiveWeightMap;
 	double GetCurrentWeightValue(int32 VertexId) const;
 	double GetCurrentWeightValueUnderBrush() const;
diff --git a/Engine/Plugins/ChaosClothAssetEditor/Source/ChaosClothAssetEditorTools/Public/ChaosClothAsset/ClothEditorContextObject.h b/Engine/Plugins/ChaosClothAssetEditor/Source/ChaosClothAssetEditorTools/Public/ChaosClothAsset/ClothEditorContextObject.h
--- a/Engine/Plugins/ChaosClothAssetEditor/Source/ChaosClothAssetEditorTools/Public/ChaosClothAsset/ClothEditorContextObject.h
+++ b/Engine/Plugins/ChaosClothAssetEditor/Source/ChaosClothAssetEditorTools/Public/ChaosClothAsset/ClothEditorContextObject.h
@@ -23,8 +23,13 @@ class CHAOSCLOTHASSETEDITORTOOLS_API UClothEditorContextObject : public UObject
 	GENERATED_BODY()
 
 public:
+	void Init(TWeakPtr<SDataflowGraphEditor> DataflowGraphEditor, TWeakPtr<Dataflow::FEngineContext> DataflowContext, UE::Chaos::ClothAsset::EClothPatternVertexType InConstructionViewMode, TWeakPtr<FManagedArrayCollection> SelectedClothCollection, bool bInUsingInputCollection);
 
-	void Init(TWeakPtr<SDataflowGraphEditor> DataflowGraphEditor, UE::Chaos::ClothAsset::EClothPatternVertexType InConstructionViewMode, TWeakPtr<FManagedArrayCollection> SelectedClothCollection, TWeakPtr<FManagedArrayCollection> SelectedInputClothCollection = nullptr);
+	UE_DEPRECATED(5.5, "Use the Init with a DataflowContext")
+	void Init(TWeakPtr<SDataflowGraphEditor> InDataflowGraphEditor, UE::Chaos::ClothAsset::EClothPatternVertexType InConstructionViewMode, TWeakPtr<FManagedArrayCollection> InSelectedClothCollection, TWeakPtr<FManagedArrayCollection> InSelectedInputClothCollection = nullptr)
+	{
+		Init(InDataflowGraphEditor, nullptr, InConstructionViewMode, InSelectedClothCollection, false);
+	}
 
 	/**
 	* Get a single selected node of the specified type. Return nullptr if the specified node is not selected, or if multiple nodes are selected
@@ -48,19 +53,38 @@ public:
 
 		return nullptr;
 	}
+	
+	TWeakPtr<Dataflow::FEngineContext> GetDataflowContext() const 
+	{
+		return DataflowContext;
+	}
+
+	void SetDataflowContext(TWeakPtr<Dataflow::FEngineContext> InDataflowContext)
+	{
+		DataflowContext = InDataflowContext;
+	}
+
+	void SetClothCollection(UE::Chaos::ClothAsset::EClothPatternVertexType ViewMode, TWeakPtr<FManagedArrayCollection> ClothCollection, bool bInUsingInputCollection);
 
-	void SetClothCollection(UE::Chaos::ClothAsset::EClothPatternVertexType ViewMode, TWeakPtr<FManagedArrayCollection> ClothCollection, TWeakPtr<FManagedArrayCollection> InputClothCollection = nullptr);
+	UE_DEPRECATED(5.5, "SetClothCollection no longer takes a separate InputClothCollection argument")
+	void SetClothCollection(UE::Chaos::ClothAsset::EClothPatternVertexType ViewMode, TWeakPtr<FManagedArrayCollection> ClothCollection, TWeakPtr<FManagedArrayCollection> InputClothCollection = nullptr)
+	{
+		SetClothCollection(ViewMode, ClothCollection, false);
+	}
 
 	const TWeakPtr<const FManagedArrayCollection> GetSelectedClothCollection() const { return SelectedClothCollection; }
-	const TWeakPtr<const FManagedArrayCollection> GetSelectedInputClothCollection() const { return SelectedInputClothCollection; }
+	UE_DEPRECATED(5.5, "There is no longer a separate input cloth collection.")
+	const TWeakPtr<const FManagedArrayCollection> GetSelectedInputClothCollection() const { return SelectedClothCollection; }
 	UE::Chaos::ClothAsset::EClothPatternVertexType GetConstructionViewMode() const { return ConstructionViewMode; }
+	bool IsUsingInputCollection() const { return bUsingInputCollection; }
 private:
 
 	TWeakPtr<SDataflowGraphEditor> DataflowGraphEditor;
+	TWeakPtr<Dataflow::FEngineContext> DataflowContext;
 
 	UE::Chaos::ClothAsset::EClothPatternVertexType ConstructionViewMode;
 	TWeakPtr<const FManagedArrayCollection> SelectedClothCollection;
-	TWeakPtr<const FManagedArrayCollection> SelectedInputClothCollection;
+	bool bUsingInputCollection = false;
 };
 
 
