Cloth Editor: don't store the animation instance in the Preview scene, just allow the SkeletMeshComponent to keep track of it.

This fixes a problem where the animation playback controls don't respond after reloading the SkeletalMesh asset in the preview scene. When a SkelMeshComponent is re-registered (e.g. after the asset is reloaded), it unconditionally destroys and re-creates its AnimScriptInstance, so we're guaranteed to be holding onto the wrong object in the editor if we are tracking it ourselves.

diff --git a/Engine/Plugins/ChaosClothAssetEditor/Source/ChaosClothAssetEditor/Private/ChaosClothAsset/ClothEditorPreviewScene.cpp b/Engine/Plugins/ChaosClothAssetEditor/Source/ChaosClothAssetEditor/Private/ChaosClothAsset/ClothEditorPreviewScene.cpp
--- a/Engine/Plugins/ChaosClothAssetEditor/Source/ChaosClothAssetEditor/Private/ChaosClothAsset/ClothEditorPreviewScene.cpp
+++ b/Engine/Plugins/ChaosClothAssetEditor/Source/ChaosClothAssetEditor/Private/ChaosClothAsset/ClothEditorPreviewScene.cpp
@@ -173,7 +173,6 @@ void FChaosClothPreviewScene::AddReferencedObjects(FReferenceCollector& Collecto
 	Collector.AddReferencedObject(ClothComponent);
 	Collector.AddReferencedObject(SkeletalMeshComponent);
 	Collector.AddReferencedObject(SceneActor);
-	Collector.AddReferencedObject(PreviewAnimInstance);
 }
 
 
@@ -186,7 +185,7 @@ void FChaosClothPreviewScene::UpdateSkeletalMeshAnimation()
 
 	if (PreviewSceneDescription->AnimationAsset)
 	{
-		PreviewAnimInstance = NewObject<UAnimSingleNodeInstance>(SkeletalMeshComponent);
+		TObjectPtr<UAnimSingleNodeInstance> PreviewAnimInstance = NewObject<UAnimSingleNodeInstance>(SkeletalMeshComponent);
 		PreviewAnimInstance->SetAnimationAsset(PreviewSceneDescription->AnimationAsset);
 
 		SkeletalMeshComponent->SetAnimationMode(EAnimationMode::AnimationSingleNode);
@@ -259,10 +258,6 @@ void FChaosClothPreviewScene::SceneDescriptionPropertyChanged(const FName& Prope
 
 	if (PropertyName == GET_MEMBER_NAME_CHECKED(UChaosClothPreviewSceneDescription, AnimationAsset))
 	{
-		if (!PreviewSceneDescription->AnimationAsset)
-		{
-			PreviewAnimInstance = nullptr;
-		}
 		UpdateSkeletalMeshAnimation();
 
 		if (UChaosClothAsset* const ClothAsset = ClothComponent->GetClothAsset())
@@ -370,12 +365,24 @@ void FChaosClothPreviewScene::SetClothAsset(UChaosClothAsset* Asset)
 
 UAnimSingleNodeInstance* FChaosClothPreviewScene::GetPreviewAnimInstance()
 {
-	return PreviewAnimInstance;
+	check(SkeletalMeshComponent);
+
+	if (SkeletalMeshComponent->AnimScriptInstance)
+	{
+		return CastChecked<UAnimSingleNodeInstance>(SkeletalMeshComponent->AnimScriptInstance);
+	}
+	return nullptr;
 }
 
 const UAnimSingleNodeInstance* const FChaosClothPreviewScene::GetPreviewAnimInstance() const
 {
-	return PreviewAnimInstance;
+	check(SkeletalMeshComponent);
+
+	if (SkeletalMeshComponent->AnimScriptInstance)
+	{
+		return CastChecked<UAnimSingleNodeInstance>(SkeletalMeshComponent->AnimScriptInstance);
+	}
+	return nullptr;
 }
 
 void FChaosClothPreviewScene::SetGizmoDataBinder(TSharedPtr<FTransformGizmoDataBinder> InDataBinder)
diff --git a/Engine/Plugins/ChaosClothAssetEditor/Source/ChaosClothAssetEditor/Private/ChaosClothAsset/ClothEditorPreviewScene.h b/Engine/Plugins/ChaosClothAssetEditor/Source/ChaosClothAssetEditor/Private/ChaosClothAsset/ClothEditorPreviewScene.h
--- a/Engine/Plugins/ChaosClothAssetEditor/Source/ChaosClothAssetEditor/Private/ChaosClothAsset/ClothEditorPreviewScene.h
+++ b/Engine/Plugins/ChaosClothAssetEditor/Source/ChaosClothAssetEditor/Private/ChaosClothAsset/ClothEditorPreviewScene.h
@@ -124,8 +124,6 @@ private:
 
 	TSharedPtr<FAssetEditorModeManager> ClothPreviewEditorModeManager;
 
-	TObjectPtr<UAnimSingleNodeInstance> PreviewAnimInstance;
-
 	TObjectPtr<AActor> SceneActor;
 
 	TObjectPtr<UChaosClothComponent> ClothComponent;
