When transferring weights to both sim and render meshes using inpaint weights method, make it so that we use closest point method for the render mesh only when transferring from sim mesh.

#rnx

diff --git a/Engine/Plugins/ChaosClothAssetEditor/Source/ChaosClothAssetDataflowNodes/Private/ChaosClothAsset/TransferSkinWeightsNode.cpp b/Engine/Plugins/ChaosClothAssetEditor/Source/ChaosClothAssetDataflowNodes/Private/ChaosClothAsset/TransferSkinWeightsNode.cpp
--- a/Engine/Plugins/ChaosClothAssetEditor/Source/ChaosClothAssetDataflowNodes/Private/ChaosClothAsset/TransferSkinWeightsNode.cpp
+++ b/Engine/Plugins/ChaosClothAssetEditor/Source/ChaosClothAssetDataflowNodes/Private/ChaosClothAsset/TransferSkinWeightsNode.cpp
@@ -799,7 +799,32 @@ void FChaosClothAssetTransferSkinWeightsNode::Evaluate(Dataflow::FContext& Conte
 			UE::Chaos::ClothAsset::FClothGeometryTools::CleanupAndCompactMesh(ClothCollection);
 
 			bool bTransferResult = false;
-			if (TransferMethod == EChaosClothAssetTransferSkinWeightsMethod::InpaintWeights)
+			if (TransferMethod == EChaosClothAssetTransferSkinWeightsMethod::InpaintWeights && 
+				TransferSettings.bTransferToRenderFromSim && 
+				TransferSettings.bTransferToSim && 
+				TransferSettings.bTransferToRender)
+			{
+				// Custom setup for the default behavior of the node that gives the best results in the common case where:
+				//  - the sim mesh is welded and manifold
+				//  - the render mesh is un-welded and highly non-manifold with many disconnected regions in areas like 
+				//    the armpits
+				//  - the render and the sim mesh have similar shapes
+				//  - we are transferring weight from the sim mesh and not the body
+
+				// First transfer to sim only using inpaint weights algorithm
+				TransferSettings.bTransferToRender = false;
+				bTransferResult = TransferInpaintWeights(TargetRefSkeleton, SourceDynamicMesh, ClothCollection, TransferSettings);
+
+				if (bTransferResult)
+				{
+					// Now transfer to render only using closest point 
+					TransferSettings.bTransferToSim = false;
+					TransferSettings.bTransferToRender = true;
+					TransferSettings.TransferMethod = FTransferBoneWeights::ETransferBoneWeightsMethod::ClosestPointOnSurface;
+					bTransferResult = TransferClosestPointOnSurface(TargetRefSkeleton, SourceDynamicMesh, ClothCollection, TransferSettings);
+				}
+			}
+			else if (TransferMethod == EChaosClothAssetTransferSkinWeightsMethod::InpaintWeights)
 			{
 				bTransferResult = TransferInpaintWeights(TargetRefSkeleton, SourceDynamicMesh, ClothCollection, TransferSettings);
 			}
