Cloth Editor Preview: when reloading or changing the SkeletalMeshAsset, save and restore the animation state so that it doesn't always start playing from time zero

diff --git a/Engine/Plugins/ChaosClothAssetEditor/Source/ChaosClothAssetEditor/Private/ChaosClothAsset/ClothEditorPreviewScene.cpp b/Engine/Plugins/ChaosClothAssetEditor/Source/ChaosClothAssetEditor/Private/ChaosClothAsset/ClothEditorPreviewScene.cpp
--- a/Engine/Plugins/ChaosClothAssetEditor/Source/ChaosClothAssetEditor/Private/ChaosClothAsset/ClothEditorPreviewScene.cpp
+++ b/Engine/Plugins/ChaosClothAssetEditor/Source/ChaosClothAssetEditor/Private/ChaosClothAsset/ClothEditorPreviewScene.cpp
@@ -5,6 +5,7 @@
 #include "ChaosClothAsset/ClothEditor.h"
 #include "Components/SkeletalMeshComponent.h"
 #include "Animation/SkeletalMeshActor.h"
+#include "Engine/SkeletalMesh.h"
 #include "AssetEditorModeManager.h"
 #include "Elements/Framework/EngineElementsLibrary.h"
 #include "Animation/AnimSingleNodeInstance.h"
@@ -17,7 +18,7 @@
 #include "Transforms/TransformGizmoDataBinder.h"
 #include "Framework/Notifications/NotificationManager.h"
 #include "Widgets/Notifications/SNotificationList.h"
-
+#include "UObject/PackageReload.h"
 
 #define LOCTEXT_NAMESPACE "UChaosClothEditorPreviewScene"
 
@@ -147,6 +148,7 @@ FChaosClothPreviewScene::FChaosClothPreviewScene(FPreviewScene::ConstructionValu
 	ClothComponent->SelectionOverrideDelegate = UPrimitiveComponent::FSelectionOverride::CreateRaw(this, &FChaosClothPreviewScene::IsComponentSelected);
 	ClothComponent->RegisterComponentWithWorld(GetWorld());
 	
+	OnPackageReloadedDelegateHandle = FCoreUObjectDelegates::OnPackageReloaded.AddRaw(this, &FChaosClothPreviewScene::HandlePackageReloaded);
 }
 
 FChaosClothPreviewScene::~FChaosClothPreviewScene()
@@ -163,6 +165,8 @@ FChaosClothPreviewScene::~FChaosClothPreviewScene()
 		ClothComponent->SelectionOverrideDelegate.Unbind();
 		ClothComponent->UnregisterComponent();
 	}
+
+	FCoreUObjectDelegates::OnPackageReloaded.Remove(OnPackageReloadedDelegateHandle);
 }
 
 void FChaosClothPreviewScene::AddReferencedObjects(FReferenceCollector& Collector)
@@ -175,6 +179,22 @@ void FChaosClothPreviewScene::AddReferencedObjects(FReferenceCollector& Collecto
 	Collector.AddReferencedObject(SceneActor);
 }
 
+void FChaosClothPreviewScene::Tick(float DeltaT)
+{
+	FAdvancedPreviewScene::Tick(DeltaT);
+
+	if (SavedAnimState)
+	{
+		if (UAnimSingleNodeInstance* const AnimInstance = GetPreviewAnimInstance())
+		{
+			AnimInstance->SetPosition(SavedAnimState->Time);
+			AnimInstance->SetReverse(SavedAnimState->bIsReverse);
+			AnimInstance->SetLooping(SavedAnimState->bIsLooping);
+			AnimInstance->SetPlaying(SavedAnimState->bIsPlaying);
+		}
+		SavedAnimState.Reset();
+	}
+}
 
 void FChaosClothPreviewScene::UpdateSkeletalMeshAnimation()
 {
@@ -233,6 +253,15 @@ void FChaosClothPreviewScene::SceneDescriptionPropertyChanged(const FName& Prope
 	{
 		check(SkeletalMeshComponent);
 
+		if (const UAnimSingleNodeInstance* const AnimInstance = GetPreviewAnimInstance())
+		{
+			SavedAnimState = FAnimState();
+			SavedAnimState->Time = AnimInstance->GetCurrentTime();
+			SavedAnimState->bIsReverse = AnimInstance->IsReverse();
+			SavedAnimState->bIsLooping = AnimInstance->IsLooping();
+			SavedAnimState->bIsPlaying = AnimInstance->IsPlaying();
+		}
+
 		SkeletalMeshComponent->SetSkeletalMeshAsset(PreviewSceneDescription->SkeletalMeshAsset);
 
 		UpdateSkeletalMeshAnimation();
@@ -390,6 +419,31 @@ void FChaosClothPreviewScene::SetGizmoDataBinder(TSharedPtr<FTransformGizmoDataB
 	DataBinder = InDataBinder;
 }
 
+
+void FChaosClothPreviewScene::HandlePackageReloaded(const EPackageReloadPhase InPackageReloadPhase, FPackageReloadedEvent* InPackageReloadedEvent)
+{
+	if (InPackageReloadPhase == EPackageReloadPhase::PrePackageFixup)
+	{
+		for (const TPair<UObject*, UObject*>& RepointPair : InPackageReloadedEvent->GetRepointedObjects())
+		{
+			if (RepointPair.Key == PreviewSceneDescription->SkeletalMeshAsset.Get())
+			{
+				// If we are going to be reloading the SkeletalMesh, first save the animation state since the AnimInstance will be reinitialized when the component is reregistered.
+				// Note we restore from the saved state in the Tick function above because AnimInstance reinitialization happens /after/ all reload delegates are called.
+				if (const UAnimSingleNodeInstance* const AnimInstance = GetPreviewAnimInstance())
+				{
+					SavedAnimState = FAnimState();
+					SavedAnimState->Time = AnimInstance->GetCurrentTime();
+					SavedAnimState->bIsReverse = AnimInstance->IsReverse();
+					SavedAnimState->bIsLooping = AnimInstance->IsLooping();
+					SavedAnimState->bIsPlaying = AnimInstance->IsPlaying();
+				}
+			}
+		}
+	}
+}
+
+
 } // namespace UE::Chaos::ClothAsset
 
 #undef LOCTEXT_NAMESPACE
diff --git a/Engine/Plugins/ChaosClothAssetEditor/Source/ChaosClothAssetEditor/Private/ChaosClothAsset/ClothEditorPreviewScene.h b/Engine/Plugins/ChaosClothAssetEditor/Source/ChaosClothAssetEditor/Private/ChaosClothAsset/ClothEditorPreviewScene.h
--- a/Engine/Plugins/ChaosClothAssetEditor/Source/ChaosClothAssetEditor/Private/ChaosClothAsset/ClothEditorPreviewScene.h
+++ b/Engine/Plugins/ChaosClothAssetEditor/Source/ChaosClothAssetEditor/Private/ChaosClothAsset/ClothEditorPreviewScene.h
@@ -112,6 +112,8 @@ public:
 
 private:
 
+	virtual void Tick(float DeltaT) override;
+
 	// Create the PreviewAnimationInstance if the AnimationAsset and SkeletalMesh both exist, and set the animation to run on the SkeletalMeshComponent
 	void UpdateSkeletalMeshAnimation();
 
@@ -120,6 +122,8 @@ private:
 
 	bool IsComponentSelected(const UPrimitiveComponent* InComponent);
 
+	void HandlePackageReloaded(const EPackageReloadPhase InPackageReloadPhase, FPackageReloadedEvent* InPackageReloadedEvent);
+
 	TObjectPtr<UChaosClothPreviewSceneDescription> PreviewSceneDescription;
 
 	TSharedPtr<FAssetEditorModeManager> ClothPreviewEditorModeManager;
@@ -131,6 +135,18 @@ private:
 	TObjectPtr<USkeletalMeshComponent> SkeletalMeshComponent;
 
 	TSharedPtr<FTransformGizmoDataBinder> DataBinder = nullptr;
+
+	struct FAnimState
+	{
+		float Time;
+		bool bIsReverse;
+		bool bIsLooping;
+		bool bIsPlaying;
+	};
+	TOptional<FAnimState> SavedAnimState;
+
+	FDelegateHandle OnPackageReloadedDelegateHandle;
+
 };
 } // namespace UE::Chaos::ClothAsset
 
