Cloth Editor: disable mouse "flight controls" in the construction viewport. Only alt+mouse or WASD should move the camera.

diff --git a/Engine/Plugins/ChaosClothAssetEditor/Source/ChaosClothAssetEditor/Private/ChaosClothAsset/ClothEditorRestSpaceViewportClient.cpp b/Engine/Plugins/ChaosClothAssetEditor/Source/ChaosClothAssetEditor/Private/ChaosClothAsset/ClothEditorRestSpaceViewportClient.cpp
--- a/Engine/Plugins/ChaosClothAssetEditor/Source/ChaosClothAssetEditor/Private/ChaosClothAsset/ClothEditorRestSpaceViewportClient.cpp
+++ b/Engine/Plugins/ChaosClothAssetEditor/Source/ChaosClothAssetEditor/Private/ChaosClothAsset/ClothEditorRestSpaceViewportClient.cpp
@@ -62,6 +62,9 @@ FChaosClothEditorRestSpaceViewportClient::FChaosClothEditorRestSpaceViewportClie
 	CameraPointLight->SetIntensity(3.0f);
 	CameraPointLight->SetCastShadows(false);
 	PreviewScene->AddComponent(CameraPointLight, FTransform());
+
+	// Disable the non-alt-key camera controls
+	bLockFlightCamera = true;
 }
 
 void FChaosClothEditorRestSpaceViewportClient::SetConstructionViewMode(EClothPatternVertexType InViewMode)
