Chaos Cloth:
- Make a render only cloth asset work without simulation.
- Fixed a crash in the transfer weight map code when the transfer collection has the weight map, but it is empty.

#tests Made a few combinations of render only Cloth Assets (all LODs, LOD1 only) and checked that they correctly rendered and skinned.

diff --git a/Engine/Plugins/ChaosClothAsset/Source/ChaosClothAsset/Private/ChaosClothAsset/CollectionClothFacade.cpp b/Engine/Plugins/ChaosClothAsset/Source/ChaosClothAsset/Private/ChaosClothAsset/CollectionClothFacade.cpp
--- a/Engine/Plugins/ChaosClothAsset/Source/ChaosClothAsset/Private/ChaosClothAsset/CollectionClothFacade.cpp
+++ b/Engine/Plugins/ChaosClothAsset/Source/ChaosClothAsset/Private/ChaosClothAsset/CollectionClothFacade.cpp
@@ -24,6 +24,23 @@ namespace UE::Chaos::ClothAsset
 		return ClothCollection->IsValid(OptionalSchemas) && ClothCollection->GetNumElements(ClothCollectionGroup::Lods) == 1;
 	}
 
+	bool FCollectionClothConstFacade::HasValidSimulationData() const
+	{
+		return IsValid() &&
+			GetNumSimPatterns() &&
+			GetNumSimVertices2D() &&
+			GetNumSimVertices3D() &&
+			GetNumSimFaces();
+	}
+
+	bool FCollectionClothConstFacade::HasValidRenderData() const
+	{
+		return IsValid() &&
+			GetNumRenderPatterns() &&
+			GetNumRenderVertices() &&
+			GetNumRenderFaces();
+	}
+
 	bool FCollectionClothConstFacade::HasValidData() const
 	{
 		return IsValid() &&
diff --git a/Engine/Plugins/ChaosClothAsset/Source/ChaosClothAsset/Public/ChaosClothAsset/CollectionClothFacade.h b/Engine/Plugins/ChaosClothAsset/Source/ChaosClothAsset/Public/ChaosClothAsset/CollectionClothFacade.h
--- a/Engine/Plugins/ChaosClothAsset/Source/ChaosClothAsset/Public/ChaosClothAsset/CollectionClothFacade.h
+++ b/Engine/Plugins/ChaosClothAsset/Source/ChaosClothAsset/Public/ChaosClothAsset/CollectionClothFacade.h
@@ -47,9 +47,13 @@ namespace UE::Chaos::ClothAsset
 		/** Return whether the facade is defined on the collection. */
 		bool IsValid(EClothCollectionOptionalSchemas OptionalSchemas = EClothCollectionOptionalSchemas::None) const;
 
-		/**
-		 * Return whether the facade has a non-empty sim and render mesh data.
-		 */
+		/** Return whether the facade has a non-empty simulation mesh data. */
+		bool HasValidSimulationData() const;
+
+		/** Return whether the facade has a non-empty sim mesh data. */
+		bool HasValidRenderData() const;
+
+		/** Return whether the facade has a non-empty sim and render mesh data. */
 		bool HasValidData() const;
 
 		uint32 CalculateTypeHash(bool bIncludeWeightMaps, uint32 PreviousHash = 0) const;
diff --git a/Engine/Plugins/ChaosClothAsset/Source/ChaosClothAssetEngine/Private/ChaosClothAsset/ClothSimulationProxy.cpp b/Engine/Plugins/ChaosClothAsset/Source/ChaosClothAssetEngine/Private/ChaosClothAsset/ClothSimulationProxy.cpp
--- a/Engine/Plugins/ChaosClothAsset/Source/ChaosClothAssetEngine/Private/ChaosClothAsset/ClothSimulationProxy.cpp
+++ b/Engine/Plugins/ChaosClothAsset/Source/ChaosClothAssetEngine/Private/ChaosClothAsset/ClothSimulationProxy.cpp
@@ -13,6 +13,7 @@
 #include "ChaosCloth/ChaosClothingSimulationSolver.h"
 #include "ChaosCloth/ChaosClothVisualization.h"
 #include "PhysicsEngine/PhysicsSettings.h"
+#include "Rendering/SkeletalMeshRenderData.h"
 #include "ClothingSimulation.h"
 
 #if INTEL_ISPC
@@ -181,6 +182,18 @@ PRAGMA_ENABLE_DEPRECATION_WARNINGS
 		// Fill a new context, note the context is also needed when the simulation is suspended
 		constexpr bool bIsInitializationFalse = false;
 		FillSimulationContext(DeltaTime, bIsInitializationFalse);
+
+		// Check that the render mesh current LOD isn't just fully skinned
+		const int32 LodIndex = ClothSimulationContext->LodIndex;
+		const UChaosClothAsset* const ClothAsset = ClothComponent.GetClothAsset();
+		if (!ClothAsset ||
+			!ClothAsset->GetResourceForRendering() ||
+			!ClothAsset->GetResourceForRendering()->LODRenderData.IsValidIndex(LodIndex) ||
+			!ClothAsset->GetResourceForRendering()->LODRenderData[LodIndex].HasClothData())
+		{
+			return false; // Not Simulating
+		}
+
 		Solver->SetEnableSolver(ShouldEnableSolver(Solver->GetEnableSolver()));
 
 		const bool bUseCache = ClothSimulationContext->CacheData.HasData();
diff --git a/Engine/Plugins/ChaosClothAssetEditor/Source/ChaosClothAssetDataflowNodes/Private/ChaosClothAsset/AddWeightMapNode.cpp b/Engine/Plugins/ChaosClothAssetEditor/Source/ChaosClothAssetDataflowNodes/Private/ChaosClothAsset/AddWeightMapNode.cpp
--- a/Engine/Plugins/ChaosClothAssetEditor/Source/ChaosClothAssetDataflowNodes/Private/ChaosClothAsset/AddWeightMapNode.cpp
+++ b/Engine/Plugins/ChaosClothAssetEditor/Source/ChaosClothAssetDataflowNodes/Private/ChaosClothAsset/AddWeightMapNode.cpp
@@ -250,7 +250,7 @@ PRAGMA_ENABLE_DEPRECATION_WARNINGS
 
 					const FName InInputName = GetInputName(Context);
 					const uint32 NameTypeHash = HashCombineFast(GetTypeHash(InInputName), (uint32)TransferType);
-					const uint32 InTransferCollectionHash = (TransferClothFacade.IsValid() && InInputName != NAME_None) ?
+					const uint32 InTransferCollectionHash = (TransferClothFacade.HasValidSimulationData() && InInputName != NAME_None) ?
 						HashCombineFast(TransferClothFacade.CalculateWeightMapTypeHash(), NameTypeHash) : 0;  // TODO: Remove after adding the function (currently shelved!)  
 				
 					if (TransferCollectionHash != InTransferCollectionHash)
diff --git a/Engine/Plugins/ChaosClothAssetEditor/Source/ChaosClothAssetDataflowNodes/Private/ChaosClothAsset/ProxyDeformerNode.cpp b/Engine/Plugins/ChaosClothAssetEditor/Source/ChaosClothAssetDataflowNodes/Private/ChaosClothAsset/ProxyDeformerNode.cpp
--- a/Engine/Plugins/ChaosClothAssetEditor/Source/ChaosClothAssetDataflowNodes/Private/ChaosClothAsset/ProxyDeformerNode.cpp
+++ b/Engine/Plugins/ChaosClothAssetEditor/Source/ChaosClothAssetDataflowNodes/Private/ChaosClothAsset/ProxyDeformerNode.cpp
@@ -314,7 +314,7 @@ void FChaosClothAssetProxyDeformerNode::Evaluate(Dataflow::FContext& Context, co
 		FManagedArrayCollection InCollection = GetValue<FManagedArrayCollection>(Context, &Collection);
 		const TSharedRef<FManagedArrayCollection> ClothCollection = MakeShared<FManagedArrayCollection>(MoveTemp(InCollection));
 
-		// Always check for a valid cloth collection/facade to avoid processing non cloth collections
+		// Always check for a valid cloth collection/facade/sim mesh to avoid processing non cloth collections or pure render mesh cloth assets
 		FCollectionClothFacade ClothFacade(ClothCollection);
 		if (ClothFacade.IsValid() && ClothFacade.HasValidData())
 		{
diff --git a/Engine/Plugins/ChaosClothAssetEditor/Source/ChaosClothAssetDataflowNodes/Private/ChaosClothAsset/TerminalNode.cpp b/Engine/Plugins/ChaosClothAssetEditor/Source/ChaosClothAssetDataflowNodes/Private/ChaosClothAsset/TerminalNode.cpp
--- a/Engine/Plugins/ChaosClothAssetEditor/Source/ChaosClothAssetDataflowNodes/Private/ChaosClothAsset/TerminalNode.cpp
+++ b/Engine/Plugins/ChaosClothAssetEditor/Source/ChaosClothAssetDataflowNodes/Private/ChaosClothAsset/TerminalNode.cpp
@@ -28,7 +28,7 @@ namespace UE::Chaos::ClothAsset::Private
 		{
 			constexpr bool bIncludeWeightMapsTrue = true; // Currently, editing weight maps is destructive
 			FCollectionClothConstFacade Cloth(ClothCollection);
-			if (Cloth.HasValidData())
+			if (Cloth.HasValidRenderData())  // The cloth collection must at least have a render mesh
 			{
 				Checksum = Cloth.CalculateTypeHash(bIncludeWeightMapsTrue, Checksum);
 
@@ -147,7 +147,7 @@ void FChaosClothAssetTerminalNode::SetAssetValue(TObjectPtr<UObject> Asset, Data
 			ClothFacade.DefineSchema();
 
 			const FCollectionClothConstFacade InClothFacade(InClothCollections[LodIndex]);
-			if (!InClothFacade.HasValidData())
+			if (!InClothFacade.HasValidRenderData())  // The cloth collection must at least have a render mesh
 			{
 				FClothDataflowTools::LogAndToastWarning(*this, LOCTEXT("InvalidLODHeadline", "Invalid LOD."),
 					FText::Format(
@@ -359,7 +359,7 @@ TArray<TSharedRef<FManagedArrayCollection>> FChaosClothAssetTerminalNode::GetCle
 		TSharedRef<FManagedArrayCollection>& CollectionLodValue = CollectionLodValues.Emplace_GetRef(MakeShared<FManagedArrayCollection>(GetValue<FManagedArrayCollection>(Context, CollectionLods[LodIndex])));
 
 		FCollectionClothFacade ClothFacade(CollectionLodValue);
-		if (ClothFacade.HasValidData())
+		if (ClothFacade.HasValidRenderData())  // The cloth collection must at least have a render mesh
 		{
 			FClothGeometryTools::CleanupAndCompactMesh(CollectionLodValue);
 			LastValidLodIndex = LodIndex;
diff --git a/Engine/Source/Runtime/ClothingSystemRuntimeCommon/Private/Utils/ClothingMeshUtils.cpp b/Engine/Source/Runtime/ClothingSystemRuntimeCommon/Private/Utils/ClothingMeshUtils.cpp
--- a/Engine/Source/Runtime/ClothingSystemRuntimeCommon/Private/Utils/ClothingMeshUtils.cpp
+++ b/Engine/Source/Runtime/ClothingSystemRuntimeCommon/Private/Utils/ClothingMeshUtils.cpp
@@ -687,6 +687,13 @@ namespace ClothingMeshUtils
 	{
 		check(SourceMesh.HasAveragedNormals());
 
+		OutMeshToMeshVertData.Reset();
+
+		if (!SourceMesh.GetIndices().Num() || !TargetMesh.GetPositions().Num())
+		{
+			return; // Empty meshes must return an empty mapping and no warning
+		}
+
 		if (!TargetMesh.HasValidMesh())  // Check that the number of positions is equal to the number of normals and that the number of indices is divisible by 3
 		{
 			UE_LOG(LogClothingMeshUtils, Warning, TEXT("Failed to generate mesh to mesh skinning data. Invalid Target Mesh."));
