Fix a check that occurs when clicking a node after doing "Save As" in the Cloth Editor

diff --git a/Engine/Plugins/ChaosClothAssetEditor/Source/ChaosClothAssetEditor/Private/ChaosClothAsset/ClothEditorToolkit.cpp b/Engine/Plugins/ChaosClothAssetEditor/Source/ChaosClothAssetEditor/Private/ChaosClothAsset/ClothEditorToolkit.cpp
--- a/Engine/Plugins/ChaosClothAssetEditor/Source/ChaosClothAssetEditor/Private/ChaosClothAsset/ClothEditorToolkit.cpp
+++ b/Engine/Plugins/ChaosClothAssetEditor/Source/ChaosClothAssetEditor/Private/ChaosClothAsset/ClothEditorToolkit.cpp
@@ -557,6 +557,11 @@ PRAGMA_ENABLE_DEPRECATION_WARNINGS
 		const TArray<UPackage*> PackagesToSave{NewClothAsset->GetOutermost()};
 		constexpr bool bPromptToSave = false;
 		FEditorFileUtils::PromptForCheckoutAndSave(PackagesToSave, bCheckDirtyOnAssetSave, bPromptToSave);
+
+		// Reload the package so the editor can update anything that references the Dataflow asset
+		// (Sadly this function is only called after the editor re-launches as part of the Save As process, so the editor is already open at this point.)
+		UPackageTools::ReloadPackages(PackagesToSave);
+
 	}
 }
 
