Chaos Cloth Editor Selection improvements to match recent weight map improvements.

Transfer collection like weight maps for changing topology.
Input selection for editing existing selections.

diff --git a/Engine/Plugins/ChaosClothAsset/Source/ChaosClothAsset/ChaosClothAsset.Build.cs b/Engine/Plugins/ChaosClothAsset/Source/ChaosClothAsset/ChaosClothAsset.Build.cs
--- a/Engine/Plugins/ChaosClothAsset/Source/ChaosClothAsset/ChaosClothAsset.Build.cs
+++ b/Engine/Plugins/ChaosClothAsset/Source/ChaosClothAsset/ChaosClothAsset.Build.cs
@@ -14,6 +14,7 @@ public class ChaosClothAsset : ModuleRules
 				"Core",
 				"GeometryCore",
 				"MeshConversion",
+				"ClothingSystemRuntimeCommon",
 			}
 		);
 		PublicDependencyModuleNames.AddRange(
diff --git a/Engine/Plugins/ChaosClothAsset/Source/ChaosClothAsset/Private/ChaosClothAsset/ClothGeometryTools.cpp b/Engine/Plugins/ChaosClothAsset/Source/ChaosClothAsset/Private/ChaosClothAsset/ClothGeometryTools.cpp
--- a/Engine/Plugins/ChaosClothAsset/Source/ChaosClothAsset/Private/ChaosClothAsset/ClothGeometryTools.cpp
+++ b/Engine/Plugins/ChaosClothAsset/Source/ChaosClothAsset/Private/ChaosClothAsset/ClothGeometryTools.cpp
@@ -1,7 +1,9 @@
 // Copyright Epic Games, Inc. All Rights Reserved.
 
 #include "ChaosClothAsset/ClothGeometryTools.h"
+#include "ChaosClothAsset/ClothCollectionGroup.h"
 #include "ChaosClothAsset/CollectionClothFacade.h"
+#include "ChaosClothAsset/CollectionClothSelectionFacade.h"
 #include "Containers/Queue.h"
 #include "DynamicMesh/DynamicMesh3.h"
 #include "DynamicMesh/DynamicMeshAttributeSet.h"
@@ -9,6 +11,7 @@
 #include "DynamicMesh/NonManifoldMappingSupport.h"
 #include "Math/Vector.h"
 #include "Util/IndexUtil.h"
+#include "Utils/ClothingMeshUtils.h"
 #include "Algo/RemoveIf.h"
 
 namespace UE::Chaos::ClothAsset
@@ -1260,5 +1263,200 @@ namespace UE::Chaos::ClothAsset
 		}
 	}
 
+	bool FClothGeometryTools::ConvertSelectionToNewGroupType(const TSharedRef<const FManagedArrayCollection>& ClothCollection, const FName& SelectionName, const FName& GroupName, bool bSecondarySelection, TSet<int32>& OutSelectionSet)
+	{
+		FCollectionClothSelectionConstFacade SelectionFacade(ClothCollection);
+		FCollectionClothConstFacade ClothFacade(ClothCollection);
+		if (!SelectionFacade.IsValid() || !ClothFacade.IsValid() || (bSecondarySelection ? !SelectionFacade.HasSelectionSecondarySet(SelectionName) : !SelectionFacade.HasSelection(SelectionName)))
+		{
+			return false;
+		}
+
+		const TSet<int32>& OrigSelectionSet = bSecondarySelection ? SelectionFacade.GetSelectionSecondarySet(SelectionName) : SelectionFacade.GetSelectionSet(SelectionName);
+		const FName OrigSelectionGroup = bSecondarySelection ? SelectionFacade.GetSelectionSecondaryGroup(SelectionName) : SelectionFacade.GetSelectionGroup(SelectionName);
+
+		if (OrigSelectionGroup == GroupName)
+		{
+			OutSelectionSet = OrigSelectionSet;
+			return true;
+		}
+
+		auto ConvertVerticesToFaces = [&OrigSelectionSet, &OutSelectionSet](const TConstArrayView<FIntVector3>& Indices)
+		{
+			OutSelectionSet.Reset();
+			OutSelectionSet.Reserve(OrigSelectionSet.Num());
+			for (int32 FaceIndex = 0; FaceIndex < Indices.Num(); ++FaceIndex)
+			{
+				const FIntVector3& Element = Indices[FaceIndex];
+				if (OrigSelectionSet.Contains(Element[0]) &&
+					OrigSelectionSet.Contains(Element[1]) &&
+					OrigSelectionSet.Contains(Element[2]))
+				{
+					OutSelectionSet.Add(FaceIndex);
+				}
+			}
+		};
+
+
+		auto ConvertFacesToVertices = [&OrigSelectionSet, &OutSelectionSet](const TConstArrayView<FIntVector3>& Indices)
+		{
+			OutSelectionSet.Reset();
+			OutSelectionSet.Reserve(OrigSelectionSet.Num());
+			for (const int32 FaceIndex : OrigSelectionSet)
+			{
+				if (Indices.IsValidIndex(FaceIndex))
+				{
+					OutSelectionSet.Add(Indices[FaceIndex][0]);
+					OutSelectionSet.Add(Indices[FaceIndex][1]);
+					OutSelectionSet.Add(Indices[FaceIndex][2]);
+				}
+			}
+		};	
+
+		if (OrigSelectionGroup == ClothCollectionGroup::SimVertices2D)
+		{
+			if (GroupName == ClothCollectionGroup::SimFaces)
+			{
+				ConvertVerticesToFaces(ClothFacade.GetSimIndices2D());
+				return true;
+			}
+			else if (GroupName == ClothCollectionGroup::SimVertices3D)
+			{
+				const TConstArrayView<int32> SimVertex3DLookup = ClothFacade.GetSimVertex3DLookup();
+				OutSelectionSet.Reset();
+				OutSelectionSet.Reserve(OrigSelectionSet.Num());
+				for (const int32 OrigSelection : OrigSelectionSet)
+				{
+					if (SimVertex3DLookup.IsValidIndex(OrigSelection))
+					{
+						OutSelectionSet.Add(SimVertex3DLookup[OrigSelection]);
+					}
+				}
+				return true;
+			}
+		}
+		else if (OrigSelectionGroup == ClothCollectionGroup::SimVertices3D)
+		{
+			if (GroupName == ClothCollectionGroup::SimFaces)
+			{
+				ConvertVerticesToFaces(ClothFacade.GetSimIndices3D());
+				return true;
+			}
+			else if (GroupName == ClothCollectionGroup::SimVertices2D)
+			{
+				const TConstArrayView<TArray<int32>> SimVertex2DLookup = ClothFacade.GetSimVertex2DLookup();
+				OutSelectionSet.Reset();
+				OutSelectionSet.Reserve(OrigSelectionSet.Num());
+				for (const int32 OrigSelection : OrigSelectionSet)
+				{
+					if (SimVertex2DLookup.IsValidIndex(OrigSelection))
+					{
+						for (const int32 Vertex2D : SimVertex2DLookup[OrigSelection])
+						{
+							OutSelectionSet.Add(Vertex2D);
+						}
+					}
+				}
+				return true;
+			}
+		}
+		else if (OrigSelectionGroup == ClothCollectionGroup::SimFaces)
+		{
+			if (GroupName == ClothCollectionGroup::SimVertices2D)
+			{
+				ConvertFacesToVertices(ClothFacade.GetSimIndices2D());
+				return true;
+			}
+			else if (GroupName == ClothCollectionGroup::SimVertices3D)
+			{
+				ConvertFacesToVertices(ClothFacade.GetSimIndices3D());
+				return true;
+			}
+		}
+		else if (OrigSelectionGroup == ClothCollectionGroup::RenderVertices)
+		{
+			if (GroupName == ClothCollectionGroup::RenderFaces)
+			{
+				ConvertVerticesToFaces(ClothFacade.GetRenderIndices());
+				return true;
+			}
+		}
+		else if (OrigSelectionGroup == ClothCollectionGroup::RenderFaces)
+		{
+			if (GroupName == ClothCollectionGroup::RenderVertices)
+			{
+				ConvertFacesToVertices(ClothFacade.GetRenderIndices());
+				return true;
+			}
+		}
+
+		return false;
+	}
+
+	void FClothGeometryTools::TransferWeightMap(
+		const TConstArrayView<FVector3f>& SourcePositions,
+		const TConstArrayView<FIntVector3>& InSourceIndices,
+		const TConstArrayView<float>& SourceWeights,
+		const TConstArrayView<FVector3f>& TargetPositions,
+		const TConstArrayView<FVector3f>& TargetNormals,
+		const TConstArrayView<FIntVector3>& InTargetIndices,
+		const TArrayView<float>& TargetWeights)
+	{
+		check(TargetWeights.Num() == TargetPositions.Num());
+		if (!ensure(SourcePositions.Num() <= 65536))
+		{
+			return;  // MeshToMeshVertData below is limited to 16bit unsigned int indexes
+		}
+
+		TArray<uint32> SourceIndices;
+		SourceIndices.Reserve(InSourceIndices.Num() * 3);
+		for (const FIntVector3& InSourceIndex : InSourceIndices)
+		{
+			SourceIndices.Add(InSourceIndex[0]);
+			SourceIndices.Add(InSourceIndex[1]);
+			SourceIndices.Add(InSourceIndex[2]);
+		}
+		TArray<uint32> TargetIndices;
+		TargetIndices.Reserve(InTargetIndices.Num() * 3);
+		for (const FIntVector3& InTargetIndex : InTargetIndices)
+		{
+			TargetIndices.Add(InTargetIndex[0]);
+			TargetIndices.Add(InTargetIndex[1]);
+			TargetIndices.Add(InTargetIndex[2]);
+		}
+
+		const ClothingMeshUtils::ClothMeshDesc SourceMeshDesc(SourcePositions, SourceIndices);
+		const ClothingMeshUtils::ClothMeshDesc TargetMeshDesc(TargetPositions, TargetNormals, TargetIndices);
+
+		TArray<FMeshToMeshVertData> MeshToMeshVertData;
+		const FPointWeightMap* const MaxDistances = nullptr; // No need to update the vertex contribution on the transition maps
+		constexpr bool bUseSmoothTransitions = false;  // Smooth transitions are only used at rendering for now and not during LOD transitions
+		constexpr bool bUseMultipleInfluences = false;  // Multiple influences must not be used for LOD transitions
+		constexpr float SkinningKernelRadius = 0.f;  // KernelRadius is only required when using multiple influences
+
+		ClothingMeshUtils::GenerateMeshToMeshVertData(
+			MeshToMeshVertData,
+			TargetMeshDesc,
+			SourceMeshDesc,
+			MaxDistances,
+			bUseSmoothTransitions,
+			bUseMultipleInfluences,
+			SkinningKernelRadius);
+
+		check(MeshToMeshVertData.Num() == TargetWeights.Num());
+		for (int32 Index = 0; Index < TargetWeights.Num(); ++Index)
+		{
+			const FMeshToMeshVertData& MeshToMeshVertDatum = MeshToMeshVertData[Index];
+
+			const uint16 VertIndex0 = MeshToMeshVertDatum.SourceMeshVertIndices[0];
+			const uint16 VertIndex1 = MeshToMeshVertDatum.SourceMeshVertIndices[1];
+			const uint16 VertIndex2 = MeshToMeshVertDatum.SourceMeshVertIndices[2];
+
+			TargetWeights[Index] = FMath::Clamp(
+				SourceWeights[VertIndex0] * MeshToMeshVertDatum.PositionBaryCoordsAndDist[0] +
+				SourceWeights[VertIndex1] * MeshToMeshVertDatum.PositionBaryCoordsAndDist[1] +
+				SourceWeights[VertIndex2] * MeshToMeshVertDatum.PositionBaryCoordsAndDist[2], 0.f, 1.f);
+		}
+	}
 
 }  // End namespace UE::Chaos::ClothAsset
diff --git a/Engine/Plugins/ChaosClothAsset/Source/ChaosClothAsset/Public/ChaosClothAsset/ClothGeometryTools.h b/Engine/Plugins/ChaosClothAsset/Source/ChaosClothAsset/Public/ChaosClothAsset/ClothGeometryTools.h
--- a/Engine/Plugins/ChaosClothAsset/Source/ChaosClothAsset/Public/ChaosClothAsset/ClothGeometryTools.h
+++ b/Engine/Plugins/ChaosClothAsset/Source/ChaosClothAsset/Public/ChaosClothAsset/ClothGeometryTools.h
@@ -117,6 +117,28 @@ namespace UE::Chaos::ClothAsset
 		 */
 		static void SampleVertices(const TConstArrayView<FVector3f> VertexPositions, float CullDiameterSq, TSet<int32>& OutVertexSet);
 
+		/**
+		 * Get a copy of the selection, converting to the desired group if possible.
+		 * Currently only conversions between vertex and face components on the same mesh type are supported.
+		 * @param ClothCollection to query
+		 * @param SelectionName the selection name
+		 * @param GroupName the group name
+		 * @param bSecondarySelection get the secondary selection
+		 * @param OutSelectionSet copy of the selection. Unchanged when function returns false.
+		 * @return success (will return false if the selection is not found or conversion is not possible)
+		 */
+		static bool ConvertSelectionToNewGroupType(const TSharedRef<const FManagedArrayCollection>& ClothCollection, const FName& SelectionName, const FName& GroupName, bool bSecondarySelection, TSet<int32>& OutSelectionSet);
 
+		/**
+		 * Transfer a vertex weight map from a source to target mesh
+		 */
+		static void TransferWeightMap(
+			const TConstArrayView<FVector3f>& SourcePositions,
+			const TConstArrayView<FIntVector3>& InSourceIndices,
+			const TConstArrayView<float>& SourceWeights,
+			const TConstArrayView<FVector3f>& TargetPositions,
+			const TConstArrayView<FVector3f>& TargetNormals,
+			const TConstArrayView<FIntVector3>& InTargetIndices,
+			const TArrayView<float>& TargetWeights);
 	};
 }  // End namespace UE::Chaos::ClothAsset
\ No newline at end of file
diff --git a/Engine/Plugins/ChaosClothAssetEditor/Source/ChaosClothAssetDataflowNodes/Private/ChaosClothAsset/AddWeightMapNode.cpp b/Engine/Plugins/ChaosClothAssetEditor/Source/ChaosClothAssetDataflowNodes/Private/ChaosClothAsset/AddWeightMapNode.cpp
--- a/Engine/Plugins/ChaosClothAssetEditor/Source/ChaosClothAssetDataflowNodes/Private/ChaosClothAsset/AddWeightMapNode.cpp
+++ b/Engine/Plugins/ChaosClothAssetEditor/Source/ChaosClothAssetDataflowNodes/Private/ChaosClothAsset/AddWeightMapNode.cpp
@@ -9,7 +9,6 @@
 #include "ChaosClothAsset/WeightedValue.h"
 #include "Dataflow/DataflowInputOutput.h"
 #include "Dataflow/DataflowObject.h"
-#include "Utils/ClothingMeshUtils.h"
 
 #include UE_INLINE_GENERATED_CPP_BY_NAME(AddWeightMapNode)
 
@@ -17,72 +16,6 @@
 
 namespace UE::Chaos::ClothAsset::Private
 {
-	void TransferWeightMap(
-		const TConstArrayView<FVector3f>& SourcePositions,
-		const TConstArrayView<FIntVector3>& InSourceIndices,
-		const TConstArrayView<float>& SourceWeights,
-		const TConstArrayView<FVector3f>& TargetPositions,
-		const TConstArrayView<FVector3f>& TargetNormals,
-		const TConstArrayView<FIntVector3>& InTargetIndices,
-		TArray<float>& TargetWeights)
-	{
-		check(TargetWeights.Num() == TargetPositions.Num());
-		if (!ensure(SourcePositions.Num() <= 65536))
-		{
-			return;  // MeshToMeshVertData below is limited to 16bit unsigned int indexes
-		}
-
-		TArray<uint32> SourceIndices;
-		SourceIndices.Reserve(InSourceIndices.Num() * 3);
-		for (const FIntVector3& InSourceIndex : InSourceIndices)
-		{
-			SourceIndices.Add(InSourceIndex[0]);
-			SourceIndices.Add(InSourceIndex[1]);
-			SourceIndices.Add(InSourceIndex[2]);
-		}
-		TArray<uint32> TargetIndices;
-		TargetIndices.Reserve(InTargetIndices.Num() * 3);
-		for (const FIntVector3& InTargetIndex : InTargetIndices)
-		{
-			TargetIndices.Add(InTargetIndex[0]);
-			TargetIndices.Add(InTargetIndex[1]);
-			TargetIndices.Add(InTargetIndex[2]);
-		}
-
-		const ClothingMeshUtils::ClothMeshDesc SourceMeshDesc(SourcePositions, SourceIndices);
-		const ClothingMeshUtils::ClothMeshDesc TargetMeshDesc(TargetPositions, TargetNormals, TargetIndices);
-
-		TArray<FMeshToMeshVertData> MeshToMeshVertData;
-		const FPointWeightMap* const MaxDistances = nullptr; // No need to update the vertex contribution on the transition maps
-		constexpr bool bUseSmoothTransitions = false;  // Smooth transitions are only used at rendering for now and not during LOD transitions
-		constexpr bool bUseMultipleInfluences = false;  // Multiple influences must not be used for LOD transitions
-		constexpr float SkinningKernelRadius = 0.f;  // KernelRadius is only required when using multiple influences
-
-		ClothingMeshUtils::GenerateMeshToMeshVertData(
-			MeshToMeshVertData,
-			TargetMeshDesc,
-			SourceMeshDesc,
-			MaxDistances,
-			bUseSmoothTransitions,
-			bUseMultipleInfluences,
-			SkinningKernelRadius);
-
-		check(MeshToMeshVertData.Num() == TargetWeights.Num());
-		for (int32 Index = 0; Index < TargetWeights.Num(); ++Index)
-		{
-			const FMeshToMeshVertData& MeshToMeshVertDatum = MeshToMeshVertData[Index];
-			
-			const uint16 VertIndex0 = MeshToMeshVertDatum.SourceMeshVertIndices[0];
-			const uint16 VertIndex1 = MeshToMeshVertDatum.SourceMeshVertIndices[1];
-			const uint16 VertIndex2 = MeshToMeshVertDatum.SourceMeshVertIndices[2];
-
-			TargetWeights[Index] = FMath::Clamp(
-				SourceWeights[VertIndex0] * MeshToMeshVertDatum.PositionBaryCoordsAndDist[0] +
-				SourceWeights[VertIndex1] * MeshToMeshVertDatum.PositionBaryCoordsAndDist[1] +
-				SourceWeights[VertIndex2] * MeshToMeshVertDatum.PositionBaryCoordsAndDist[2], 0.f, 1.f);
-		}
-	}
-
 	void TransferWeightMap(
 		const TConstArrayView<FVector2f>& InSourcePositions,
 		const TConstArrayView<FIntVector3>& SourceIndices,
@@ -120,7 +53,7 @@ namespace UE::Chaos::ClothAsset::Private
 		TArray<float> TargetWeights;
 		TargetWeights.SetNumUninitialized(TargetPositions.Num());
 
-		TransferWeightMap(SourcePositions, SourceIndices, SourceWeights, TargetPositions, TargetNormals, TargetIndices, TargetWeights);
+		FClothGeometryTools::TransferWeightMap(SourcePositions, SourceIndices, SourceWeights, TargetPositions, TargetNormals, TargetIndices, TArrayView<float>(TargetWeights));
 
 		for (int32 Index = 0; Index < TargetWeights.Num(); ++Index)
 		{
@@ -279,14 +212,14 @@ PRAGMA_ENABLE_DEPRECATION_WARNINGS
 										RemappedWeights);
 										break;
 								case EChaosClothAssetWeightMapTransferType::Use3DSimMesh:
-									Private::TransferWeightMap(
+									FClothGeometryTools::TransferWeightMap(
 										TransferClothFacade.GetSimPosition3D(),
 										TransferClothFacade.GetSimIndices3D(),
 										TransferClothFacade.GetWeightMap(InInputName),
 										ClothFacade.GetSimPosition3D(),
 										ClothFacade.GetSimNormal(),
 										ClothFacade.GetSimIndices3D(),
-										RemappedWeights);
+										TArrayView<float>(RemappedWeights));
 										break;
 								default: unimplemented();
 								}
diff --git a/Engine/Plugins/ChaosClothAssetEditor/Source/ChaosClothAssetDataflowNodes/Private/ChaosClothAsset/SelectionNode.cpp b/Engine/Plugins/ChaosClothAssetEditor/Source/ChaosClothAssetDataflowNodes/Private/ChaosClothAsset/SelectionNode.cpp
--- a/Engine/Plugins/ChaosClothAssetEditor/Source/ChaosClothAssetDataflowNodes/Private/ChaosClothAsset/SelectionNode.cpp
+++ b/Engine/Plugins/ChaosClothAssetEditor/Source/ChaosClothAssetDataflowNodes/Private/ChaosClothAsset/SelectionNode.cpp
@@ -1,24 +1,326 @@
 // Copyright Epic Games, Inc. All Rights Reserved.
 
 #include "ChaosClothAsset/SelectionNode.h"
+#include "ChaosClothAsset/ClothAsset.h"
 #include "ChaosClothAsset/ClothCollectionGroup.h"
 #include "ChaosClothAsset/CollectionClothFacade.h"
 #include "ChaosClothAsset/ClothDataflowTools.h"
+#include "ChaosClothAsset/ClothGeometryTools.h"
 #include "ChaosClothAsset/CollectionClothSelectionFacade.h"
+#include "ChaosClothAsset/WeightedValue.h"
 #include "Dataflow/DataflowInputOutput.h"
+#include "Dataflow/DataflowObject.h"
 
 #include UE_INLINE_GENERATED_CPP_BY_NAME(SelectionNode)
 
 #define LOCTEXT_NAMESPACE "FChaosClothAssetSelectionNode"
 
+namespace UE::Chaos::ClothAsset::Private
+{
+void ConvertWeightMapToVertexSelection(const TArray<float>& WeightMap, const float TransferSelectionThreshold, TSet<int32>& OutSelection)
+{
+	OutSelection.Reset();
+	for (int32 Index = 0; Index < WeightMap.Num(); ++Index)
+	{
+		if (WeightMap[Index] >= TransferSelectionThreshold)
+		{
+			OutSelection.Add(Index);
+		}
+	}
+}
+
+void ConvertWeightMapToFaceSelection(const TArray<float>& WeightMap, const float TransferSelectionThreshold, const TConstArrayView<FIntVector3>& Indices, TSet<int32>& OutSelection)
+{
+	OutSelection.Reset();
+	for (int32 FaceIndex = 0; FaceIndex < Indices.Num(); ++FaceIndex)
+	{
+		const FIntVector& Element = Indices[FaceIndex];
+		if (WeightMap[Element[0]] >= TransferSelectionThreshold &&
+			WeightMap[Element[1]] >= TransferSelectionThreshold &&
+			WeightMap[Element[2]] >= TransferSelectionThreshold)
+		{
+			OutSelection.Add(FaceIndex);
+		}
+	}
+}
+
+template<bool bIsSecondarySelection>
+bool TransferSelectionSet(const TSharedRef<const FManagedArrayCollection>& TransferClothCollection, const TSharedRef<FManagedArrayCollection>& ClothCollection,
+	const FName& InInputName, const FName& SelectionGroupName, const EChaosClothAssetWeightMapTransferType SimTransferType, const float TransferSelectionThreshold, TSet<int32>& OutSelection)
+{
+	FCollectionClothConstFacade ClothFacade(ClothCollection);
+	FCollectionClothConstFacade TransferClothFacade(TransferClothCollection);
+	FCollectionClothSelectionConstFacade TransferSelectionFacade(TransferClothCollection);
+
+	const bool bIsValidRenderSelection = SelectionGroupName == ClothCollectionGroup::RenderFaces || SelectionGroupName == ClothCollectionGroup::RenderVertices; 
+	const bool bIsValidSimSelection = SelectionGroupName == ClothCollectionGroup::SimFaces || SelectionGroupName == ClothCollectionGroup::SimVertices2D || SelectionGroupName == ClothCollectionGroup::SimVertices3D;
+
+	if (!bIsValidRenderSelection && !bIsValidSimSelection)
+	{
+		return false;
+	}
+
+	// Get the selection as a vertex set.
+	TSet<int32> TransferSet;
+	const FName DesiredTransferGroup = bIsValidRenderSelection ? ClothCollectionGroup::RenderVertices :
+		(SimTransferType == EChaosClothAssetWeightMapTransferType::Use2DSimMesh ? ClothCollectionGroup::SimVertices2D : ClothCollectionGroup::SimVertices3D);
+	if (!FClothGeometryTools::ConvertSelectionToNewGroupType(TransferClothCollection, InInputName, DesiredTransferGroup, bIsSecondarySelection, TransferSet))
+	{
+		return false;
+	}
+
+	// Convert to weights that are 0 on unselected vertices and 1 on selected vertices
+	TArray<float> TransferWeights;
+	TransferWeights.SetNumZeroed(TransferClothCollection->NumElements(DesiredTransferGroup));
+	for (const int32 SetIndex : TransferSet)
+	{
+		if (TransferWeights.IsValidIndex(SetIndex))
+		{
+			TransferWeights[SetIndex] = 1.f;
+		}
+	}
+
+	// Transfer weights
+	TArray<float> RemappedWeights;
+	RemappedWeights.SetNumZeroed(ClothCollection->NumElements(DesiredTransferGroup));
+	if (bIsValidRenderSelection)
+	{
+		FClothGeometryTools::TransferWeightMap(
+			TransferClothFacade.GetRenderPosition(),
+			TransferClothFacade.GetRenderIndices(),
+			TransferWeights,
+			ClothFacade.GetRenderPosition(),
+			ClothFacade.GetRenderNormal(),
+			ClothFacade.GetRenderIndices(),
+			TArrayView<float>(RemappedWeights));
+
+		OutSelection.Reset();
+
+		if (SelectionGroupName == ClothCollectionGroup::RenderFaces)
+		{
+			ConvertWeightMapToFaceSelection(RemappedWeights, TransferSelectionThreshold, ClothFacade.GetRenderIndices(), OutSelection);
+		}
+		else
+		{
+			check(SelectionGroupName == ClothCollectionGroup::RenderVertices);
+			ConvertWeightMapToVertexSelection(RemappedWeights, TransferSelectionThreshold, OutSelection);
+		}
+	}
+	else if (SimTransferType == EChaosClothAssetWeightMapTransferType::Use2DSimMesh)
+	{
+		TArray<FVector3f> TransferSimPositions2DAs3D;
+		TConstArrayView<FVector2f> TransferPositions2D = TransferClothFacade.GetSimPosition2D();
+		TransferSimPositions2DAs3D.SetNumUninitialized(TransferPositions2D.Num());
+		for (int32 Index = 0; Index < TransferPositions2D.Num(); ++Index)
+		{
+			TransferSimPositions2DAs3D[Index] = FVector3f(TransferPositions2D[Index], 0.f);
+		}
+
+		TConstArrayView<FVector2f> Positions2D = ClothFacade.GetSimPosition2D();
+		TArray<FVector3f> Positions2DAs3D;
+		TArray<FVector3f> NormalsZAxis;
+		Positions2DAs3D.SetNumUninitialized(Positions2D.Num());
+		NormalsZAxis.Init(FVector3f::ZAxisVector, Positions2D.Num());
+		for (int32 Index = 0; Index < Positions2D.Num(); ++Index)
+		{
+			Positions2DAs3D[Index] = FVector3f(Positions2D[Index], 0.f);
+		}
+
+		FClothGeometryTools::TransferWeightMap(
+			TConstArrayView<FVector3f>(TransferSimPositions2DAs3D),
+			TransferClothFacade.GetSimIndices2D(),
+			TransferWeights,
+			TConstArrayView<FVector3f>(Positions2DAs3D),
+			TConstArrayView<FVector3f>(NormalsZAxis),
+			ClothFacade.GetSimIndices2D(),
+			TArrayView<float>(RemappedWeights));
+
+		if (SelectionGroupName == ClothCollectionGroup::SimFaces)
+		{
+			ConvertWeightMapToFaceSelection(RemappedWeights, TransferSelectionThreshold, ClothFacade.GetSimIndices2D(), OutSelection);
+		}
+		else if (SelectionGroupName == ClothCollectionGroup::SimVertices2D)
+		{
+			ConvertWeightMapToVertexSelection(RemappedWeights, TransferSelectionThreshold, OutSelection);
+		}
+		else
+		{
+			check(SelectionGroupName == ClothCollectionGroup::SimVertices3D);
+			TSet<int32> Selection2D;
+			ConvertWeightMapToVertexSelection(RemappedWeights, TransferSelectionThreshold, Selection2D);
+			const TConstArrayView<int32> SimVertex3DLookup = ClothFacade.GetSimVertex3DLookup();
+			for (const int32 Vertex2D : Selection2D)
+			{
+				OutSelection.Add(SimVertex3DLookup[Vertex2D]);
+			}
+		}
+	}
+	else
+	{
+		check(SimTransferType == EChaosClothAssetWeightMapTransferType::Use3DSimMesh);
+		FClothGeometryTools::TransferWeightMap(
+			TransferClothFacade.GetSimPosition3D(),
+			TransferClothFacade.GetSimIndices3D(),
+			TransferWeights,
+			ClothFacade.GetSimPosition3D(),
+			ClothFacade.GetSimNormal(),
+			ClothFacade.GetSimIndices3D(),
+			TArrayView<float>(RemappedWeights));
+
+		if (SelectionGroupName == ClothCollectionGroup::SimFaces)
+		{
+			ConvertWeightMapToFaceSelection(RemappedWeights, TransferSelectionThreshold, ClothFacade.GetSimIndices3D(), OutSelection);
+		}
+		else if (SelectionGroupName == ClothCollectionGroup::SimVertices3D)
+		{
+			ConvertWeightMapToVertexSelection(RemappedWeights, TransferSelectionThreshold, OutSelection);
+		}
+		else
+		{
+			check(SelectionGroupName == ClothCollectionGroup::SimVertices2D);
+			TSet<int32> Selection3D;
+			ConvertWeightMapToVertexSelection(RemappedWeights, TransferSelectionThreshold, Selection3D);
+			const TConstArrayView<TArray<int32>> SimVertex2DLookup = ClothFacade.GetSimVertex2DLookup();
+			for (const int32 Vertex3D : Selection3D)
+			{
+				for (const int32 Vertex2D : SimVertex2DLookup[Vertex3D])
+				{
+					OutSelection.Add(Vertex2D);
+				}
+			}
+		}
+	}
+	return true;
+}
+
+
+void SetIndices(const TSet<int32>& InputSet, const TSet<int32>& FinalSet, EChaosClothAssetSelectionOverrideType OverrideType, TSet<int32>& Indices, TSet<int32>& RemoveIndices)
+{
+	if (InputSet.IsEmpty() || OverrideType == EChaosClothAssetSelectionOverrideType::ReplaceAll)
+	{
+		Indices = FinalSet;
+		RemoveIndices.Reset();
+		return;
+	}
+
+	Indices = FinalSet.Difference(InputSet);
+	RemoveIndices = InputSet.Difference(FinalSet);
+}
+
+void CalculateFinalSet(const TSet<int32>& InputSet, TSet<int32>& FinalSet, EChaosClothAssetSelectionOverrideType OverrideType, const TSet<int32>& Indices, const TSet<int32>& RemoveIndices)
+{
+	if (InputSet.IsEmpty() || OverrideType == EChaosClothAssetSelectionOverrideType::ReplaceAll)
+	{
+		FinalSet = Indices;
+		return;
+	}
+
+	FinalSet = InputSet;
+	FinalSet.Append(Indices);
+	if (!RemoveIndices.IsEmpty())
+	{
+		FinalSet = FinalSet.Difference(RemoveIndices);
+	}
+}
+}
+
 FChaosClothAssetSelectionNode::FChaosClothAssetSelectionNode(const Dataflow::FNodeParameters& InParam, FGuid InGuid)
-	: FDataflowNode(InParam, InGuid)
+	: FDataflowTerminalNode(InParam, InGuid)
 {
 	RegisterInputConnection(&Collection);
+	RegisterInputConnection(&InputName.StringValue, GET_MEMBER_NAME_CHECKED(FChaosClothAssetConnectableIStringValue, StringValue));
+	RegisterInputConnection(&TransferCollection);
 	RegisterOutputConnection(&Collection, &Collection);
 	RegisterOutputConnection(&Name);
 }
 
+
+void FChaosClothAssetSelectionNode::SetAssetValue(TObjectPtr<UObject> Asset, Dataflow::FContext& Context) const
+{
+	using namespace UE::Chaos::ClothAsset;
+
+	if (UChaosClothAsset* const ClothAsset = Cast<UChaosClothAsset>(Asset.Get()))
+	{
+		PRAGMA_DISABLE_DEPRECATION_WARNINGS  // TODO: Remove after adding the getter function (currently shelved!)  
+		if (UDataflow* const DataflowAsset = ClothAsset->DataflowAsset)
+		PRAGMA_ENABLE_DEPRECATION_WARNINGS
+		{
+			const TSharedPtr<Dataflow::FGraph, ESPMode::ThreadSafe> Dataflow = DataflowAsset->GetDataflow();
+			if (const TSharedPtr<FDataflowNode> BaseNode = Dataflow->FindBaseNode(this->GetGuid()))  // This is basically a safe const_cast
+			{
+				FChaosClothAssetSelectionNode* const MutableThis = static_cast<FChaosClothAssetSelectionNode*>(BaseNode.Get());
+				check(MutableThis == this);
+
+				// Make the name a valid attribute name, and replace the value in the UI
+				FWeightMapTools::MakeWeightMapName(MutableThis->Name);
+
+				const FName SelectionGroupName(*Group.Name);
+				const FName SelectionSecondaryGroupName(*SecondaryGroup.Name);
+
+				// Transfer selection if the transfer collection input has changed and is valid
+				FManagedArrayCollection InCollection = GetValue<FManagedArrayCollection>(Context, &Collection);
+				const TSharedRef<FManagedArrayCollection> ClothCollection = MakeShared<FManagedArrayCollection>(MoveTemp(InCollection));
+				FCollectionClothConstFacade ClothFacade(ClothCollection);
+				if (ClothFacade.HasValidSimulationData())  // Can only act on the collection if it is a valid cloth collection
+				{
+					FManagedArrayCollection InTransferCollection = GetValue<FManagedArrayCollection>(Context, &TransferCollection);
+					const TSharedRef<const FManagedArrayCollection> TransferClothCollection = MakeShared<const FManagedArrayCollection>(MoveTemp(InTransferCollection));
+					FCollectionClothConstFacade TransferClothFacade(TransferClothCollection);
+					FCollectionClothSelectionConstFacade TransferSelectionFacade(TransferClothCollection);
+
+					const FName InInputName = GetInputName(Context);
+					uint32 InTransferCollectionHash = HashCombineFast(GetTypeHash(InInputName), GetTypeHash(SelectionGroupName));
+					InTransferCollectionHash = HashCombineFast(InTransferCollectionHash, GetTypeHash(SelectionSecondaryGroupName));
+					InTransferCollectionHash = HashCombineFast(InTransferCollectionHash, (uint32)SimTransferType);
+					if (TransferClothFacade.HasValidSimulationData() && TransferSelectionFacade.IsValid() && InInputName != NAME_None && TransferSelectionFacade.HasSelection(InInputName))
+					{
+						InTransferCollectionHash = HashCombineFast(InTransferCollectionHash, GetTypeHash(TransferSelectionFacade.GetSelectionGroup(InInputName)));
+						const TArray<int32> SelectionAsArray = TransferSelectionFacade.GetSelectionSet(InInputName).Array();
+						InTransferCollectionHash = GetArrayHash(SelectionAsArray.GetData(), SelectionAsArray.Num(), InTransferCollectionHash);
+
+						if (TransferSelectionFacade.HasSelectionSecondarySet(InInputName))
+						{
+							InTransferCollectionHash = HashCombineFast(InTransferCollectionHash, GetTypeHash(TransferSelectionFacade.GetSelectionSecondaryGroup(InInputName)));
+							const TArray<int32> SecondarySelectionAsArray = TransferSelectionFacade.GetSelectionSecondarySet(InInputName).Array();
+							InTransferCollectionHash = GetArrayHash(SecondarySelectionAsArray.GetData(), SecondarySelectionAsArray.Num(), InTransferCollectionHash);
+						}
+					}
+					else
+					{
+						InTransferCollectionHash = 0;
+					}
+
+					if (TransferCollectionHash != InTransferCollectionHash)
+					{
+						MutableThis->TransferCollectionHash = InTransferCollectionHash;
+
+						if(TransferCollectionHash)
+						{
+							TSet<int32> PrimaryFinalSelection;
+							if (Private::TransferSelectionSet<false>(TransferClothCollection, ClothCollection, InInputName, SelectionGroupName, SimTransferType, TransferSelectionThreshold, PrimaryFinalSelection))
+							{
+								TSet<int32> InputSelection;
+								FClothGeometryTools::ConvertSelectionToNewGroupType(ClothCollection, InInputName, SelectionGroupName, false, InputSelection);
+
+								MutableThis->SetIndices(InputSelection, PrimaryFinalSelection);
+							}
+
+							TSet<int32> SecondaryFinalSelection;
+							if (Private::TransferSelectionSet<true>(TransferClothCollection, ClothCollection, InInputName, SelectionSecondaryGroupName, SimTransferType, TransferSelectionThreshold, SecondaryFinalSelection))
+							{
+								TSet<int32> InputSelection;
+								FClothGeometryTools::ConvertSelectionToNewGroupType(ClothCollection, InInputName, SelectionGroupName, true, InputSelection);
+
+								MutableThis->SetSecondaryIndices(InputSelection, SecondaryFinalSelection);
+							}
+						}
+					}
+				}
+			}
+		}
+	}
+}
+
 void FChaosClothAssetSelectionNode::Evaluate(Dataflow::FContext& Context, const FDataflowOutput* Out) const
 {
 	using namespace UE::Chaos::ClothAsset;
@@ -64,14 +366,19 @@ void FChaosClothAssetSelectionNode::Evaluate(Dataflow::FContext& Context, const
 
 	if (Out->IsA<FManagedArrayCollection>(&Collection))
 	{
-		if (Name.IsEmpty() || Group.Name.IsEmpty())
+		// Evaluate InputName
+		const FName InInputName = GetInputName(Context);
+		InputName.StringValue_Override = GetValue<FString>(Context, &InputName.StringValue, UE::Chaos::ClothAsset::FWeightMapTools::NotOverridden);
+
+		const FName SelectionName(Name.IsEmpty() ? InInputName : FName(Name));
+
+		if (SelectionName == NAME_None || Group.Name.IsEmpty())
 		{
 			const FManagedArrayCollection& SelectionCollection = GetValue<FManagedArrayCollection>(Context, &Collection);
 			SetValue(Context, SelectionCollection, &Collection);
 			return;
 		}
 
-		const FName SelectionName(*Name);
 		const FName SelectionGroupName(*Group.Name);
 
 		FManagedArrayCollection InSelectionCollection = GetValue<FManagedArrayCollection>(Context, &Collection);
@@ -81,23 +388,35 @@ void FChaosClothAssetSelectionNode::Evaluate(Dataflow::FContext& Context, const
 		SelectionFacade.DefineSchema();
 		check(SelectionFacade.IsValid());
 
+		TSet<int32> InputSelectionSet;
+		FClothGeometryTools::ConvertSelectionToNewGroupType(SelectionCollection, InInputName, SelectionGroupName, false, InputSelectionSet);
+		TSet<int32> FinalSet;
+		CalculateFinalSet(InputSelectionSet, FinalSet);
+
 		TSet<int32>& SelectionSet = SelectionFacade.FindOrAddSelectionSet(SelectionName, SelectionGroupName);
-		CopyIntoSelection(SelectionCollection, SelectionGroupName, Indices, SelectionSet);
+		CopyIntoSelection(SelectionCollection, SelectionGroupName, FinalSet, SelectionSet);
 
 		if (!SecondaryGroup.Name.IsEmpty() && !SecondaryIndices.IsEmpty())
 		{
-			const FName SecondarySelectionName(*Name);
-			const FName SecondarySelectionGroupName = (*SecondaryGroup.Name);		
+			const FName SecondarySelectionName(Name.IsEmpty() ? InInputName : FName(Name));
+			const FName SecondarySelectionGroupName = (*SecondaryGroup.Name);
+
+			InputSelectionSet.Reset();
+			FClothGeometryTools::ConvertSelectionToNewGroupType(SelectionCollection, InInputName, SelectionGroupName, true, InputSelectionSet);
+			FinalSet.Reset();
+			CalculateFinalSecondarySet(InputSelectionSet, FinalSet);
 			TSet<int32>& SecondarySelectionSet = SelectionFacade.FindOrAddSelectionSecondarySet(SecondarySelectionName, SecondarySelectionGroupName);
 
-			CopyIntoSelection(SelectionCollection, SecondarySelectionGroupName, SecondaryIndices, SecondarySelectionSet);
+			CopyIntoSelection(SelectionCollection, SecondarySelectionGroupName, FinalSet, SecondarySelectionSet);
 		}
 
 		SetValue(Context, MoveTemp(*SelectionCollection), &Collection);
 	}
 	else if (Out->IsA<FString>(&Name))
 	{
-		SetValue(Context, Name, &Name);
+		FString InputNameString = GetValue<FString>(Context, &InputName.StringValue);
+		UE::Chaos::ClothAsset::FWeightMapTools::MakeWeightMapName(InputNameString);
+		SetValue(Context, InputNameString.IsEmpty() ? Name : InputNameString, &Name);
 	}
 }
 
@@ -154,4 +473,33 @@ PRAGMA_DISABLE_DEPRECATION_WARNINGS
 PRAGMA_ENABLE_DEPRECATION_WARNINGS
 }
 
+
+FName FChaosClothAssetSelectionNode::GetInputName(Dataflow::FContext& Context) const
+{
+	FString InputNameString = GetValue<FString>(Context, &InputName.StringValue);
+	UE::Chaos::ClothAsset::FWeightMapTools::MakeWeightMapName(InputNameString);
+	const FName InInputName(*InputNameString);
+	return InInputName != NAME_None ? InInputName : FName(Name);
+}
+
+void FChaosClothAssetSelectionNode::SetIndices(const TSet<int32>& InputSet, const TSet<int32>& FinalSet)
+{
+	UE::Chaos::ClothAsset::Private::SetIndices(InputSet, FinalSet, SelectionOverrideType, Indices, RemoveIndices);
+}
+
+void FChaosClothAssetSelectionNode::SetSecondaryIndices(const TSet<int32>& InputSet, const TSet<int32>& FinalSet)
+{
+	UE::Chaos::ClothAsset::Private::SetIndices(InputSet, FinalSet, SelectionOverrideType, SecondaryIndices, RemoveSecondaryIndices);
+}
+
+void FChaosClothAssetSelectionNode::CalculateFinalSet(const TSet<int32>& InputSet, TSet<int32>& FinalSet) const
+{
+	UE::Chaos::ClothAsset::Private::CalculateFinalSet(InputSet, FinalSet, SelectionOverrideType, Indices, RemoveIndices);
+}
+
+void FChaosClothAssetSelectionNode::CalculateFinalSecondarySet(const TSet<int32>& InputSet, TSet<int32>& FinalSet) const
+{
+	UE::Chaos::ClothAsset::Private::CalculateFinalSet(InputSet, FinalSet, SelectionOverrideType, SecondaryIndices, RemoveSecondaryIndices);
+}
+
 #undef LOCTEXT_NAMESPACE
diff --git a/Engine/Plugins/ChaosClothAssetEditor/Source/ChaosClothAssetDataflowNodes/Public/ChaosClothAsset/SelectionNode.h b/Engine/Plugins/ChaosClothAssetEditor/Source/ChaosClothAssetDataflowNodes/Public/ChaosClothAsset/SelectionNode.h
--- a/Engine/Plugins/ChaosClothAssetEditor/Source/ChaosClothAssetDataflowNodes/Public/ChaosClothAsset/SelectionNode.h
+++ b/Engine/Plugins/ChaosClothAssetEditor/Source/ChaosClothAssetDataflowNodes/Public/ChaosClothAsset/SelectionNode.h
@@ -4,6 +4,7 @@
 
 #include "Dataflow/DataflowNode.h"
 #include "GeometryCollection/ManagedArrayCollection.h"
+#include "ChaosClothAsset/AddWeightMapNode.h"
 #include "SelectionNode.generated.h"
 
 /** What type of element is selected in the Selection */
@@ -42,11 +43,21 @@ struct FChaosClothAssetNodeSelectionGroup
 	FString Name;
 };
 
+/** How the sets stored on the SelectionNode should be applied to an existing set. If no set exists, it is treated as empty.*/
+UENUM()
+enum class EChaosClothAssetSelectionOverrideType : uint8
+{
+	/** Replace the existing set with values in [Secondary]Indices.*/
+	ReplaceAll,
+	/** Append values in [Secondary]Indices. Remove values in [Secondary]RemoveIndices.*/
+	Modify
+};
+
 PRAGMA_DISABLE_DEPRECATION_WARNINGS  // For EChaosClothAssetSelectionType
 
 /** Integer index set selection node. */
 USTRUCT(Meta = (DataflowCloth))
-struct FChaosClothAssetSelectionNode : public FDataflowNode
+struct FChaosClothAssetSelectionNode : public FDataflowTerminalNode
 {
 	GENERATED_USTRUCT_BODY()
 	DATAFLOW_NODE_DEFINE_INTERNAL(FChaosClothAssetSelectionNode, "Selection", "Cloth", "Cloth Selection")
@@ -56,10 +67,29 @@ public:
 	UPROPERTY(Meta = (Dataflowinput, DataflowOutput, DataflowPassthrough = "Collection"))
 	FManagedArrayCollection Collection;
 
+	/**
+	 * The collection used to transfer sets from.
+	 * Connecting a collection containing a set with Input Name (or Name if Input Name is empty) and compatible Group
+	 * will transfer the set to the input collection topology.
+	 * Note this operation only happens once when the TransferCollection is first connected, or updated.
+	 * Changing the InputName or the TransferType will also redo the transfer operation.
+	 */
+	UPROPERTY(Meta = (DataflowInput))
+	FManagedArrayCollection TransferCollection;
+
 	/** The name to give the selection attribute */
 	UPROPERTY(EditAnywhere, Category = "Selection", Meta = (DataflowOutput))
 	FString Name;
 
+	/** The name to populate this set from and override based on Selection Override Type. Name will be used if Input Name is empty.*/
+	UPROPERTY(EditAnywhere, Category = "Selection")
+	FChaosClothAssetConnectableIStringValue InputName = { TEXT("") };
+
+	/** How to apply this node's Indices onto existing sets. Changing this value will change the output set.
+	 *  To change how the node's stored indices are calculated, change the equivalent value on the Selection Tool context.*/
+	UPROPERTY(EditAnywhere, Category = "Selection")
+	EChaosClothAssetSelectionOverrideType SelectionOverrideType = EChaosClothAssetSelectionOverrideType::ReplaceAll;
+
 	/** The type of element the selection refers to */
 	UE_DEPRECATED(5.4, "Use Group instead")
 	UPROPERTY()
@@ -73,6 +103,10 @@ public:
 	UPROPERTY(EditAnywhere, Category = "Selection", Meta = (ClampMin = "0"))
 	TSet<int32> Indices;
 
+	/** Indices to remove from the Input selection */
+	UPROPERTY(EditAnywhere, Category = "Selection", Meta = (ClampMin = "0", EditCondition = "SelectionOverrideType == EChaosClothAssetSelectionOverrideType::Modify"))
+	TSet<int32> RemoveIndices;
+
 	/** The type of element the secondary selection refers to */
 	UPROPERTY(EditAnywhere, Category = "Secondary Selection")
 	FChaosClothAssetNodeSelectionGroup SecondaryGroup;
@@ -81,19 +115,50 @@ public:
 	UPROPERTY(EditAnywhere, Category = "Secondary Selection", Meta = (ClampMin = "0"))
 	TSet<int32> SecondaryIndices;
 
+	/** Secondary set indices to remove from the Input selection */
+	UPROPERTY(EditAnywhere, Category = "Secondary Selection", Meta = (ClampMin = "0", EditCondition = "SelectionOverrideType == EChaosClothAssetSelectionOverrideType::Modify"))
+	TSet<int32> RemoveSecondaryIndices;
+
+	/**
+	 * The type of transfer used to transfer the sim mesh sets when a TransferCollection is connected.
+	 * This property is disabled when no TransferCollection input has been connected.
+	 */
+	UPROPERTY(EditAnywhere, Category = "Selection Transfer", Meta = (EditCondition = "TransferCollectionHash != 0"))
+	EChaosClothAssetWeightMapTransferType SimTransferType = EChaosClothAssetWeightMapTransferType::Use2DSimMesh;
+
+	/**
+	 * Selections are internally converted to maps in order to do the transfer and then converted back.
+	 * This value is used to do the conversion back. Decrease this value to (possibly) expand the converted selection.
+	 */
+	UPROPERTY(EditAnywhere, Category = "Selection Transfer", Meta = (EditCondition = "TransferCollectionHash != 0", ClampMin = "0", ClampMax = "1"))
+	float TransferSelectionThreshold = 0.95f;
+
 	FChaosClothAssetSelectionNode(const Dataflow::FNodeParameters& InParam, FGuid InGuid = FGuid::NewGuid());
 
 	/** Return a cached array of all the groups used by the input collection during at the time of the latest evaluation. */
 	const TArray<FName>& GetCachedCollectionGroupNames() const { return CachedCollectionGroupNames; }
 
+	FName CHAOSCLOTHASSETDATAFLOWNODES_API GetInputName(Dataflow::FContext& Context) const;
+
+	// Set Indices based on SelectionOverrideType
+	void CHAOSCLOTHASSETDATAFLOWNODES_API SetIndices(const TSet<int32>& InputSet, const TSet<int32>& FinalSet);
+	void CHAOSCLOTHASSETDATAFLOWNODES_API SetSecondaryIndices(const TSet<int32>& InputSet, const TSet<int32>& FinalSet);
+
+	void CHAOSCLOTHASSETDATAFLOWNODES_API CalculateFinalSet(const TSet<int32>& InputSet, TSet<int32>& FinalSet) const;
+	void CHAOSCLOTHASSETDATAFLOWNODES_API CalculateFinalSecondarySet(const TSet<int32>& InputSet, TSet<int32>& FinalSet) const;
+
 private:
 
+	virtual void SetAssetValue(TObjectPtr<UObject> Asset, Dataflow::FContext& Context) const override;
 	virtual void Evaluate(Dataflow::FContext& Context, const FDataflowOutput* Out) const override;
 	virtual void OnSelected(Dataflow::FContext& Context) override;
 	virtual void OnDeselected() override;
 	virtual void Serialize(FArchive& Ar);
 
 	TArray<FName> CachedCollectionGroupNames;
+
+	UPROPERTY()
+	uint32 TransferCollectionHash = 0;
 };
 
 PRAGMA_ENABLE_DEPRECATION_WARNINGS  // For EChaosClothAssetSelectionType
diff --git a/Engine/Plugins/ChaosClothAssetEditor/Source/ChaosClothAssetEditorTools/Private/ChaosClothAsset/ClothEditorToolBuilders.cpp b/Engine/Plugins/ChaosClothAssetEditor/Source/ChaosClothAssetEditorTools/Private/ChaosClothAsset/ClothEditorToolBuilders.cpp
--- a/Engine/Plugins/ChaosClothAssetEditor/Source/ChaosClothAssetEditorTools/Private/ChaosClothAsset/ClothEditorToolBuilders.cpp
+++ b/Engine/Plugins/ChaosClothAssetEditor/Source/ChaosClothAssetEditorTools/Private/ChaosClothAsset/ClothEditorToolBuilders.cpp
@@ -4,6 +4,7 @@
 #include "ChaosClothAsset/ClothPatternVertexType.h"
 #include "TargetInterfaces/PrimitiveComponentBackedTarget.h"
 #include "ChaosClothAsset/ClothEditorContextObject.h"
+#include "ChaosClothAsset/SelectionNode.h"
 #include "ToolContextInterfaces.h"
 #include "ToolTargetManager.h"
 #include "ContextObjectStore.h"
@@ -53,7 +54,12 @@ const FToolTargetTypeRequirements& UClothMeshSelectionToolBuilder::GetTargetRequ
 
 bool UClothMeshSelectionToolBuilder::CanBuildTool(const FToolBuilderState& SceneState) const
 {
-	return (SceneState.TargetManager->CountSelectedAndTargetable(SceneState, GetTargetRequirements()) == 1);
+	if (UClothEditorContextObject* const ContextObject = SceneState.ToolManager->GetContextObjectStore()->FindContext<UClothEditorContextObject>())
+	{
+		return ContextObject->GetSingleSelectedNodeOfType<FChaosClothAssetSelectionNode>() != nullptr && (SceneState.TargetManager->CountSelectedAndTargetable(SceneState, GetTargetRequirements()) == 1);
+	}
+
+	return false;
 }
 
 UInteractiveTool* UClothMeshSelectionToolBuilder::BuildTool(const FToolBuilderState& SceneState) const
diff --git a/Engine/Plugins/ChaosClothAssetEditor/Source/ChaosClothAssetEditorTools/Private/ChaosClothAsset/ClothMeshSelectionTool.cpp b/Engine/Plugins/ChaosClothAssetEditor/Source/ChaosClothAssetEditorTools/Private/ChaosClothAsset/ClothMeshSelectionTool.cpp
--- a/Engine/Plugins/ChaosClothAssetEditor/Source/ChaosClothAssetEditorTools/Private/ChaosClothAsset/ClothMeshSelectionTool.cpp
+++ b/Engine/Plugins/ChaosClothAssetEditor/Source/ChaosClothAssetEditorTools/Private/ChaosClothAsset/ClothMeshSelectionTool.cpp
@@ -3,6 +3,7 @@
 #include "ChaosClothAsset/ClothMeshSelectionTool.h"
 #include "ChaosClothAsset/ClothCollectionGroup.h"
 #include "ChaosClothAsset/ClothEditorContextObject.h"
+#include "ChaosClothAsset/ClothGeometryTools.h"
 #include "ChaosClothAsset/ClothPatternVertexType.h"
 #include "ChaosClothAsset/SelectionNode.h"
 #include "ChaosClothAsset/CollectionClothFacade.h"
@@ -310,7 +311,8 @@ void UClothMeshSelectionTool::Setup()
 	// Initialize the Selection from the selected Dataflow node
 	FString ExistingSelectionName;
 	FGroupTopologySelection ExistingNodeSelection;
-	GetSelectedNodeInfo(ExistingSelectionName, ExistingNodeSelection);
+	EChaosClothAssetSelectionOverrideType ExistingOverrideType;
+	GetSelectedNodeInfo(ExistingSelectionName, ExistingNodeSelection, ExistingOverrideType);
 
 	constexpr bool bBroadcastChange = false;
 	SelectionMechanic->SetSelection(ExistingNodeSelection, bBroadcastChange);
@@ -337,6 +339,7 @@ void UClothMeshSelectionTool::Setup()
 			bAnyChangeMade = true;
 		}
 	});
+	ToolProperties->SelectionOverrideType = ExistingOverrideType;
 
 
 	// 
@@ -365,9 +368,9 @@ void UClothMeshSelectionTool::OnShutdown(EToolShutdownType ShutdownType)
 
 	// Invalidate the node even if we are hitting cancel. We could have saved new selection information to the node by switching primary/secondary modes and we'd expect 
 	// that information to make its way into the ClothCollection
-	if (FChaosClothAssetSelectionNode* const MeshSelectionNode = ClothEditorContextObject->GetSingleSelectedNodeOfType<FChaosClothAssetSelectionNode>())
+	if (SelectionNodeToUpdate)
 	{
-		MeshSelectionNode->Invalidate();
+		SelectionNodeToUpdate->Invalidate();
 	}
 
 	SelectionMechanic->Properties->SaveProperties(this);
@@ -409,7 +412,7 @@ void UClothMeshSelectionTool::OnTick(float DeltaTime)
 
 bool UClothMeshSelectionTool::CanAccept() const
 {
-	return bAnyChangeMade;
+	return bAnyChangeMade || SelectionNodeToUpdate->SelectionOverrideType != ToolProperties->SelectionOverrideType;
 }
 
 FBox UClothMeshSelectionTool::GetWorldSpaceFocusBox()
@@ -424,12 +427,12 @@ void UClothMeshSelectionTool::SetClothEditorContextObject(TObjectPtr<UClothEdito
 	ClothEditorContextObject = InClothEditorContextObject;
 }
 
-bool UClothMeshSelectionTool::GetSelectedNodeInfo(FString& OutSelectionName, UE::Geometry::FGroupTopologySelection& OutSelection)
+bool UClothMeshSelectionTool::GetSelectedNodeInfo(FString& OutSelectionName, UE::Geometry::FGroupTopologySelection& OutSelection, EChaosClothAssetSelectionOverrideType& OutOverrideType)
 {
 	using namespace UE::Chaos::ClothAsset;
 
-	const FChaosClothAssetSelectionNode* const MeshSelectionNode = ClothEditorContextObject->GetSingleSelectedNodeOfType<FChaosClothAssetSelectionNode>();
-	check(MeshSelectionNode);
+	SelectionNodeToUpdate = ClothEditorContextObject->GetSingleSelectedNodeOfType<FChaosClothAssetSelectionNode>();
+	check(SelectionNodeToUpdate);
 
 	// We need to sanitize the incoming indices, as the user can manually set them to anything on the node
 
@@ -461,7 +464,7 @@ bool UClothMeshSelectionTool::GetSelectedNodeInfo(FString& OutSelectionName, UE:
 		});
 	};
 
-	auto ReadFromNode = [this, &AppendVerticesIfValid, &AppendFacesIfValid, &OutSelection](const FChaosClothAssetNodeSelectionGroup& SourceGroup, const TSet<int32>& SourceIndices)
+	auto AppendSet = [this, &AppendVerticesIfValid, &AppendFacesIfValid, &OutSelection](const FChaosClothAssetNodeSelectionGroup& SourceGroup, const TSet<int32>& SourceIndices)
 	{
 		if (SourceGroup.Name == ClothCollectionGroup::SimVertices2D.ToString() ||
 			SourceGroup.Name == ClothCollectionGroup::SimVertices3D.ToString() ||
@@ -489,16 +492,41 @@ bool UClothMeshSelectionTool::GetSelectedNodeInfo(FString& OutSelectionName, UE:
 		}
 	};
 
-	if (ToolProperties && ToolProperties->bSecondarySelection)
+	const bool bSecondarySelection = ToolProperties && ToolProperties->bSecondarySelection;
+
+	// Get the input set.
+	InputSelectionSet.Reset();
+	if (ClothEditorContextObject)
+	{
+		ensure(ClothEditorContextObject->IsUsingInputCollection());
+		if (TSharedPtr<const FManagedArrayCollection> ClothCollection = ClothEditorContextObject->GetSelectedClothCollection().Pin())
+		{
+			if (TSharedPtr<Dataflow::FEngineContext> DataflowContext = ClothEditorContextObject->GetDataflowContext().Pin())
+			{
+				using namespace UE::Chaos::ClothAsset;
+				const FName InputName = SelectionNodeToUpdate->GetInputName(*DataflowContext);
+				const FName GroupName = bSecondarySelection ? FName(*SelectionNodeToUpdate->SecondaryGroup.Name) : FName(*SelectionNodeToUpdate->Group.Name);
+
+				FClothGeometryTools::ConvertSelectionToNewGroupType(ClothCollection.ToSharedRef(), InputName, GroupName, bSecondarySelection, InputSelectionSet);
+			}
+		}
+	}
+
+	if (bSecondarySelection)
 	{
-		ReadFromNode(MeshSelectionNode->SecondaryGroup, MeshSelectionNode->SecondaryIndices);
+		TSet<int32> FinalSet;
+		SelectionNodeToUpdate->CalculateFinalSecondarySet(InputSelectionSet, FinalSet);
+		AppendSet(SelectionNodeToUpdate->SecondaryGroup, FinalSet);
 	}
 	else
 	{
-		ReadFromNode(MeshSelectionNode->Group, MeshSelectionNode->Indices);
+		TSet<int32> FinalSet;
+		SelectionNodeToUpdate->CalculateFinalSet(InputSelectionSet, FinalSet);
+		AppendSet(SelectionNodeToUpdate->Group, FinalSet);
 	}
 
-	OutSelectionName = MeshSelectionNode->Name;
+	OutSelectionName = SelectionNodeToUpdate->Name;
+	OutOverrideType = SelectionNodeToUpdate->SelectionOverrideType;
 
 	return true;
 }
@@ -550,12 +578,10 @@ void UClothMeshSelectionTool::UpdateSelectedNode()
 		}
 	}
 
-	FChaosClothAssetSelectionNode* const MeshSelectionNode = ClothEditorContextObject->GetSingleSelectedNodeOfType<FChaosClothAssetSelectionNode>();
-	check(MeshSelectionNode);
+	check(SelectionNodeToUpdate);
 
-	auto WriteToNode = [this, &Indices, &GroupName, MeshSelectionNode](FChaosClothAssetNodeSelectionGroup& TargetGroup, TSet<int32>& TargetIndices)
+	auto GetFinalSet = [this, &Indices, &GroupName](FChaosClothAssetNodeSelectionGroup& TargetGroup, TSet<int32>& TargetIndices)
 	{
-		MeshSelectionNode->Name = ToolProperties->Name;
 		TargetGroup.Name = GroupName.ToString();
 
 		if (SelectionMechanic->Properties->bSelectVertices && bHasNonManifoldMapping)
@@ -573,13 +599,19 @@ void UClothMeshSelectionTool::UpdateSelectedNode()
 		}
 	};
 
+	SelectionNodeToUpdate->Name = ToolProperties->Name;
+	SelectionNodeToUpdate->SelectionOverrideType = ToolProperties->SelectionOverrideType;
 	if (!ToolProperties->bSecondarySelection)
 	{
-		WriteToNode(MeshSelectionNode->Group, MeshSelectionNode->Indices);
+		TSet<int32> FinalSet;
+		GetFinalSet(SelectionNodeToUpdate->Group, FinalSet);
+		SelectionNodeToUpdate->SetIndices(InputSelectionSet, FinalSet);
 	}
 	else
 	{
-		WriteToNode(MeshSelectionNode->SecondaryGroup, MeshSelectionNode->SecondaryIndices);
+		TSet<int32> FinalSet;
+		GetFinalSet(SelectionNodeToUpdate->SecondaryGroup, FinalSet);
+		SelectionNodeToUpdate->SetSecondaryIndices(InputSelectionSet, FinalSet);
 	}
 	
 }
@@ -746,7 +778,8 @@ void UClothMeshSelectionTool::TogglePrimarySecondaryAction()
 	// Re-initialize the Selection from the selected Dataflow node
 	FString ExistingSelectionName;
 	FGroupTopologySelection ExistingNodeSelection;
-	GetSelectedNodeInfo(ExistingSelectionName, ExistingNodeSelection);
+	EChaosClothAssetSelectionOverrideType ExistingOverrideType;
+	GetSelectedNodeInfo(ExistingSelectionName, ExistingNodeSelection, ExistingOverrideType);
 
 	constexpr bool bBroadcastChange = false;
 	SelectionMechanic->SetSelection(ExistingNodeSelection, bBroadcastChange);
@@ -765,6 +798,7 @@ void UClothMeshSelectionTool::TogglePrimarySecondaryAction()
 	}
 
 	ToolProperties->Name = ExistingSelectionName;
+	ToolProperties->SelectionOverrideType = ExistingOverrideType;
 
 	PreviewMesh->FastNotifySecondaryTrianglesChanged();
 
diff --git a/Engine/Plugins/ChaosClothAssetEditor/Source/ChaosClothAssetEditorTools/Private/ChaosClothAsset/ClothMeshSelectionTool.h b/Engine/Plugins/ChaosClothAssetEditor/Source/ChaosClothAssetEditorTools/Private/ChaosClothAsset/ClothMeshSelectionTool.h
--- a/Engine/Plugins/ChaosClothAssetEditor/Source/ChaosClothAssetEditorTools/Private/ChaosClothAsset/ClothMeshSelectionTool.h
+++ b/Engine/Plugins/ChaosClothAssetEditor/Source/ChaosClothAssetEditorTools/Private/ChaosClothAsset/ClothMeshSelectionTool.h
@@ -9,6 +9,8 @@
 class UPolygonSelectionMechanic;
 class UClothEditorContextObject;
 class UPreviewMesh;
+struct FChaosClothAssetSelectionNode;
+enum class EChaosClothAssetSelectionOverrideType : uint8;
 
 namespace UE::Geometry
 {
@@ -100,9 +102,12 @@ class CHAOSCLOTHASSETEDITORTOOLS_API UClothMeshSelectionToolProperties : public
 
 public:
 
-	UPROPERTY(EditAnywhere, Transient, Category = Name, meta = (DisplayName = "Name", TransientToolProperty))
+	UPROPERTY(EditAnywhere, Transient, Category = Selection, meta = (DisplayName = "Name", TransientToolProperty))
 	FString Name;
 
+	UPROPERTY(EditAnywhere, Transient, Category = Selection, meta = (TransientToolProperty))
+	EChaosClothAssetSelectionOverrideType SelectionOverrideType;
+
 	/**
 	* Whether the user is editing the primary or secondary selection set. Stored as a property here so that the tool will remember what the
 	* user was doing the last time the tool shut down.
@@ -146,7 +151,7 @@ private:
 	virtual FBox GetWorldSpaceFocusBox() override;
 
 	void SetClothEditorContextObject(TObjectPtr<UClothEditorContextObject> InClothEditorContextObject);
-	bool GetSelectedNodeInfo(FString& OutMapName, UE::Geometry::FGroupTopologySelection& OutSelection);
+	bool GetSelectedNodeInfo(FString& OutMapName, UE::Geometry::FGroupTopologySelection& OutSelection, EChaosClothAssetSelectionOverrideType& OutOverrideType);
 	void UpdateSelectedNode();
 
 	UPROPERTY()
@@ -169,6 +174,8 @@ private:
 	TArray<int32> DynamicMeshToSelection;
 	TArray<TArray<int32>> SelectionToDynamicMesh;
 
+	FChaosClothAssetSelectionNode* SelectionNodeToUpdate = nullptr;
+	TSet<int32> InputSelectionSet;
 	//
 	// Action support
 	//
