Chaos Cloth Asset - Added missing references to the imported static meshes, textures and materials objects to the USDImport node.

#tests Reloaded and resaved some existing USD assets to make sure the missing references were correctly added and saved some newly imported ones.

diff --git a/Engine/Plugins/ChaosClothAssetEditor/Source/ChaosClothAssetDataflowNodes/Private/ChaosClothAsset/USDImportNode.cpp b/Engine/Plugins/ChaosClothAssetEditor/Source/ChaosClothAssetDataflowNodes/Private/ChaosClothAsset/USDImportNode.cpp
--- a/Engine/Plugins/ChaosClothAssetEditor/Source/ChaosClothAssetDataflowNodes/Private/ChaosClothAsset/USDImportNode.cpp
+++ b/Engine/Plugins/ChaosClothAssetEditor/Source/ChaosClothAssetDataflowNodes/Private/ChaosClothAsset/USDImportNode.cpp
@@ -389,13 +389,18 @@ void FChaosClothAssetUSDImportNode::Evaluate(Dataflow::FContext& Context, const
 	}
 }
 
-void FChaosClothAssetUSDImportNode::Serialize(FArchive& Archive)
+void FChaosClothAssetUSDImportNode::Serialize(FArchive& Ar)
 {
 	using namespace UE::Chaos::ClothAsset;
 
-	::Chaos::FChaosArchive ChaosArchive(Archive);
+	Ar.UsingCustomVersion(FFortniteMainBranchObjectVersion::GUID);
+
+	::Chaos::FChaosArchive ChaosArchive(Ar);
 	CollectionCache.Serialize(ChaosArchive);
-	if (Archive.IsLoading())
+
+	Ar << FileHash;
+
+	if (Ar.IsLoading())
 	{
 		// Make sure to always have a valid cloth collection on reload, some new attributes could be missing from the cached collection
 		const TSharedRef<FManagedArrayCollection> ClothCollection = MakeShared<FManagedArrayCollection>(MoveTemp(CollectionCache));
@@ -405,9 +410,12 @@ void FChaosClothAssetUSDImportNode::Serialize(FArchive& Archive)
 			ClothFacade.DefineSchema();
 		}
 		CollectionCache = MoveTemp(*ClothCollection);
-	}
 
-	Archive << FileHash;
+		if (Ar.CustomVer(FFortniteMainBranchObjectVersion::GUID) < FFortniteMainBranchObjectVersion::ChaosClothAssetUSDImportNodeAddAssetDependencies)
+		{
+			UpdateImportedAssets();
+		}
+	}
 }
 
 bool FChaosClothAssetUSDImportNode::ImportFromFile(const FString& UsdFilePath, const FString& AssetPath, FText& OutErrorText)
@@ -1121,6 +1129,10 @@ bool FChaosClothAssetUSDImportNode::ImportFromFile(const FString& UsdFilePath, c
 	
 	SlowTask.EnterProgressFrame(1.f);
 	SlowTask.ForceRefresh();
+
+	// Fill up the asset list from the imported USD assets
+	UpdateImportedAssets();
+
 	return true;
 
 #else  // #if USE_USD_SDK
@@ -1131,45 +1143,52 @@ bool FChaosClothAssetUSDImportNode::ImportFromFile(const FString& UsdFilePath, c
 #endif  // #else #if USE_USD_SDK
 }
 
-bool FChaosClothAssetUSDImportNode::ImportFromCache(const TSharedRef<FManagedArrayCollection>& OutClothCollection, FText& OutErrorText) const
+void FChaosClothAssetUSDImportNode::UpdateImportedAssets()
 {
-	using namespace UE::Chaos::ClothAsset;
-
-	// Initialize from collection cache
-	// TODO: Until we have a schema so that we can use the asset cache and remove the collection cache
-	*OutClothCollection = CollectionCache;
-
-	// Initialize from asset registry
-	TArray<FAssetData> AssetData;
+	ImportedAssets.Reset();
 
 	if (!PackagePath.IsEmpty())
 	{
+		TArray<FAssetData> AssetData;
+
 		const FAssetRegistryModule& AssetRegistryModule = FModuleManager::LoadModuleChecked<FAssetRegistryModule>("AssetRegistry");
 
 		const UClass* const Class = UStaticMesh::StaticClass();
 		constexpr bool bRecursive = true;
 		constexpr bool bIncludeOnlyOnDiskAssets = false;
 		AssetRegistryModule.Get().GetAssetsByPath(FName(*PackagePath), AssetData, bRecursive, bIncludeOnlyOnDiskAssets);
-	}
 
-	for (const FAssetData& AssetDatum : AssetData)
-	{
-		if (!(AssetDatum.IsUAsset() && AssetDatum.IsTopLevelAsset()))
+		ImportedAssets.Reserve(AssetData.Num());
+
+		for (const FAssetData& AssetDatum : AssetData)
 		{
-			continue;
+			if (AssetDatum.IsUAsset() && AssetDatum.IsTopLevelAsset())  // IsUAsset returns false for redirects
+			{
+				ImportedAssets.Emplace(AssetDatum.GetAsset());  // GetAsset does not handle redirects
+
+				UE_LOG(LogChaosClothAssetDataflowNodes,
+					Verbose,
+					TEXT("Imported USD Object %s of type %s, path: %s"),
+					*AssetDatum.AssetName.ToString(),
+					*AssetDatum.AssetClassPath.ToString(),
+					*AssetDatum.GetFullName());
+			}
 		}
+	}
+}
+
+bool FChaosClothAssetUSDImportNode::ImportFromCache(const TSharedRef<FManagedArrayCollection>& OutClothCollection, FText& OutErrorText) const
+{
+	using namespace UE::Chaos::ClothAsset;
 
-		UE_LOG(LogChaosClothAssetDataflowNodes,
-			Verbose,
-			TEXT("Imported USD Object %s of type %s, path: %s"),
-			*AssetDatum.AssetName.ToString(),
-			*AssetDatum.AssetClassPath.ToString(),
-			*AssetDatum.GetFullName());
+	// Initialize from collection cache
+	// TODO: Until we have a schema so that we can use the asset cache and remove the collection cache
+	*OutClothCollection = CollectionCache;
 
-		if (AssetDatum.GetClass() == UStaticMesh::StaticClass())
+	for (const UObject* const Asset : ImportedAssets)
+	{
+		if (const UStaticMesh* const StaticMesh = Cast<UStaticMesh>(Asset))
 		{
-			UStaticMesh* const StaticMesh = CastChecked<UStaticMesh>(AssetDatum.GetAsset());
-
 			if (StaticMesh->GetNumSourceModels() > 0)  // Only deals with LOD 0 for now
 			{
 				using namespace UE::Chaos::ClothAsset::Private;
diff --git a/Engine/Plugins/ChaosClothAssetEditor/Source/ChaosClothAssetDataflowNodes/Public/ChaosClothAsset/USDImportNode.h b/Engine/Plugins/ChaosClothAssetEditor/Source/ChaosClothAssetDataflowNodes/Public/ChaosClothAsset/USDImportNode.h
--- a/Engine/Plugins/ChaosClothAssetEditor/Source/ChaosClothAssetDataflowNodes/Public/ChaosClothAsset/USDImportNode.h
+++ b/Engine/Plugins/ChaosClothAssetEditor/Source/ChaosClothAssetDataflowNodes/Public/ChaosClothAsset/USDImportNode.h
@@ -33,9 +33,15 @@ private:
 
 	bool ImportFromFile(const FString& UsdPath, const FString& AssetPath, class FText& OutErrorText);
 	bool ImportFromCache(const TSharedRef<FManagedArrayCollection>& ClothCollection, class FText& OutErrorText) const;
+	void UpdateImportedAssets();
 
+	/** Content folder where all the USD assets are imported. */
 	UPROPERTY(VisibleAnywhere, Category = "USD Import")
-	FString PackagePath;  // Content folder where all the USD asset are imported
+	FString PackagePath;
+
+	/** List of all the dependent assets created from the USD import process. */
+	UPROPERTY(VisibleAnywhere, Category = "USD Import")
+	TArray<TObjectPtr<UObject>> ImportedAssets;
 
 	FMD5Hash FileHash;
 	FManagedArrayCollection CollectionCache;  // Content cache for data that hasn't got a USD schema yet
diff --git a/Engine/Plugins/ChaosClothAssetEditor/Source/ChaosClothAssetEditor/Private/ChaosClothAsset/ClothEditorToolkit.cpp b/Engine/Plugins/ChaosClothAssetEditor/Source/ChaosClothAssetEditor/Private/ChaosClothAsset/ClothEditorToolkit.cpp
--- a/Engine/Plugins/ChaosClothAssetEditor/Source/ChaosClothAssetEditor/Private/ChaosClothAsset/ClothEditorToolkit.cpp
+++ b/Engine/Plugins/ChaosClothAssetEditor/Source/ChaosClothAssetEditor/Private/ChaosClothAsset/ClothEditorToolkit.cpp
@@ -520,6 +520,21 @@ void FChaosClothAssetEditorToolkit::GetSaveableObjects(TArray<UObject*>& OutObje
 	{
 		check(DataflowAsset->IsAsset());
 		OutObjects.Add(DataflowAsset);
+
+		if (DataflowAsset->Dataflow)
+		{
+			TArray<UObject*> References;
+			FReferenceFinder ReferenceFinder(References);
+			DataflowAsset->Dataflow->AddReferencedObjects(ReferenceFinder);
+
+			for (UObject* const Reference : References)
+			{
+				if (Reference->IsAsset())
+				{
+					OutObjects.Add(Reference);
+				}
+			}
+		}
 	}
 }
 
diff --git a/Engine/Source/Runtime/Core/Public/UObject/FortniteMainBranchObjectVersion.h b/Engine/Source/Runtime/Core/Public/UObject/FortniteMainBranchObjectVersion.h
--- a/Engine/Source/Runtime/Core/Public/UObject/FortniteMainBranchObjectVersion.h
+++ b/Engine/Source/Runtime/Core/Public/UObject/FortniteMainBranchObjectVersion.h
@@ -447,6 +447,9 @@ struct FFortniteMainBranchObjectVersion
 		// Added support for to record sim and query data of Shape Instance data in CVD
 		SimAndQueryDataSupportInChaosVisualDebugger,
 
+		// Add the imported asset dependencies to the Cloth Asset USD Import node
+		ChaosClothAssetUSDImportNodeAddAssetDependencies,
+
 		// -----<new versions can be added above this line>-------------------------------------------------
 		VersionPlusOne,
 		LatestVersion = VersionPlusOne - 1
