CIS - Fixed false positive static analysis ClothEditorToolkit.cpp(530): warning V1078: An empty container is iterated. The loop will not be executed.

diff --git a/Engine/Plugins/ChaosClothAssetEditor/Source/ChaosClothAssetEditor/Private/ChaosClothAsset/ClothEditorToolkit.cpp b/Engine/Plugins/ChaosClothAssetEditor/Source/ChaosClothAssetEditor/Private/ChaosClothAsset/ClothEditorToolkit.cpp
--- a/Engine/Plugins/ChaosClothAssetEditor/Source/ChaosClothAssetEditor/Private/ChaosClothAsset/ClothEditorToolkit.cpp
+++ b/Engine/Plugins/ChaosClothAssetEditor/Source/ChaosClothAssetEditor/Private/ChaosClothAsset/ClothEditorToolkit.cpp
@@ -527,7 +527,7 @@ void FChaosClothAssetEditorToolkit::GetSaveableObjects(TArray<UObject*>& OutObje
 			FReferenceFinder ReferenceFinder(References);
 			DataflowAsset->Dataflow->AddReferencedObjects(ReferenceFinder);
 
-			for (UObject* const Reference : References)
+			for (UObject* const Reference : References) //-V1078
 			{
 				if (Reference->IsAsset())
 				{
