Cloth Editor Preview: save and load animation state when reimporting skeletal mesh. Previously this happened only when reloading the asset, not reimporting.

diff --git a/Engine/Plugins/ChaosClothAssetEditor/Source/ChaosClothAssetEditor/Private/ChaosClothAsset/ClothEditorPreviewScene.cpp b/Engine/Plugins/ChaosClothAssetEditor/Source/ChaosClothAssetEditor/Private/ChaosClothAsset/ClothEditorPreviewScene.cpp
--- a/Engine/Plugins/ChaosClothAssetEditor/Source/ChaosClothAssetEditor/Private/ChaosClothAsset/ClothEditorPreviewScene.cpp
+++ b/Engine/Plugins/ChaosClothAssetEditor/Source/ChaosClothAssetEditor/Private/ChaosClothAsset/ClothEditorPreviewScene.cpp
@@ -19,6 +19,7 @@
 #include "Framework/Notifications/NotificationManager.h"
 #include "Widgets/Notifications/SNotificationList.h"
 #include "UObject/PackageReload.h"
+#include "EditorReimportHandler.h"
 
 #define LOCTEXT_NAMESPACE "UChaosClothEditorPreviewScene"
 
@@ -149,6 +150,7 @@ FChaosClothPreviewScene::FChaosClothPreviewScene(FPreviewScene::ConstructionValu
 	ClothComponent->RegisterComponentWithWorld(GetWorld());
 	
 	OnPackageReloadedDelegateHandle = FCoreUObjectDelegates::OnPackageReloaded.AddRaw(this, &FChaosClothPreviewScene::HandlePackageReloaded);
+	OnPostReimportDelegateHandle = FReimportManager::Instance()->OnPostReimport().AddRaw(this, &FChaosClothPreviewScene::HandleReimportManagerPostReimport);
 }
 
 FChaosClothPreviewScene::~FChaosClothPreviewScene()
@@ -167,6 +169,7 @@ FChaosClothPreviewScene::~FChaosClothPreviewScene()
 	}
 
 	FCoreUObjectDelegates::OnPackageReloaded.Remove(OnPackageReloadedDelegateHandle);
+	FReimportManager::Instance()->OnPreReimport().Remove(OnPostReimportDelegateHandle);
 }
 
 void FChaosClothPreviewScene::AddReferencedObjects(FReferenceCollector& Collector)
@@ -185,14 +188,7 @@ void FChaosClothPreviewScene::Tick(float DeltaT)
 
 	if (SavedAnimState)
 	{
-		if (UAnimSingleNodeInstance* const AnimInstance = GetPreviewAnimInstance())
-		{
-			AnimInstance->SetPosition(SavedAnimState->Time);
-			AnimInstance->SetReverse(SavedAnimState->bIsReverse);
-			AnimInstance->SetLooping(SavedAnimState->bIsLooping);
-			AnimInstance->SetPlaying(SavedAnimState->bIsPlaying);
-		}
-		SavedAnimState.Reset();
+		RestoreSavedAnimationState();
 	}
 }
 
@@ -252,15 +248,8 @@ void FChaosClothPreviewScene::SceneDescriptionPropertyChanged(const FName& Prope
 	if (PropertyName == GET_MEMBER_NAME_CHECKED(UChaosClothPreviewSceneDescription, SkeletalMeshAsset))
 	{
 		check(SkeletalMeshComponent);
-
-		if (const UAnimSingleNodeInstance* const AnimInstance = GetPreviewAnimInstance())
-		{
-			SavedAnimState = FAnimState();
-			SavedAnimState->Time = AnimInstance->GetCurrentTime();
-			SavedAnimState->bIsReverse = AnimInstance->IsReverse();
-			SavedAnimState->bIsLooping = AnimInstance->IsLooping();
-			SavedAnimState->bIsPlaying = AnimInstance->IsPlaying();
-		}
+		
+		SaveAnimationState();
 
 		SkeletalMeshComponent->SetSkeletalMeshAsset(PreviewSceneDescription->SkeletalMeshAsset);
 
@@ -420,6 +409,33 @@ void FChaosClothPreviewScene::SetGizmoDataBinder(TSharedPtr<FTransformGizmoDataB
 }
 
 
+void FChaosClothPreviewScene::SaveAnimationState()
+{
+	if (const UAnimSingleNodeInstance* const AnimInstance = GetPreviewAnimInstance())
+	{
+		SavedAnimState = FAnimState();
+		SavedAnimState->Time = AnimInstance->GetCurrentTime();
+		SavedAnimState->bIsReverse = AnimInstance->IsReverse();
+		SavedAnimState->bIsLooping = AnimInstance->IsLooping();
+		SavedAnimState->bIsPlaying = AnimInstance->IsPlaying();
+	}
+}
+
+void FChaosClothPreviewScene::RestoreSavedAnimationState()
+{
+	if (SavedAnimState)
+	{
+		if (UAnimSingleNodeInstance* const AnimInstance = GetPreviewAnimInstance())
+		{
+			AnimInstance->SetPosition(SavedAnimState->Time);
+			AnimInstance->SetReverse(SavedAnimState->bIsReverse);
+			AnimInstance->SetLooping(SavedAnimState->bIsLooping);
+			AnimInstance->SetPlaying(SavedAnimState->bIsPlaying);
+		}
+		SavedAnimState.Reset();
+	}
+}
+
 void FChaosClothPreviewScene::HandlePackageReloaded(const EPackageReloadPhase InPackageReloadPhase, FPackageReloadedEvent* InPackageReloadedEvent)
 {
 	if (InPackageReloadPhase == EPackageReloadPhase::PrePackageFixup)
@@ -430,19 +446,20 @@ void FChaosClothPreviewScene::HandlePackageReloaded(const EPackageReloadPhase In
 			{
 				// If we are going to be reloading the SkeletalMesh, first save the animation state since the AnimInstance will be reinitialized when the component is reregistered.
 				// Note we restore from the saved state in the Tick function above because AnimInstance reinitialization happens /after/ all reload delegates are called.
-				if (const UAnimSingleNodeInstance* const AnimInstance = GetPreviewAnimInstance())
-				{
-					SavedAnimState = FAnimState();
-					SavedAnimState->Time = AnimInstance->GetCurrentTime();
-					SavedAnimState->bIsReverse = AnimInstance->IsReverse();
-					SavedAnimState->bIsLooping = AnimInstance->IsLooping();
-					SavedAnimState->bIsPlaying = AnimInstance->IsPlaying();
-				}
+				SaveAnimationState();
 			}
 		}
 	}
 }
 
+void FChaosClothPreviewScene::HandleReimportManagerPostReimport(UObject* ReimportedObject, bool bWasSuccessful)
+{
+	if (ReimportedObject == PreviewSceneDescription->SkeletalMeshAsset && bWasSuccessful)
+	{
+		// If we have reimported the SkeletalMesh, save the animation state since the AnimInstance will be reinitialized when the component is reregistered.
+		SaveAnimationState();
+	}
+}
 
 } // namespace UE::Chaos::ClothAsset
 
diff --git a/Engine/Plugins/ChaosClothAssetEditor/Source/ChaosClothAssetEditor/Private/ChaosClothAsset/ClothEditorPreviewScene.h b/Engine/Plugins/ChaosClothAssetEditor/Source/ChaosClothAssetEditor/Private/ChaosClothAsset/ClothEditorPreviewScene.h
--- a/Engine/Plugins/ChaosClothAssetEditor/Source/ChaosClothAssetEditor/Private/ChaosClothAsset/ClothEditorPreviewScene.h
+++ b/Engine/Plugins/ChaosClothAssetEditor/Source/ChaosClothAssetEditor/Private/ChaosClothAsset/ClothEditorPreviewScene.h
@@ -122,7 +122,10 @@ private:
 
 	bool IsComponentSelected(const UPrimitiveComponent* InComponent);
 
+	void SaveAnimationState();
+	void RestoreSavedAnimationState();
 	void HandlePackageReloaded(const EPackageReloadPhase InPackageReloadPhase, FPackageReloadedEvent* InPackageReloadedEvent);
+	void HandleReimportManagerPostReimport(UObject* ReimportedObject, bool bWasSuccessful);
 
 	TObjectPtr<UChaosClothPreviewSceneDescription> PreviewSceneDescription;
 
@@ -146,7 +149,7 @@ private:
 	TOptional<FAnimState> SavedAnimState;
 
 	FDelegateHandle OnPackageReloadedDelegateHandle;
-
+	FDelegateHandle OnPostReimportDelegateHandle;
 };
 } // namespace UE::Chaos::ClothAsset
 
