Fix typo: OnPreReimport -> OnPostReimport. Fixes a crash when reimporting SkelMesh after the Cloth Editor is closed.

diff --git a/Engine/Plugins/ChaosClothAssetEditor/Source/ChaosClothAssetEditor/Private/ChaosClothAsset/ClothEditorPreviewScene.cpp b/Engine/Plugins/ChaosClothAssetEditor/Source/ChaosClothAssetEditor/Private/ChaosClothAsset/ClothEditorPreviewScene.cpp
--- a/Engine/Plugins/ChaosClothAssetEditor/Source/ChaosClothAssetEditor/Private/ChaosClothAsset/ClothEditorPreviewScene.cpp
+++ b/Engine/Plugins/ChaosClothAssetEditor/Source/ChaosClothAssetEditor/Private/ChaosClothAsset/ClothEditorPreviewScene.cpp
@@ -169,7 +169,7 @@ FChaosClothPreviewScene::~FChaosClothPreviewScene()
 	}
 
 	FCoreUObjectDelegates::OnPackageReloaded.Remove(OnPackageReloadedDelegateHandle);
-	FReimportManager::Instance()->OnPreReimport().Remove(OnPostReimportDelegateHandle);
+	FReimportManager::Instance()->OnPostReimport().Remove(OnPostReimportDelegateHandle);
 }
 
 void FChaosClothPreviewScene::AddReferencedObjects(FReferenceCollector& Collector)
