Cloth Asset Editor - Prevent SaveAs from asking to rename fixed asset dependencies while still saving them.

#rnx

diff --git a/Engine/Plugins/ChaosClothAssetEditor/Source/ChaosClothAssetEditor/Private/ChaosClothAsset/ClothEditorToolkit.cpp b/Engine/Plugins/ChaosClothAssetEditor/Source/ChaosClothAssetEditor/Private/ChaosClothAsset/ClothEditorToolkit.cpp
--- a/Engine/Plugins/ChaosClothAssetEditor/Source/ChaosClothAssetEditor/Private/ChaosClothAsset/ClothEditorToolkit.cpp
+++ b/Engine/Plugins/ChaosClothAssetEditor/Source/ChaosClothAssetEditor/Private/ChaosClothAsset/ClothEditorToolkit.cpp
@@ -520,7 +520,19 @@ void FChaosClothAssetEditorToolkit::GetSaveableObjects(TArray<UObject*>& OutObje
 	{
 		check(DataflowAsset->IsAsset());
 		OutObjects.Add(DataflowAsset);
+	}
+}
+
+void FChaosClothAssetEditorToolkit::OnAssetsSaved(const TArray<UObject*>& SavedObjects)
+{
+	// Also save the Dataflow asset's dependent objects if required
+	TArray<UPackage*> PackagesToSave;
 
+	const UChaosClothAsset* const ClothAsset = GetAsset();
+	check(ClothAsset);
+
+	if (UDataflow* const DataflowAsset = Private::GetDataflowFrom(ClothAsset))
+	{
 		if (DataflowAsset->Dataflow)
 		{
 			TArray<UObject*> References;
@@ -531,11 +543,15 @@ void FChaosClothAssetEditorToolkit::GetSaveableObjects(TArray<UObject*>& OutObje
 			{
 				if (Reference->IsAsset())
 				{
-					OutObjects.Add(Reference);
+					PackagesToSave.Add(Reference->GetOutermost());
 				}
 			}
 		}
 	}
+	constexpr bool bPromptToSave = true;
+	constexpr bool bCheckDirtyOnReferenceAssetSave = true;
+
+	FEditorFileUtils::PromptForCheckoutAndSave(PackagesToSave, bCheckDirtyOnReferenceAssetSave, bPromptToSave);
 }
 
 bool FChaosClothAssetEditorToolkit::ShouldReopenEditorForSavedAsset(const UObject* Asset) const
diff --git a/Engine/Plugins/ChaosClothAssetEditor/Source/ChaosClothAssetEditor/Private/ChaosClothAsset/ClothEditorToolkit.h b/Engine/Plugins/ChaosClothAssetEditor/Source/ChaosClothAssetEditor/Private/ChaosClothAsset/ClothEditorToolkit.h
--- a/Engine/Plugins/ChaosClothAssetEditor/Source/ChaosClothAssetEditor/Private/ChaosClothAsset/ClothEditorToolkit.h
+++ b/Engine/Plugins/ChaosClothAssetEditor/Source/ChaosClothAssetEditor/Private/ChaosClothAsset/ClothEditorToolkit.h
@@ -89,6 +89,7 @@ private:
 	virtual void InitToolMenuContext(FToolMenuContext& MenuContext) override;
 	virtual void GetSaveableObjects(TArray<UObject*>& OutObjects) const override;
 	virtual bool ShouldReopenEditorForSavedAsset(const UObject* Asset) const override;
+	virtual void OnAssetsSaved(const TArray<UObject*>& SavedObjects) override;
 	virtual void OnAssetsSavedAs(const TArray<UObject*>& SavedObjects) override;
 
 	// IAssetEditorInstance
diff --git a/Engine/Source/Editor/UnrealEd/Private/Toolkits/AssetEditorToolkit.cpp b/Engine/Source/Editor/UnrealEd/Private/Toolkits/AssetEditorToolkit.cpp
--- a/Engine/Source/Editor/UnrealEd/Private/Toolkits/AssetEditorToolkit.cpp
+++ b/Engine/Source/Editor/UnrealEd/Private/Toolkits/AssetEditorToolkit.cpp
@@ -652,6 +652,9 @@ void FAssetEditorToolkit::SaveAsset_Execute()
 		return;
 	}
 
+	TArray<UObject*> SavedObjects;
+	SavedObjects.Reserve(ObjectsToSave.Num());
+
 	TArray<UPackage*> PackagesToSave;
 
 	for (UObject* Object : ObjectsToSave)
@@ -664,10 +667,13 @@ void FAssetEditorToolkit::SaveAsset_Execute()
 		else
 		{
 			PackagesToSave.Add(Object->GetOutermost());
+			SavedObjects.Add(Object);
 		}
 	}
 
 	FEditorFileUtils::PromptForCheckoutAndSave(PackagesToSave, bCheckDirtyOnAssetSave, /*bPromptToSave=*/ false);
+
+	OnAssetsSaved(SavedObjects);
 }
 
 bool FAssetEditorToolkit::CanSaveAssetAs_Internal() const
diff --git a/Engine/Source/Editor/UnrealEd/Public/Toolkits/AssetEditorToolkit.h b/Engine/Source/Editor/UnrealEd/Public/Toolkits/AssetEditorToolkit.h
--- a/Engine/Source/Editor/UnrealEd/Public/Toolkits/AssetEditorToolkit.h
+++ b/Engine/Source/Editor/UnrealEd/Public/Toolkits/AssetEditorToolkit.h
@@ -350,6 +350,9 @@ protected:
 	/** Called when "Save" is clicked for this asset */
 	UNREALED_API virtual void SaveAsset_Execute();
 
+	/** Called from SaveAsset_Execute when assets have been saved */
+	virtual void OnAssetsSaved(const TArray<UObject*>& SavedObjects) {}
+
 	/** Called to test if "Save As" should be enabled for this asset */
 	virtual bool CanSaveAssetAs() const { return true; }
 
@@ -363,7 +366,7 @@ protected:
 	virtual bool ShouldReopenEditorForSavedAsset(const UObject* Asset) const { return true; }
 
 	/** Called from SaveAssetAs_Execute when assets have been saved */
-	UNREALED_API virtual void OnAssetsSavedAs(const TArray<UObject*>& SavedObjects) {}
+	virtual void OnAssetsSavedAs(const TArray<UObject*>& SavedObjects) {}
 
 	/** Called to test if "Find in Content Browser" should be enabled for this asset */
 	virtual bool CanFindInContentBrowser() const { return true; }
