Cloth Editor: disable mouse "flight controls" in the Preview viewport. Only alt+mouse or WASD should move the camera.
Note this change was made previously for the Construction viewport, this one applies to the Preview viewport.

diff --git a/Engine/Plugins/ChaosClothAssetEditor/Source/ChaosClothAssetEditor/Private/ChaosClothAsset/ClothEditor3DViewportClient.cpp b/Engine/Plugins/ChaosClothAssetEditor/Source/ChaosClothAssetEditor/Private/ChaosClothAsset/ClothEditor3DViewportClient.cpp
--- a/Engine/Plugins/ChaosClothAssetEditor/Source/ChaosClothAssetEditor/Private/ChaosClothAsset/ClothEditor3DViewportClient.cpp
+++ b/Engine/Plugins/ChaosClothAssetEditor/Source/ChaosClothAssetEditor/Private/ChaosClothAsset/ClothEditor3DViewportClient.cpp
@@ -62,6 +62,9 @@ FChaosClothAssetEditor3DViewportClient::FChaosClothAssetEditor3DViewportClient(F
 
 	// Set correct flags according to current profile settings
 	SetAdvancedShowFlagsForScene(UAssetViewerSettings::Get()->Profiles[GetMutableDefault<UEditorPerProjectUserSettings>()->AssetViewerProfileIndex].bPostProcessingEnabled);
+
+	// Disable the non-alt-key camera controls
+	bLockFlightCamera = true;
 }
 
 void FChaosClothAssetEditor3DViewportClient::RegisterDelegates()
