Cloth Editor: when selecting a node for the first time, focus the construction viewport to capture the dynamic mesh

diff --git a/Engine/Plugins/ChaosClothAssetEditor/Source/ChaosClothAssetEditor/Private/ChaosClothAsset/ClothEditorMode.cpp b/Engine/Plugins/ChaosClothAssetEditor/Source/ChaosClothAssetEditor/Private/ChaosClothAsset/ClothEditorMode.cpp
--- a/Engine/Plugins/ChaosClothAssetEditor/Source/ChaosClothAssetEditor/Private/ChaosClothAsset/ClothEditorMode.cpp
+++ b/Engine/Plugins/ChaosClothAssetEditor/Source/ChaosClothAssetEditor/Private/ChaosClothAsset/ClothEditorMode.cpp
@@ -1214,6 +1214,9 @@ void UChaosClothAssetEditorMode::ModeTick(float DeltaTime)
 	if (bDynamicMeshComponentInitDeferred)
 	{
 		ReinitializeDynamicMeshComponents();
+
+		// The first time we get a valid mesh, refocus the camera on it
+		FirstTimeFocusRestSpaceViewport();
 	}
 
 	if (PreviewScene->GetWorld())
