Cloth weight map paint tool: make changes to the weight map node undoable

diff --git a/Engine/Plugins/ChaosClothAssetEditor/Source/ChaosClothAssetDataflowNodes/Public/ChaosClothAsset/AddWeightMapNode.h b/Engine/Plugins/ChaosClothAssetEditor/Source/ChaosClothAssetDataflowNodes/Public/ChaosClothAsset/AddWeightMapNode.h
--- a/Engine/Plugins/ChaosClothAssetEditor/Source/ChaosClothAssetDataflowNodes/Public/ChaosClothAsset/AddWeightMapNode.h
+++ b/Engine/Plugins/ChaosClothAssetEditor/Source/ChaosClothAssetDataflowNodes/Public/ChaosClothAsset/AddWeightMapNode.h
@@ -93,6 +93,7 @@ public:
 
 private:
 	friend class UClothEditorWeightMapPaintTool;
+	friend class FWeightMapNodeChange;
 
 	const TArray<float>& GetVertexWeights() const { return VertexWeights; }
 	TArray<float>& GetVertexWeights() { return VertexWeights; }
diff --git a/Engine/Plugins/ChaosClothAssetEditor/Source/ChaosClothAssetEditorTools/Private/ChaosClothAsset/ClothWeightMapPaintTool.cpp b/Engine/Plugins/ChaosClothAssetEditor/Source/ChaosClothAssetEditorTools/Private/ChaosClothAsset/ClothWeightMapPaintTool.cpp
--- a/Engine/Plugins/ChaosClothAssetEditor/Source/ChaosClothAssetEditorTools/Private/ChaosClothAsset/ClothWeightMapPaintTool.cpp
+++ b/Engine/Plugins/ChaosClothAssetEditor/Source/ChaosClothAssetEditorTools/Private/ChaosClothAsset/ClothWeightMapPaintTool.cpp
@@ -1967,6 +1967,71 @@ void UClothEditorWeightMapPaintTool::ClearHiddenAction()
 	GetToolManager()->PostInvalidation();
 }
 
+
+//
+// Object encapsulating a change to the AddWeightMap node's values. Used for Undo/Redo.
+//
+class FWeightMapNodeChange : public FToolCommandChange
+{
+
+public:
+	
+	FWeightMapNodeChange(FGuid NodeGuid, const TArray<float>& Weights, TArray<float>& RenderWeights, EChaosClothAssetWeightMapOverrideType MapOverrideType, const FString& WeightMapName) :
+		NodeGuid(NodeGuid),
+		SavedWeights(Weights),
+		SavedRenderWeights(RenderWeights),
+		SavedMapOverrideType(MapOverrideType),
+		SavedWeightMapName(WeightMapName)
+	{}
+	
+
+private:
+
+	FGuid NodeGuid;
+
+	TArray<float> SavedWeights;
+
+	// Note we could store only one set of weights and use a bool to determine whether we are updating sim or render vertices, however in the future 
+	// we may enable writing both weight maps to the node at once.
+	TArray<float> SavedRenderWeights;
+
+	EChaosClothAssetWeightMapOverrideType SavedMapOverrideType;
+	FString SavedWeightMapName;
+
+	virtual FString ToString() const override { return TEXT("WeightMapNodeUpdateChange"); }
+
+	virtual void Apply(UObject* Object) final
+	{
+		SwapApplyRevert(Object);
+	}
+
+	virtual void Revert(UObject* Object) final
+	{
+		SwapApplyRevert(Object);
+	}
+
+	void SwapApplyRevert(UObject* Object)
+	{
+		if (UDataflow* const Dataflow = Cast<UDataflow>(Object))
+		{
+			if (const TSharedPtr<FDataflowNode> BaseNode = Dataflow->GetDataflow()->FindBaseNode(NodeGuid))
+			{
+				if (FChaosClothAssetAddWeightMapNode* const Node = BaseNode->AsType<FChaosClothAssetAddWeightMapNode>())
+				{
+					Swap(Node->GetVertexWeights(), SavedWeights);
+					Swap(Node->GetRenderVertexWeights(), SavedRenderWeights);
+					Swap(Node->MapOverrideType, SavedMapOverrideType);
+					Swap(Node->Name, SavedWeightMapName);
+
+					Node->Invalidate();
+				}
+			}
+		}
+	}
+
+};
+
+
 void UClothEditorWeightMapPaintTool::UpdateSelectedNode()
 {
 	check(ActiveWeightMap);
@@ -1974,6 +2039,21 @@ void UClothEditorWeightMapPaintTool::UpdateSelectedNode()
 	GetCurrentWeightMap(CurrentWeights);
 
 	check(WeightMapNodeToUpdate);
+
+	// Save previous state for undo
+	if (UDataflow* const Dataflow = ClothEditorContextObject->GetDataflowAsset())
+	{
+		TUniquePtr<FWeightMapNodeChange> Change = MakeUnique<FWeightMapNodeChange>(
+			WeightMapNodeToUpdate->GetGuid(), 
+			WeightMapNodeToUpdate->GetVertexWeights(),
+			WeightMapNodeToUpdate->GetRenderVertexWeights(), 
+			WeightMapNodeToUpdate->MapOverrideType, 
+			WeightMapNodeToUpdate->Name);
+
+		GetToolManager()->GetContextTransactionsAPI()->AppendChange(Dataflow, MoveTemp(Change), LOCTEXT("WeightMapNodeChangeDescription", "Update Weight Map Node"));
+	}
+
+
 	WeightMapNodeToUpdate->MapOverrideType = UpdateWeightMapProperties->MapOverrideType;
 	WeightMapNodeToUpdate->Name = UpdateWeightMapProperties->Name;
 
diff --git a/Engine/Plugins/ChaosClothAssetEditor/Source/ChaosClothAssetEditorTools/Public/ChaosClothAsset/ClothEditorContextObject.h b/Engine/Plugins/ChaosClothAssetEditor/Source/ChaosClothAssetEditorTools/Public/ChaosClothAsset/ClothEditorContextObject.h
--- a/Engine/Plugins/ChaosClothAssetEditor/Source/ChaosClothAssetEditorTools/Public/ChaosClothAsset/ClothEditorContextObject.h
+++ b/Engine/Plugins/ChaosClothAssetEditor/Source/ChaosClothAssetEditorTools/Public/ChaosClothAsset/ClothEditorContextObject.h
@@ -5,6 +5,7 @@
 #include "UObject/Object.h"
 #include "Dataflow/DataflowEdNode.h"
 #include "Dataflow/DataflowGraphEditor.h"
+#include "Dataflow/DataflowObject.h"
 #include "ClothEditorContextObject.generated.h"
 
 class UDataflow;
@@ -64,6 +65,21 @@ public:
 		DataflowContext = InDataflowContext;
 	}
 
+	UDataflow* GetDataflowAsset() const
+	{
+		if (const TSharedPtr<SDataflowGraphEditor> GraphEditor = DataflowGraphEditor.Pin())
+		{
+			if (UEdGraph* const EdGraph = GraphEditor->GetCurrentGraph())
+			{
+				if (UDataflow* const Dataflow = Cast<UDataflow>(EdGraph))
+				{
+					return Dataflow;
+				}
+			}
+		}
+		return nullptr;
+	}
+
 	void SetClothCollection(UE::Chaos::ClothAsset::EClothPatternVertexType ViewMode, TWeakPtr<FManagedArrayCollection> ClothCollection, bool bInUsingInputCollection);
 
 	UE_DEPRECATED(5.5, "SetClothCollection no longer takes a separate InputClothCollection argument")
