Fix missing default values in the imported properties

diff --git a/Engine/Plugins/ChaosClothAssetEditor/Source/ChaosClothAssetDataflowNodes/Public/ChaosClothAsset/ImportedValue.h b/Engine/Plugins/ChaosClothAssetEditor/Source/ChaosClothAssetDataflowNodes/Public/ChaosClothAsset/ImportedValue.h
--- a/Engine/Plugins/ChaosClothAssetEditor/Source/ChaosClothAssetDataflowNodes/Public/ChaosClothAsset/ImportedValue.h
+++ b/Engine/Plugins/ChaosClothAssetEditor/Source/ChaosClothAssetDataflowNodes/Public/ChaosClothAsset/ImportedValue.h
@@ -15,7 +15,7 @@ struct FChaosClothAssetImportedVectorValue
 	 * Property vector Value
 	*/
 	UPROPERTY(EditAnywhere, Category = "Imported Value", Meta = (EditCondition = "!bUseImportedValue"))
-	mutable FVector3f ImportedValue;
+	mutable FVector3f ImportedValue = FVector3f::Zero();
 	
 	/**
 	 * Whether the property can use the values imported from USD
@@ -35,7 +35,7 @@ struct FChaosClothAssetImportedIntValue
 	 * Property integer Value
 	*/
 	UPROPERTY(EditAnywhere, Category = "Imported Value", Meta = (EditCondition = "!bUseImportedValue"))
-	mutable int32 ImportedValue;
+	mutable int32 ImportedValue = 0;
 	
 	/**
 	 * Whether the property can use the values imported from USD
@@ -55,7 +55,7 @@ struct FChaosClothAssetImportedFloatValue
 	 * Property float value
 	*/
 	UPROPERTY(EditAnywhere, Category = "Imported Value", Meta = (EditCondition = "!bUseImportedValue"))
-	mutable float ImportedValue;
+	mutable float ImportedValue = 0.0f;
 	
 	/**
 	 * Whether the property can use the values imported from USD
