Chaos Cloth Asset Editor - Remove Datasmith dependency.

#rnx
#tests Reloaded Cloth Assets containing a DatasmithImport node and made sure the geometry still remained and the file path property was disabled.

diff --git a/Engine/Plugins/ChaosClothAssetEditor/ChaosClothAssetEditor.uplugin b/Engine/Plugins/ChaosClothAssetEditor/ChaosClothAssetEditor.uplugin
--- a/Engine/Plugins/ChaosClothAssetEditor/ChaosClothAssetEditor.uplugin
+++ b/Engine/Plugins/ChaosClothAssetEditor/ChaosClothAssetEditor.uplugin
@@ -66,17 +66,6 @@
 			"Name": "Dataflow",
 			"Enabled": true
 		},
-		{
-			"Name": "DatasmithImporter",
-			"Enabled": true
-		},
-		{
-			"Name": "DatasmithCloTranslator",
-			"Enabled": true,
-			"SupportedTargetPlatforms": [
-				"Win64"
-			]
-		},
 		{
 			"Name": "USDImporter",
 			"Enabled": true,
@@ -85,4 +74,4 @@
 			]
 		}
 	]
-}
\ No newline at end of file
+}
diff --git a/Engine/Plugins/ChaosClothAssetEditor/Source/ChaosClothAssetDataflowNodes/ChaosClothAssetDataflowNodes.Build.cs b/Engine/Plugins/ChaosClothAssetEditor/Source/ChaosClothAssetDataflowNodes/ChaosClothAssetDataflowNodes.Build.cs
--- a/Engine/Plugins/ChaosClothAssetEditor/Source/ChaosClothAssetDataflowNodes/ChaosClothAssetDataflowNodes.Build.cs
+++ b/Engine/Plugins/ChaosClothAssetEditor/Source/ChaosClothAssetDataflowNodes/ChaosClothAssetDataflowNodes.Build.cs
@@ -13,7 +13,12 @@ public class ChaosClothAssetDataflowNodes : ModuleRules
 				// ... add other public dependencies that you statically link with here ...
 			}
 		);
-		
+
+		PrivateIncludePathModuleNames.AddRange(
+			new string[] {
+				"DatasmithCore",  // Deprecated in 5.5 (remove in 5.7), for DatasmithCloth.h
+		});
+
 		PrivateDependencyModuleNames.AddRange(
 			new string[]
 			{
@@ -27,13 +32,10 @@ public class ChaosClothAssetDataflowNodes : ModuleRules
 				"DataflowCore",
 				"DataflowEditor",
 				"DataflowEngine",
-				"DatasmithCore",
-				"DatasmithImporter",
 				"DesktopWidgets",  // For SFilePathPicker
 				"DetailCustomizations",
 				"DynamicMesh",
 				"Engine",
-				"ExternalSource",
 				"GeometryCore",
 				"InputCore",
 				"MeshConversion",
diff --git a/Engine/Plugins/ChaosClothAssetEditor/Source/ChaosClothAssetDataflowNodes/Private/ChaosClothAsset/ChaosClothAssetDataflowNodesModule.cpp b/Engine/Plugins/ChaosClothAssetEditor/Source/ChaosClothAssetDataflowNodes/Private/ChaosClothAsset/ChaosClothAssetDataflowNodesModule.cpp
--- a/Engine/Plugins/ChaosClothAssetEditor/Source/ChaosClothAssetDataflowNodes/Private/ChaosClothAsset/ChaosClothAssetDataflowNodesModule.cpp
+++ b/Engine/Plugins/ChaosClothAssetEditor/Source/ChaosClothAssetDataflowNodes/Private/ChaosClothAsset/ChaosClothAssetDataflowNodesModule.cpp
@@ -160,9 +160,6 @@ namespace UE::Chaos::ClothAsset
 				PropertyModule->RegisterCustomPropertyTypeLayout(FChaosClothAssetImportedFloatValue::StaticStruct()->GetFName(), FOnGetPropertyTypeCustomizationInstance::CreateStatic(&FImportedValueCustomization::MakeInstance));
 				PropertyModule->RegisterCustomPropertyTypeLayout(FChaosClothAssetImportedIntValue::StaticStruct()->GetFName(), FOnGetPropertyTypeCustomizationInstance::CreateStatic(&FImportedValueCustomization::MakeInstance));
 			}
-
-			// Register modular features
-			FChaosClothAssetDatasmithImportNode::RegisterModularFeature();
 		}
 
 		virtual void ShutdownModule() override
@@ -189,9 +186,6 @@ namespace UE::Chaos::ClothAsset
 					PropertyModule->UnregisterCustomPropertyTypeLayout(FChaosClothAssetImportedIntValue::StaticStruct()->GetFName());
 				}
 			}
-
-			// Unregister modular features
-			FChaosClothAssetDatasmithImportNode::UnregisterModularFeature();
 		}
 	};
 }  // End namespace UE::Chaos::ClothAsset
diff --git a/Engine/Plugins/ChaosClothAssetEditor/Source/ChaosClothAssetDataflowNodes/Private/ChaosClothAsset/DatasmithImportNode.cpp b/Engine/Plugins/ChaosClothAssetEditor/Source/ChaosClothAssetDataflowNodes/Private/ChaosClothAsset/DatasmithImportNode.cpp
--- a/Engine/Plugins/ChaosClothAssetEditor/Source/ChaosClothAssetDataflowNodes/Private/ChaosClothAsset/DatasmithImportNode.cpp
+++ b/Engine/Plugins/ChaosClothAssetEditor/Source/ChaosClothAssetDataflowNodes/Private/ChaosClothAsset/DatasmithImportNode.cpp
@@ -5,142 +5,12 @@
 #include "ChaosClothAsset/ClothComponent.h"
 #include "ChaosClothAsset/CollectionClothFacade.h"
 #include "Dataflow/DataflowInputOutput.h"
-#include "Features/IModularFeatures.h"
-#include "DatasmithImportContext.h"
-#include "DatasmithImportFactory.h"
-#include "ExternalSource.h"
-#include "ExternalSourceModule.h"
-#include "SourceUri.h"
+#include "Misc/SecureHash.h"
 
 #include UE_INLINE_GENERATED_CPP_BY_NAME(DatasmithImportNode)
 
 #define LOCTEXT_NAMESPACE "ChaosClothAssetDatasmithImportNode"
 
-UChaosClothAssetDatasmithClothAssetFactory::UChaosClothAssetDatasmithClothAssetFactory() = default;
-UChaosClothAssetDatasmithClothAssetFactory::~UChaosClothAssetDatasmithClothAssetFactory() = default;
-
-UObject* UChaosClothAssetDatasmithClothAssetFactory::CreateClothAsset(UObject* Outer, const FName& Name, EObjectFlags Flags) const
-{
-	return CastChecked<UObject>(NewObject<UChaosClothAsset>(Outer, Name, Flags));
-}
-
-UObject* UChaosClothAssetDatasmithClothAssetFactory::DuplicateClothAsset(UObject* ClothAsset, UObject* Outer, const FName& Name) const
-{
-	return CastChecked<UObject>(DuplicateObject<UChaosClothAsset>(CastChecked<UChaosClothAsset>(ClothAsset), Outer, Name));
-}
-
-void UChaosClothAssetDatasmithClothAssetFactory::InitializeClothAsset(UObject* ClothAsset, const FDatasmithCloth& DatasmithCloth) const
-{
-	using namespace UE::Chaos::ClothAsset;
-
-	UChaosClothAsset* const ChaosClothAsset = CastChecked<UChaosClothAsset>(ClothAsset);
-
-	TArray<TSharedRef<FManagedArrayCollection>>& Collections = ChaosClothAsset->GetClothCollections();
-	Collections.Reset(1);
-	FCollectionClothFacade Cloth(Collections.Emplace_GetRef(MakeShared<FManagedArrayCollection>()));
-	Cloth.DefineSchema();
-
-	for (const FDatasmithClothPattern& Pattern : DatasmithCloth.Patterns)
-	{
-		if (Pattern.IsValid())
-		{
-			FCollectionClothSimPatternFacade ClothPattern = Cloth.AddGetSimPattern();
-			ClothPattern.Initialize(Pattern.SimPosition, Pattern.SimRestPosition, Pattern.SimTriangleIndices);
-		}
-	}
-
-	for (const FDatasmithClothSewingInfo& SeamInfo : DatasmithCloth.Sewing)
-	{
-		const int32 SeamPattern0 = (int32)SeamInfo.Seam0PanelIndex;
-		const int32 SeamPattern1 = (int32)SeamInfo.Seam1PanelIndex;
-
-		if (SeamPattern0 >= 0 && SeamPattern0 < Cloth.GetNumSimPatterns() &&
-			SeamPattern1 >= 0 && SeamPattern1 < Cloth.GetNumSimPatterns())
-		{
-			const FCollectionClothSimPatternConstFacade ClothPattern0 = Cloth.GetSimPattern(SeamPattern0);
-			const FCollectionClothSimPatternConstFacade ClothPattern1 = Cloth.GetSimPattern(SeamPattern1);
-
-			const int32 ClothPattern0VerticesOffset = ClothPattern0.GetSimVertices2DOffset();
-			const int32 ClothPattern1VerticesOffset = ClothPattern1.GetSimVertices2DOffset();
-
-			TArray<FIntVector2> Stitches;
-			const uint32 StitchesCount = FMath::Min(SeamInfo.Seam0MeshIndices.Num(), SeamInfo.Seam1MeshIndices.Num());
-			Stitches.Reserve(StitchesCount);
-			for (uint32 StitchIndex = 0; StitchIndex < StitchesCount; ++StitchIndex)
-			{
-				Stitches.Emplace(
-					(int32)SeamInfo.Seam0MeshIndices[StitchIndex] + ClothPattern0VerticesOffset,
-					(int32)SeamInfo.Seam1MeshIndices[StitchIndex] + ClothPattern1VerticesOffset);
-			}
-
-			FCollectionClothSeamFacade SeamFacade = Cloth.AddGetSeam();
-			SeamFacade.Initialize(Stitches);
-		}
-	}
-
-	// Set the render mesh to duplicate the sim mesh
-	ChaosClothAsset->CopySimMeshToRenderMesh();
-
-	// Set a default skeleton and rebuild the asset
-	ChaosClothAsset->SetReferenceSkeleton(nullptr);  // This creates a default reference skeleton, redoes the bindings, and rebuilds the asset
-}
-
-UChaosClothAssetDatasmithClothComponentFactory::UChaosClothAssetDatasmithClothComponentFactory() = default;
-UChaosClothAssetDatasmithClothComponentFactory::~UChaosClothAssetDatasmithClothComponentFactory() = default;
-
-USceneComponent* UChaosClothAssetDatasmithClothComponentFactory::CreateClothComponent(UObject* Outer) const
-{
-	return CastChecked<USceneComponent>(NewObject<UChaosClothComponent>(Outer));
-}
-
-void UChaosClothAssetDatasmithClothComponentFactory::InitializeClothComponent(USceneComponent* ClothComponent, UObject* ClothAsset, USceneComponent* RootComponent) const
-{
-	UChaosClothComponent* const ChaosClothComponent = CastChecked<UChaosClothComponent>(ClothComponent);
-	ChaosClothComponent->SetClothAsset(Cast<UChaosClothAsset>(ClothAsset));
-	ChaosClothComponent->SetupAttachment(RootComponent);
-}
-
-namespace UE::Chaos::ClothAsset::Private
-{
-	/** A modular interface to provide and initialize cloth asset objects. */
-	class FDatasmithClothFactoryClassesProvider final: public IDatasmithClothFactoryClassesProvider
-	{
-	public:
-		FDatasmithClothFactoryClassesProvider() = default;
-		virtual ~FDatasmithClothFactoryClassesProvider() override = default;
-
-		virtual FName GetName() const override
-		{
-			static const FName Name = TEXT("ChaosClothAsset");
-			return Name;
-		}
-
-		virtual TSubclassOf<UDatasmithClothAssetFactory> GetClothAssetFactoryClass() const override
-		{
-			return TSubclassOf<UDatasmithClothAssetFactory>(UChaosClothAssetDatasmithClothAssetFactory::StaticClass());
-		}
-
-		virtual TSubclassOf<UDatasmithClothComponentFactory> GetClothComponentFactoryClass() const override
-		{
-			return TSubclassOf<UDatasmithClothComponentFactory>(UChaosClothAssetDatasmithClothComponentFactory::StaticClass());
-		}
-	};
-
-	static FDatasmithClothFactoryClassesProvider DatasmithClothFactoryClassesProvider;
-}
-
-void FChaosClothAssetDatasmithImportNode::RegisterModularFeature()
-{
-	using namespace UE::Chaos::ClothAsset::Private;
-	IModularFeatures::Get().RegisterModularFeature(IDatasmithClothFactoryClassesProvider::FeatureName, &DatasmithClothFactoryClassesProvider);
-}
-
-void FChaosClothAssetDatasmithImportNode::UnregisterModularFeature()
-{
-	using namespace UE::Chaos::ClothAsset::Private;
-	IModularFeatures::Get().UnregisterModularFeature(IDatasmithClothFactoryClassesProvider::FeatureName, &DatasmithClothFactoryClassesProvider);
-}
-
 FChaosClothAssetDatasmithImportNode::FChaosClothAssetDatasmithImportNode(const Dataflow::FNodeParameters& InParam, FGuid InGuid)
 	: FDataflowNode(InParam, InGuid)
 {
@@ -164,108 +34,14 @@ void FChaosClothAssetDatasmithImportNode::Serialize(FArchive& Archive)
 		}
 		ImportCache = MoveTemp(*ClothCollection);
 	}
+	FMD5Hash ImportHash;
 	Archive << ImportHash;
 }
 
-bool FChaosClothAssetDatasmithImportNode::EvaluateImpl(Dataflow::FContext& Context, FManagedArrayCollection& OutCollection) const
-{
-	using namespace UE::DatasmithImporter;
-
-	const FSourceUri SourceUri = FSourceUri::FromFilePath(ImportFile.FilePath);
-	const TSharedPtr<FExternalSource> ExternalSource = IExternalSourceModule::GetOrCreateExternalSource(SourceUri);
-	if (!ExternalSource)
-	{
-		return false;
-	}
-
-	constexpr bool bLoadConfig = false; // TODO: Not sure what this does
-	const FName LoggerName(TEXT("ImportDatasmithClothNode"));
-	const FText LoggerLabel(NSLOCTEXT("ImportDatasmithClothNode", "LoggerLabel", "ImportDatasmithClothNode"));
-	FDatasmithImportContext DatasmithImportContext(ExternalSource, bLoadConfig, LoggerName, LoggerLabel);
-
-	const FName PackageName = MakeUniqueObjectName(nullptr, UPackage::StaticClass(), FName(FPaths::Combine(GetTransientPackage()->GetPathName(), ExternalSource->GetSourceName())));
-	const TStrongObjectPtr<UPackage> DestinationPackage(CreatePackage(*PackageName.ToString()));
-	if (!ensure(DestinationPackage))
-	{
-		// Failed to create the package to hold this asset for some reason
-		return false;
-	}
-
-	// Don't create the Actors in the level, just read the Assets
-	DatasmithImportContext.Options->BaseOptions.SceneHandling = EDatasmithImportScene::AssetsOnly;
-
-	constexpr EObjectFlags NewObjectFlags = RF_Public | RF_Transactional | RF_Transient | RF_Standalone;
-	const TSharedPtr<FJsonObject> ImportSettingsJson;
-	constexpr bool bIsSilent = true;
-
-	if (!DatasmithImportContext.Init(DestinationPackage->GetPathName(), NewObjectFlags, GWarn, ImportSettingsJson, bIsSilent))
-	{
-		return false;
-	}
-
-	if (const TSharedPtr<IDatasmithScene> LoadedScene = ExternalSource->TryLoad())
-	{
-		DatasmithImportContext.InitScene(LoadedScene.ToSharedRef());
-	}
-	else
-	{
-		return false;
-	}
-
-	bool bUserCancelled = false;
-	bool bImportSucceed = DatasmithImportFactoryImpl::ImportDatasmithScene(DatasmithImportContext, bUserCancelled);
-	bImportSucceed &= !bUserCancelled;
-
-	if (bImportSucceed && DatasmithImportContext.ImportedClothes.Num() > 0)
-	{
-		UObject* const ClothObject = DatasmithImportContext.ImportedClothes.CreateIterator()->Value;
-		UChaosClothAsset* const DatasmithClothAsset = Cast<UChaosClothAsset>(ClothObject);
-		if (ensure(DatasmithClothAsset))
-		{
-			if (DatasmithClothAsset->GetClothCollections().Num() > 0)
-			{
-				DatasmithClothAsset->GetClothCollections()[0]->CopyTo(&OutCollection);
-				return true;
-			}
-		}
-		DatasmithClothAsset->ClearFlags(RF_Standalone);
-	}
-
-	return false;
-}
-
 void FChaosClothAssetDatasmithImportNode::Evaluate(Dataflow::FContext& Context, const FDataflowOutput* Out) const
 {
 	if (Out->IsA<FManagedArrayCollection>(&Collection))
 	{
-		const FMD5Hash FileHash = ImportFile.FilePath.IsEmpty() ?
-			FMD5Hash() :  // Reset to an empty cloth collection
-			FPaths::FileExists(ImportFile.FilePath) ?
-				FMD5Hash::HashFile(*ImportFile.FilePath) :  // Update cache to the file content if necessary
-				ImportHash;  // Keep the current cached file
-
-		if (FileHash != ImportHash)
-		{
-			ImportHash = FileHash;
-
-			// The file has changed, reimport
-			const bool bSuccess = EvaluateImpl(Context, ImportCache);
-
-			if (!bSuccess)
-			{
-				using namespace UE::Chaos::ClothAsset;
-
-				// The import has failed, reset the cache to an empty cloth collection
-				ImportCache.Reset();
-
-				const TSharedRef<FManagedArrayCollection> ClothCollection = MakeShared<FManagedArrayCollection>(MoveTemp(ImportCache));
-				FCollectionClothFacade CollectionClothFacade(ClothCollection);
-				CollectionClothFacade.DefineSchema();
-
-				ImportCache = MoveTemp(*ClothCollection);
-			}
-		}
-
 		SetValue(Context, ImportCache, &Collection);
 	}
 }
diff --git a/Engine/Plugins/ChaosClothAssetEditor/Source/ChaosClothAssetDataflowNodes/Public/ChaosClothAsset/DatasmithImportNode.h b/Engine/Plugins/ChaosClothAssetEditor/Source/ChaosClothAssetDataflowNodes/Public/ChaosClothAsset/DatasmithImportNode.h
--- a/Engine/Plugins/ChaosClothAssetEditor/Source/ChaosClothAssetDataflowNodes/Public/ChaosClothAsset/DatasmithImportNode.h
+++ b/Engine/Plugins/ChaosClothAssetEditor/Source/ChaosClothAssetDataflowNodes/Public/ChaosClothAsset/DatasmithImportNode.h
@@ -4,42 +4,55 @@
 
 #include "Dataflow/DataflowNode.h"
 #include "GeometryCollection/ManagedArrayCollection.h"
-#include "Misc/SecureHash.h"
+
+#if UE_ENABLE_INCLUDE_ORDER_DEPRECATED_IN_5_5
 #include "DatasmithCloth.h"
+#include "Misc/SecureHash.h"
+#endif
+
 #include "DatasmithImportNode.generated.h"
 
-/** Cloth asset factory providing and initializing cloth assets on behalf of the Datasmith importer. */
+/** Deprecated. */
 UCLASS()
-class UChaosClothAssetDatasmithClothAssetFactory final : public UDatasmithClothAssetFactory
+class UChaosClothAssetDatasmithClothAssetFactory final : public UObject
 {
 	GENERATED_BODY()
 
 public:
-	UChaosClothAssetDatasmithClothAssetFactory();
-	virtual ~UChaosClothAssetDatasmithClothAssetFactory() override;
-
-	virtual UObject* CreateClothAsset(UObject* Outer, const FName& Name, EObjectFlags ObjectFlags) const override;
-	virtual UObject* DuplicateClothAsset(UObject* ClothAsset, UObject* Outer, const FName& Name) const override;
-	virtual void InitializeClothAsset(UObject* ClothAsset, const FDatasmithCloth& DatasmithCloth) const override;
+	UE_DEPRECATED(5.5, "The experimental Cloth importer is no longer supported.")
+	UChaosClothAssetDatasmithClothAssetFactory() = default;
+	UE_DEPRECATED(5.5, "The experimental Cloth importer is no longer supported.")
+	virtual ~UChaosClothAssetDatasmithClothAssetFactory() = default;
+
+	UE_DEPRECATED(5.5, "The experimental Cloth importer is no longer supported.")
+	virtual UObject* CreateClothAsset(UObject* /*Outer*/, const FName& /*Name*/, EObjectFlags /*ObjectFlags*/) const { return nullptr; }
+	UE_DEPRECATED(5.5, "The experimental Cloth importer is no longer supported.")
+	virtual UObject* DuplicateClothAsset(UObject* /*ClothAsset*/, UObject* Outer, const FName& /*Name*/) const { return nullptr; }
+	UE_DEPRECATED(5.5, "The experimental Cloth importer is no longer supported.")
+	virtual void InitializeClothAsset(UObject* /*ClothAsset*/, const class FDatasmithCloth& /*DatasmithCloth*/) const {}
 };
 
 
-/** Cloth component factory providing and initializing cloth components on behalf of the Datasmith importer. */
+/** Deprecated. */
 UCLASS()
-class UChaosClothAssetDatasmithClothComponentFactory final : public UDatasmithClothComponentFactory
+class UChaosClothAssetDatasmithClothComponentFactory final : public UObject
 {
 	GENERATED_BODY()
 
 public:
-	UChaosClothAssetDatasmithClothComponentFactory();
-	virtual ~UChaosClothAssetDatasmithClothComponentFactory() override;
-
-	virtual USceneComponent* CreateClothComponent(UObject* Outer) const override;
-	virtual void InitializeClothComponent(class USceneComponent* ClothComponent, UObject* ClothAsset, class USceneComponent* RootComponent) const override;
+	UE_DEPRECATED(5.5, "The experimental Cloth importer is no longer supported.")
+	UChaosClothAssetDatasmithClothComponentFactory() = default;
+	UE_DEPRECATED(5.5, "The experimental Cloth importer is no longer supported.")
+	virtual ~UChaosClothAssetDatasmithClothComponentFactory() = default;
+
+	UE_DEPRECATED(5.5, "The experimental Cloth importer is no longer supported.")
+	virtual USceneComponent* CreateClothComponent(UObject* /*Outer*/) const { return nullptr; }
+	UE_DEPRECATED(5.5, "The experimental Cloth importer is no longer supported.")
+	virtual void InitializeClothComponent(class USceneComponent* /*ClothComponent*/, UObject* /*ClothAsset*/, class USceneComponent* /*RootComponent*/) const {}
 };
 
-/** Import a file from a third party garment construction package compatible with the Datasmith scene format. */
-USTRUCT(meta = (DataflowCloth, Deprecated = "5.4"))
+/** Deprecated. The experimental Datasmith cloth importer is no longer supported. Use the USDImport node instead. */
+USTRUCT(Meta = (DataflowCloth, Deprecated = "5.4"))
 struct FChaosClothAssetDatasmithImportNode : public FDataflowNode
 {
 	GENERATED_USTRUCT_BODY()
@@ -50,13 +63,13 @@ public:
 	FManagedArrayCollection Collection;
 
 	/** Path of the file to import using any available Datasmith cloth translator. */
-	UPROPERTY(EditAnywhere, Category = "Datasmith Import")
+	UPROPERTY(VisibleAnywhere, Category = "Datasmith Import", Meta = (EditCondition = "false"))
 	FFilePath ImportFile;
 
-	/** Register the Datasmith cloth factory classes provider. */
-	static void RegisterModularFeature();
-	/** Unregister the Datasmith cloth factory classes provider. */
-	static void UnregisterModularFeature();
+	UE_DEPRECATED(5.5, "The experimental Cloth importer is no longer supported.")
+	static void RegisterModularFeature() {}
+	UE_DEPRECATED(5.5, "The experimental Cloth importer is no longer supported.")
+	static void UnregisterModularFeature() {}
 
 	FChaosClothAssetDatasmithImportNode(const Dataflow::FNodeParameters& InParam, FGuid InGuid = FGuid::NewGuid());
 
@@ -66,8 +79,5 @@ private:
 	virtual void Serialize(FArchive& Archive) override;
 	//~ End FDataflowNode interface
 
-	bool EvaluateImpl(Dataflow::FContext& Context, FManagedArrayCollection& OutCollection) const;
-
-	mutable FManagedArrayCollection ImportCache;
-	mutable FMD5Hash ImportHash;
+	FManagedArrayCollection ImportCache;
 };
