CIS - Fixed Issue 460222: Declaring virtual methods in a class marked as 'final' is pointless.

#rnx

diff --git a/Engine/Plugins/ChaosClothAssetEditor/Source/ChaosClothAssetDataflowNodes/Public/ChaosClothAsset/DatasmithImportNode.h b/Engine/Plugins/ChaosClothAssetEditor/Source/ChaosClothAssetDataflowNodes/Public/ChaosClothAsset/DatasmithImportNode.h
--- a/Engine/Plugins/ChaosClothAssetEditor/Source/ChaosClothAssetDataflowNodes/Public/ChaosClothAsset/DatasmithImportNode.h
+++ b/Engine/Plugins/ChaosClothAssetEditor/Source/ChaosClothAssetDataflowNodes/Public/ChaosClothAsset/DatasmithImportNode.h
@@ -25,11 +25,11 @@ public:
 	virtual ~UChaosClothAssetDatasmithClothAssetFactory() = default;
 
 	UE_DEPRECATED(5.5, "The experimental Cloth importer is no longer supported.")
-	virtual UObject* CreateClothAsset(UObject* /*Outer*/, const FName& /*Name*/, EObjectFlags /*ObjectFlags*/) const { return nullptr; }
+	UObject* CreateClothAsset(UObject* /*Outer*/, const FName& /*Name*/, EObjectFlags /*ObjectFlags*/) const { return nullptr; }
 	UE_DEPRECATED(5.5, "The experimental Cloth importer is no longer supported.")
-	virtual UObject* DuplicateClothAsset(UObject* /*ClothAsset*/, UObject* Outer, const FName& /*Name*/) const { return nullptr; }
+	UObject* DuplicateClothAsset(UObject* /*ClothAsset*/, UObject* Outer, const FName& /*Name*/) const { return nullptr; }
 	UE_DEPRECATED(5.5, "The experimental Cloth importer is no longer supported.")
-	virtual void InitializeClothAsset(UObject* /*ClothAsset*/, const class FDatasmithCloth& /*DatasmithCloth*/) const {}
+	void InitializeClothAsset(UObject* /*ClothAsset*/, const class FDatasmithCloth& /*DatasmithCloth*/) const {}
 };
 
 
@@ -46,9 +46,9 @@ public:
 	virtual ~UChaosClothAssetDatasmithClothComponentFactory() = default;
 
 	UE_DEPRECATED(5.5, "The experimental Cloth importer is no longer supported.")
-	virtual USceneComponent* CreateClothComponent(UObject* /*Outer*/) const { return nullptr; }
+	USceneComponent* CreateClothComponent(UObject* /*Outer*/) const { return nullptr; }
 	UE_DEPRECATED(5.5, "The experimental Cloth importer is no longer supported.")
-	virtual void InitializeClothComponent(class USceneComponent* /*ClothComponent*/, UObject* /*ClothAsset*/, class USceneComponent* /*RootComponent*/) const {}
+	void InitializeClothComponent(class USceneComponent* /*ClothComponent*/, UObject* /*ClothAsset*/, class USceneComponent* /*RootComponent*/) const {}
 };
 
 /** Deprecated. The experimental Datasmith cloth importer is no longer supported. Use the USDImport node instead. */
