Chaos Cloth Asset: fix regression where the Name output on AddWeightMap and Selection nodes outputs InputName instead of Name when InputName is non-empty.

diff --git a/Engine/Plugins/ChaosClothAssetEditor/Source/ChaosClothAssetDataflowNodes/Private/ChaosClothAsset/AddWeightMapNode.cpp b/Engine/Plugins/ChaosClothAssetEditor/Source/ChaosClothAssetDataflowNodes/Private/ChaosClothAsset/AddWeightMapNode.cpp
--- a/Engine/Plugins/ChaosClothAssetEditor/Source/ChaosClothAssetDataflowNodes/Private/ChaosClothAsset/AddWeightMapNode.cpp
+++ b/Engine/Plugins/ChaosClothAssetEditor/Source/ChaosClothAssetDataflowNodes/Private/ChaosClothAsset/AddWeightMapNode.cpp
@@ -318,7 +318,7 @@ void FChaosClothAssetAddWeightMapNode::Evaluate(Dataflow::FContext& Context, con
 	{
 		FString InputNameString = GetValue<FString>(Context, &InputName.StringValue);
 		UE::Chaos::ClothAsset::FWeightMapTools::MakeWeightMapName(InputNameString);
-		SetValue(Context, InputNameString.IsEmpty() ? Name : InputNameString, &Name);
+		SetValue(Context, Name.IsEmpty() ? InputNameString : Name, &Name);
 	}
 }
 
diff --git a/Engine/Plugins/ChaosClothAssetEditor/Source/ChaosClothAssetDataflowNodes/Private/ChaosClothAsset/SelectionNode.cpp b/Engine/Plugins/ChaosClothAssetEditor/Source/ChaosClothAssetDataflowNodes/Private/ChaosClothAsset/SelectionNode.cpp
--- a/Engine/Plugins/ChaosClothAssetEditor/Source/ChaosClothAssetDataflowNodes/Private/ChaosClothAsset/SelectionNode.cpp
+++ b/Engine/Plugins/ChaosClothAssetEditor/Source/ChaosClothAssetDataflowNodes/Private/ChaosClothAsset/SelectionNode.cpp
@@ -416,7 +416,7 @@ void FChaosClothAssetSelectionNode::Evaluate(Dataflow::FContext& Context, const
 	{
 		FString InputNameString = GetValue<FString>(Context, &InputName.StringValue);
 		UE::Chaos::ClothAsset::FWeightMapTools::MakeWeightMapName(InputNameString);
-		SetValue(Context, InputNameString.IsEmpty() ? Name : InputNameString, &Name);
+		SetValue(Context, Name.IsEmpty() ? InputNameString : Name, &Name);
 	}
 }
 
