Chaos Cloth Asset - Fixed Warning: Failed to find object PhysicsAsset when setting an empty physics asset to the SetPhysicsAsset node.
#tests Tested with a SetPhysicsAsset node set to None, and the warning wasn't showing anymore.

diff --git a/Engine/Plugins/ChaosClothAssetEditor/Source/ChaosClothAssetDataflowNodes/Private/ChaosClothAsset/SetPhysicsAssetNode.cpp b/Engine/Plugins/ChaosClothAssetEditor/Source/ChaosClothAssetDataflowNodes/Private/ChaosClothAsset/SetPhysicsAssetNode.cpp
--- a/Engine/Plugins/ChaosClothAssetEditor/Source/ChaosClothAssetDataflowNodes/Private/ChaosClothAsset/SetPhysicsAssetNode.cpp
+++ b/Engine/Plugins/ChaosClothAssetEditor/Source/ChaosClothAssetDataflowNodes/Private/ChaosClothAsset/SetPhysicsAssetNode.cpp
@@ -29,7 +29,7 @@ void FChaosClothAssetSetPhysicsAssetNode::Evaluate(Dataflow::FContext& Context,
 		if (FCollectionClothFacade(ClothCollection).IsValid())  // Can only act on the collection if it is a valid cloth collection
 		{
 			FCollectionClothFacade CollectionClothFacade(ClothCollection);
-			CollectionClothFacade.SetPhysicsAssetPathName(PhysicsAsset.GetPathName());
+			CollectionClothFacade.SetPhysicsAssetPathName(PhysicsAsset ? PhysicsAsset.GetPathName() : FString());
 		}
 
 		SetValue(Context, MoveTemp(*ClothCollection), &Collection);
