Cloth Editor: make Selection node changes via the tool undoable, similar to previous change to the Weight Map Paint tool.

Also make the previously-added FWeightMapNodeChange a nested class inside FChaosClothAssetAddWeightMapNode

diff --git a/Engine/Plugins/ChaosClothAssetEditor/Source/ChaosClothAssetDataflowNodes/Private/ChaosClothAsset/AddWeightMapNode.cpp b/Engine/Plugins/ChaosClothAssetEditor/Source/ChaosClothAssetDataflowNodes/Private/ChaosClothAsset/AddWeightMapNode.cpp
--- a/Engine/Plugins/ChaosClothAssetEditor/Source/ChaosClothAssetDataflowNodes/Private/ChaosClothAsset/AddWeightMapNode.cpp
+++ b/Engine/Plugins/ChaosClothAssetEditor/Source/ChaosClothAssetDataflowNodes/Private/ChaosClothAsset/AddWeightMapNode.cpp
@@ -9,6 +9,7 @@
 #include "ChaosClothAsset/WeightedValue.h"
 #include "Dataflow/DataflowInputOutput.h"
 #include "Dataflow/DataflowObject.h"
+#include "InteractiveToolChange.h"
 
 #include UE_INLINE_GENERATED_CPP_BY_NAME(AddWeightMapNode)
 
@@ -349,4 +350,74 @@ void FChaosClothAssetAddWeightMapNode::CalculateFinalRenderVertexWeightValues(co
 {
 	UE::Chaos::ClothAsset::Private::CalculateFinalVertexWeightValues(InputMap, FinalOutputMap, MapOverrideType, GetRenderVertexWeights());
 }
+
+
+// Object encapsulating a change to the AddWeightMap node's values. Used for Undo/Redo.
+class FChaosClothAssetAddWeightMapNode::FWeightMapNodeChange final : public FToolCommandChange
+{
+
+public: 
+
+	FWeightMapNodeChange(const FChaosClothAssetAddWeightMapNode& Node) :
+		NodeGuid(Node.GetGuid()),
+		SavedWeights(Node.GetVertexWeights()),
+		SavedRenderWeights(Node.GetRenderVertexWeights()),
+		SavedMapOverrideType(Node.MapOverrideType),
+		SavedWeightMapName(Node.Name)
+	{}
+
+private:
+
+	FGuid NodeGuid;
+	TArray<float> SavedWeights;
+
+	// Note we could store only one set of weights and use a bool to determine whether we are updating sim or render vertices, however in the future 
+	// we may enable writing both weight maps to the node at once.
+	TArray<float> SavedRenderWeights;
+
+	EChaosClothAssetWeightMapOverrideType SavedMapOverrideType;
+	FString SavedWeightMapName;
+
+	virtual FString ToString() const final
+	{
+		return TEXT("ChaosClothAssetAddWeightMapNodeChange");
+	}
+
+	virtual void Apply(UObject* Object) final
+	{
+		SwapApplyRevert(Object);
+	}
+
+	virtual void Revert(UObject* Object) final
+	{
+		SwapApplyRevert(Object);
+	}
+
+	void SwapApplyRevert(UObject* Object)
+	{
+		if (UDataflow* const Dataflow = Cast<UDataflow>(Object))
+		{
+			if (const TSharedPtr<FDataflowNode> BaseNode = Dataflow->GetDataflow()->FindBaseNode(NodeGuid))
+			{
+				if (FChaosClothAssetAddWeightMapNode* const Node = BaseNode->AsType<FChaosClothAssetAddWeightMapNode>())
+				{
+					Swap(Node->GetVertexWeights(), SavedWeights);
+					Swap(Node->GetRenderVertexWeights(), SavedRenderWeights);
+					Swap(Node->MapOverrideType, SavedMapOverrideType);
+					Swap(Node->Name, SavedWeightMapName);
+
+					Node->Invalidate();
+				}
+			}
+		}
+	}
+};
+
+TUniquePtr<FToolCommandChange> FChaosClothAssetAddWeightMapNode::MakeWeightMapNodeChange(const FChaosClothAssetAddWeightMapNode& Node)
+{
+	return MakeUnique<FChaosClothAssetAddWeightMapNode::FWeightMapNodeChange>(Node);
+}
+
+
+
 #undef LOCTEXT_NAMESPACE
diff --git a/Engine/Plugins/ChaosClothAssetEditor/Source/ChaosClothAssetDataflowNodes/Private/ChaosClothAsset/SelectionNode.cpp b/Engine/Plugins/ChaosClothAssetEditor/Source/ChaosClothAssetDataflowNodes/Private/ChaosClothAsset/SelectionNode.cpp
--- a/Engine/Plugins/ChaosClothAssetEditor/Source/ChaosClothAssetDataflowNodes/Private/ChaosClothAsset/SelectionNode.cpp
+++ b/Engine/Plugins/ChaosClothAssetEditor/Source/ChaosClothAssetDataflowNodes/Private/ChaosClothAsset/SelectionNode.cpp
@@ -10,6 +10,7 @@
 #include "ChaosClothAsset/WeightedValue.h"
 #include "Dataflow/DataflowInputOutput.h"
 #include "Dataflow/DataflowObject.h"
+#include "InteractiveToolChange.h"
 
 #include UE_INLINE_GENERATED_CPP_BY_NAME(SelectionNode)
 
@@ -502,4 +503,78 @@ void FChaosClothAssetSelectionNode::CalculateFinalSecondarySet(const TSet<int32>
 	UE::Chaos::ClothAsset::Private::CalculateFinalSet(InputSet, FinalSet, SelectionOverrideType, SecondaryIndices, RemoveSecondaryIndices);
 }
 
+
+// Object encapsulating a change to the Selection Node's values. Used for Undo/Redo.
+class FChaosClothAssetSelectionNode::FSelectionNodeChange final : public FToolCommandChange
+{
+public:
+
+	FSelectionNodeChange(const FChaosClothAssetSelectionNode& Node) :
+		NodeGuid(Node.GetGuid()),
+		SavedName(Node.Name),
+		SavedSelectionOverrideType(Node.SelectionOverrideType),
+		SavedGroup(Node.Group),
+		SavedIndices(Node.Indices),
+		SavedRemoveIndices(Node.RemoveIndices),
+		SavedSecondaryGroup(Node.SecondaryGroup),
+		SavedSecondaryIndices(Node.SecondaryIndices),
+		SavedRemoveSecondaryIndices(Node.RemoveSecondaryIndices)
+	{}
+
+private:
+
+	FGuid NodeGuid;
+	FString SavedName;
+	EChaosClothAssetSelectionOverrideType SavedSelectionOverrideType;
+	FChaosClothAssetNodeSelectionGroup SavedGroup;
+	TSet<int32> SavedIndices;
+	TSet<int32> SavedRemoveIndices;
+	FChaosClothAssetNodeSelectionGroup SavedSecondaryGroup;
+	TSet<int32> SavedSecondaryIndices;
+	TSet<int32> SavedRemoveSecondaryIndices;
+
+	virtual FString ToString() const final
+	{
+		return TEXT("ChaosClothAssetSelectionNodeChange");
+	}
+
+	virtual void Apply(UObject* Object) final
+	{
+		SwapApplyRevert(Object);
+	}
+
+	virtual void Revert(UObject* Object) final
+	{
+		SwapApplyRevert(Object);
+	}
+
+	void SwapApplyRevert(UObject* Object)
+	{
+		if (UDataflow* const Dataflow = Cast<UDataflow>(Object))
+		{
+			if (const TSharedPtr<FDataflowNode> BaseNode = Dataflow->GetDataflow()->FindBaseNode(NodeGuid))
+			{
+				if (FChaosClothAssetSelectionNode* const Node = BaseNode->AsType<FChaosClothAssetSelectionNode>())
+				{
+					Swap(Node->Name, SavedName);
+					Swap(Node->SelectionOverrideType, SavedSelectionOverrideType);
+					Swap(Node->Group, SavedGroup);
+					Swap(Node->Indices, SavedIndices);
+					Swap(Node->RemoveIndices, SavedRemoveIndices);
+					Swap(Node->SecondaryGroup, SavedSecondaryGroup);
+					Swap(Node->SecondaryIndices, SavedSecondaryIndices);
+					Swap(Node->RemoveSecondaryIndices, SavedRemoveSecondaryIndices);
+					Node->Invalidate();
+				}
+			}
+		}
+	}
+};
+
+TUniquePtr<FToolCommandChange> FChaosClothAssetSelectionNode::MakeWeightMapNodeChange(const FChaosClothAssetSelectionNode& Node)
+{
+	return MakeUnique<FChaosClothAssetSelectionNode::FSelectionNodeChange>(Node);
+}
+
+
 #undef LOCTEXT_NAMESPACE
diff --git a/Engine/Plugins/ChaosClothAssetEditor/Source/ChaosClothAssetDataflowNodes/Public/ChaosClothAsset/AddWeightMapNode.h b/Engine/Plugins/ChaosClothAssetEditor/Source/ChaosClothAssetDataflowNodes/Public/ChaosClothAsset/AddWeightMapNode.h
--- a/Engine/Plugins/ChaosClothAssetEditor/Source/ChaosClothAssetDataflowNodes/Public/ChaosClothAsset/AddWeightMapNode.h
+++ b/Engine/Plugins/ChaosClothAssetEditor/Source/ChaosClothAssetDataflowNodes/Public/ChaosClothAsset/AddWeightMapNode.h
@@ -93,7 +93,9 @@ public:
 
 private:
 	friend class UClothEditorWeightMapPaintTool;
-	friend class FWeightMapNodeChange;
+
+	class FWeightMapNodeChange;
+	static TUniquePtr<class FToolCommandChange> CHAOSCLOTHASSETDATAFLOWNODES_API MakeWeightMapNodeChange(const FChaosClothAssetAddWeightMapNode& Node);
 
 	const TArray<float>& GetVertexWeights() const { return VertexWeights; }
 	TArray<float>& GetVertexWeights() { return VertexWeights; }
diff --git a/Engine/Plugins/ChaosClothAssetEditor/Source/ChaosClothAssetDataflowNodes/Public/ChaosClothAsset/SelectionNode.h b/Engine/Plugins/ChaosClothAssetEditor/Source/ChaosClothAssetDataflowNodes/Public/ChaosClothAsset/SelectionNode.h
--- a/Engine/Plugins/ChaosClothAssetEditor/Source/ChaosClothAssetDataflowNodes/Public/ChaosClothAsset/SelectionNode.h
+++ b/Engine/Plugins/ChaosClothAssetEditor/Source/ChaosClothAssetDataflowNodes/Public/ChaosClothAsset/SelectionNode.h
@@ -149,6 +149,11 @@ public:
 
 private:
 
+	friend class UClothMeshSelectionTool;
+
+	class FSelectionNodeChange;
+	static TUniquePtr<class FToolCommandChange> CHAOSCLOTHASSETDATAFLOWNODES_API MakeWeightMapNodeChange(const FChaosClothAssetSelectionNode& Node);
+
 	virtual void SetAssetValue(TObjectPtr<UObject> Asset, Dataflow::FContext& Context) const override;
 	virtual void Evaluate(Dataflow::FContext& Context, const FDataflowOutput* Out) const override;
 	virtual void OnSelected(Dataflow::FContext& Context) override;
diff --git a/Engine/Plugins/ChaosClothAssetEditor/Source/ChaosClothAssetEditorTools/Private/ChaosClothAsset/ClothMeshSelectionTool.cpp b/Engine/Plugins/ChaosClothAssetEditor/Source/ChaosClothAssetEditorTools/Private/ChaosClothAsset/ClothMeshSelectionTool.cpp
--- a/Engine/Plugins/ChaosClothAssetEditor/Source/ChaosClothAssetEditorTools/Private/ChaosClothAsset/ClothMeshSelectionTool.cpp
+++ b/Engine/Plugins/ChaosClothAssetEditor/Source/ChaosClothAssetEditorTools/Private/ChaosClothAsset/ClothMeshSelectionTool.cpp
@@ -536,6 +536,16 @@ void UClothMeshSelectionTool::UpdateSelectedNode()
 {
 	using namespace UE::Chaos::ClothAsset;
 
+	checkf(SelectionNodeToUpdate, TEXT("Expected non-null pointer to Selection Node"));
+
+	// Save previous state for undo
+	if (UDataflow* const Dataflow = ClothEditorContextObject->GetDataflowAsset())
+	{
+		GetToolManager()->GetContextTransactionsAPI()->AppendChange(Dataflow, 
+			FChaosClothAssetSelectionNode::MakeWeightMapNodeChange(*SelectionNodeToUpdate),
+			LOCTEXT("SelectionNodeChangeDescription", "Update Selection Node"));
+	}
+
 	const UE::Geometry::FGroupTopologySelection& Selection = SelectionMechanic->GetActiveSelection();
 	const EClothPatternVertexType ViewMode = ClothEditorContextObject->GetConstructionViewMode();
 
diff --git a/Engine/Plugins/ChaosClothAssetEditor/Source/ChaosClothAssetEditorTools/Private/ChaosClothAsset/ClothWeightMapPaintTool.cpp b/Engine/Plugins/ChaosClothAssetEditor/Source/ChaosClothAssetEditorTools/Private/ChaosClothAsset/ClothWeightMapPaintTool.cpp
--- a/Engine/Plugins/ChaosClothAssetEditor/Source/ChaosClothAssetEditorTools/Private/ChaosClothAsset/ClothWeightMapPaintTool.cpp
+++ b/Engine/Plugins/ChaosClothAssetEditor/Source/ChaosClothAssetEditorTools/Private/ChaosClothAsset/ClothWeightMapPaintTool.cpp
@@ -1968,89 +1968,20 @@ void UClothEditorWeightMapPaintTool::ClearHiddenAction()
 }
 
 
-//
-// Object encapsulating a change to the AddWeightMap node's values. Used for Undo/Redo.
-//
-class FWeightMapNodeChange : public FToolCommandChange
-{
-
-public:
-	
-	FWeightMapNodeChange(FGuid NodeGuid, const TArray<float>& Weights, TArray<float>& RenderWeights, EChaosClothAssetWeightMapOverrideType MapOverrideType, const FString& WeightMapName) :
-		NodeGuid(NodeGuid),
-		SavedWeights(Weights),
-		SavedRenderWeights(RenderWeights),
-		SavedMapOverrideType(MapOverrideType),
-		SavedWeightMapName(WeightMapName)
-	{}
-	
-
-private:
-
-	FGuid NodeGuid;
-
-	TArray<float> SavedWeights;
-
-	// Note we could store only one set of weights and use a bool to determine whether we are updating sim or render vertices, however in the future 
-	// we may enable writing both weight maps to the node at once.
-	TArray<float> SavedRenderWeights;
-
-	EChaosClothAssetWeightMapOverrideType SavedMapOverrideType;
-	FString SavedWeightMapName;
-
-	virtual FString ToString() const override { return TEXT("WeightMapNodeUpdateChange"); }
-
-	virtual void Apply(UObject* Object) final
-	{
-		SwapApplyRevert(Object);
-	}
-
-	virtual void Revert(UObject* Object) final
-	{
-		SwapApplyRevert(Object);
-	}
-
-	void SwapApplyRevert(UObject* Object)
-	{
-		if (UDataflow* const Dataflow = Cast<UDataflow>(Object))
-		{
-			if (const TSharedPtr<FDataflowNode> BaseNode = Dataflow->GetDataflow()->FindBaseNode(NodeGuid))
-			{
-				if (FChaosClothAssetAddWeightMapNode* const Node = BaseNode->AsType<FChaosClothAssetAddWeightMapNode>())
-				{
-					Swap(Node->GetVertexWeights(), SavedWeights);
-					Swap(Node->GetRenderVertexWeights(), SavedRenderWeights);
-					Swap(Node->MapOverrideType, SavedMapOverrideType);
-					Swap(Node->Name, SavedWeightMapName);
-
-					Node->Invalidate();
-				}
-			}
-		}
-	}
-
-};
-
-
 void UClothEditorWeightMapPaintTool::UpdateSelectedNode()
 {
 	check(ActiveWeightMap);
 	TArray<float> CurrentWeights;
 	GetCurrentWeightMap(CurrentWeights);
 
-	check(WeightMapNodeToUpdate);
+	checkf(WeightMapNodeToUpdate, TEXT("Expected non-null pointer to Add Weight Map Node"));
 
 	// Save previous state for undo
 	if (UDataflow* const Dataflow = ClothEditorContextObject->GetDataflowAsset())
 	{
-		TUniquePtr<FWeightMapNodeChange> Change = MakeUnique<FWeightMapNodeChange>(
-			WeightMapNodeToUpdate->GetGuid(), 
-			WeightMapNodeToUpdate->GetVertexWeights(),
-			WeightMapNodeToUpdate->GetRenderVertexWeights(), 
-			WeightMapNodeToUpdate->MapOverrideType, 
-			WeightMapNodeToUpdate->Name);
-
-		GetToolManager()->GetContextTransactionsAPI()->AppendChange(Dataflow, MoveTemp(Change), LOCTEXT("WeightMapNodeChangeDescription", "Update Weight Map Node"));
+		GetToolManager()->GetContextTransactionsAPI()->AppendChange(Dataflow, 
+			FChaosClothAssetAddWeightMapNode::MakeWeightMapNodeChange(*WeightMapNodeToUpdate),
+			LOCTEXT("WeightMapNodeChangeDescription", "Update Weight Map Node"));
 	}
 
 
