Chaos Cloth Asset Editor: Added weight map selection menu to debug draw.

diff --git a/Engine/Plugins/ChaosCloth/Source/ChaosCloth/Private/ChaosCloth/ChaosClothVisualization.cpp b/Engine/Plugins/ChaosCloth/Source/ChaosCloth/Private/ChaosCloth/ChaosClothVisualization.cpp
--- a/Engine/Plugins/ChaosCloth/Source/ChaosCloth/Private/ChaosCloth/ChaosClothVisualization.cpp
+++ b/Engine/Plugins/ChaosCloth/Source/ChaosCloth/Private/ChaosCloth/ChaosClothVisualization.cpp
@@ -2682,6 +2682,16 @@ FLinearColor PseudoRandomColor(int32 NumColorRotations)
 	}
 
 #if WITH_EDITOR
+	TArray<FString> FClothVisualization::GetAllWeightMapNames() const
+	{
+		TSet<FString> AllNames;
+		for (const FClothingSimulationCloth* const Cloth : Solver->GetCloths())
+		{
+			AllNames.Append(Cloth->GetAllWeightMapNames());
+		}
+		return AllNames.Array();
+	}
+
 	void FClothVisualization::DrawKinematicColliderShaded(FPrimitiveDrawInterface* PDI) const
 	{
 		if (!Solver || !CollisionMaterial)
diff --git a/Engine/Plugins/ChaosCloth/Source/ChaosCloth/Private/ChaosCloth/ChaosClothingSimulationCloth.cpp b/Engine/Plugins/ChaosCloth/Source/ChaosCloth/Private/ChaosCloth/ChaosClothingSimulationCloth.cpp
--- a/Engine/Plugins/ChaosCloth/Source/ChaosCloth/Private/ChaosCloth/ChaosClothingSimulationCloth.cpp
+++ b/Engine/Plugins/ChaosCloth/Source/ChaosCloth/Private/ChaosCloth/ChaosClothingSimulationCloth.cpp
@@ -969,6 +969,18 @@ TConstArrayView<FRealSingle> FClothingSimulationCloth::GetWeightMapByProperty(co
 	return GetWeightMapByName(Solver, PropertyString);
 }
 
+TSet<FString> FClothingSimulationCloth::GetAllWeightMapNames() const
+{
+	TSet<FString> Names;
+	for (const TUniquePtr<FLODData>& Data : LODData)
+	{
+		TSet<FString> LODNames;
+		Data->WeightMaps.GetKeys(LODNames);
+		Names.Append(MoveTemp(LODNames));
+	}
+	return Names;
+}
+
 TConstArrayView<int32> FClothingSimulationCloth::GetFaceIntMapByName(const FClothingSimulationSolver* Solver, const FString& Name) const
 {
 	const int32 LODIndex = LODIndices.FindChecked(Solver);
diff --git a/Engine/Plugins/ChaosCloth/Source/ChaosCloth/Public/ChaosCloth/ChaosClothVisualization.h b/Engine/Plugins/ChaosCloth/Source/ChaosCloth/Public/ChaosCloth/ChaosClothVisualization.h
--- a/Engine/Plugins/ChaosCloth/Source/ChaosCloth/Public/ChaosCloth/ChaosClothVisualization.h
+++ b/Engine/Plugins/ChaosCloth/Source/ChaosCloth/Public/ChaosCloth/ChaosClothVisualization.h
@@ -96,6 +96,8 @@ namespace Chaos
 		CHAOSCLOTH_API void DrawSelfCollisionLayers(FPrimitiveDrawInterface* PDI = nullptr) const;
 		CHAOSCLOTH_API void DrawInpaintWeightsMatched(FPrimitiveDrawInterface* PDI = nullptr) const;
 		CHAOSCLOTH_API void DrawKinematicColliderShaded(FPrimitiveDrawInterface* PDI) const;
+		CHAOSCLOTH_API void DrawWeightMapWithName(FPrimitiveDrawInterface* PDI, const FString& Name) const;
+		CHAOSCLOTH_API TArray<FString> GetAllWeightMapNames() const;
 	protected:
 		// FGCObject interface
 		CHAOSCLOTH_API virtual void AddReferencedObjects(FReferenceCollector& Collector) override;
@@ -120,8 +122,6 @@ namespace Chaos
 		TObjectPtr<const UMaterial> ClothMaterialColor = nullptr;
 		TObjectPtr<const UMaterial> ClothMaterialVertex = nullptr;
 		TObjectPtr<const UMaterial> CollisionMaterial = nullptr;
-		
-		void DrawWeightMapWithName(FPrimitiveDrawInterface* PDI, const FString& Name) const;
 #endif  // #if WITH_EDITOR
 #endif  // #if CHAOS_DEBUG_DRAW
 	};
diff --git a/Engine/Plugins/ChaosCloth/Source/ChaosCloth/Public/ChaosCloth/ChaosClothingSimulationCloth.h b/Engine/Plugins/ChaosCloth/Source/ChaosCloth/Public/ChaosCloth/ChaosClothingSimulationCloth.h
--- a/Engine/Plugins/ChaosCloth/Source/ChaosCloth/Public/ChaosCloth/ChaosClothingSimulationCloth.h
+++ b/Engine/Plugins/ChaosCloth/Source/ChaosCloth/Public/ChaosCloth/ChaosClothingSimulationCloth.h
@@ -183,6 +183,8 @@ PRAGMA_ENABLE_DEPRECATION_WARNINGS
 		CHAOSCLOTH_API TConstArrayView<FRealSingle> GetWeightMapByName(const FClothingSimulationSolver* Solver, const FString& Name) const;
 		// Return the weight map of the specified property name if it exists and is available on the current LOD, or an empty array view otherwise.
 		CHAOSCLOTH_API TConstArrayView<FRealSingle> GetWeightMapByProperty(const FClothingSimulationSolver* Solver, const FString& Property) const;
+		// Return list of weight map names available across all LODs
+		CHAOSCLOTH_API TSet<FString> GetAllWeightMapNames() const;
 		// Return the face int map of the specified name if available on the current LOD, or an empty array view otherwise.
 		CHAOSCLOTH_API TConstArrayView<int32> GetFaceIntMapByName(const FClothingSimulationSolver* Solver, const FString& Name) const;
 		// Return the face int map of the specified property name if it exists and is available on the current LOD, or an empty array view otherwise.
diff --git a/Engine/Plugins/ChaosClothAssetEditor/Source/ChaosClothAssetEditor/Private/ChaosClothAsset/ClothEditorSimulationVisualization.cpp b/Engine/Plugins/ChaosClothAssetEditor/Source/ChaosClothAssetEditor/Private/ChaosClothAsset/ClothEditorSimulationVisualization.cpp
--- a/Engine/Plugins/ChaosClothAssetEditor/Source/ChaosClothAssetEditor/Private/ChaosClothAsset/ClothEditorSimulationVisualization.cpp
+++ b/Engine/Plugins/ChaosClothAssetEditor/Source/ChaosClothAssetEditor/Private/ChaosClothAsset/ClothEditorSimulationVisualization.cpp
@@ -1,12 +1,17 @@
 // Copyright Epic Games, Inc. All Rights Reserved.
 
 #include "ClothEditorSimulationVisualization.h"
+#include "ChaosClothAsset/ClothAsset.h"
 #include "ChaosClothAsset/ClothEditor3DViewportClient.h"
 #include "ChaosClothAsset/ClothSimulationProxy.h"
 #include "ChaosClothAsset/ClothComponent.h"
 #include "ChaosCloth/ChaosClothVisualization.h"
 #include "Framework/MultiBox/MultiBoxBuilder.h"
 #include "Rendering/SkeletalMeshRenderData.h"
+#include "Delegates/Delegate.h"
+#include "Delegates/DelegateCombinations.h"
+#include "Widgets/Input/STextComboBox.h"
+#include "Types/SlateEnums.h"
 
 
 #define LOCTEXT_NAMESPACE "ClothEditorSimulationVisualization"
@@ -24,68 +29,73 @@ FText ConcatenateLine(const FText& InText, const FText& InNewLine)
 	return FText::Format(LOCTEXT("ViewportTextNewlineFormatter", "{0}\n{1}"), InText, InNewLine);
 }
 
-static FText GetSimulationStatisticsString(const FClothSimulationProxy* SimProxy)
+static FText GetSimulationStatisticsString(const FClothSimulationProxy& SimProxy)
 {
 	FText TextValue;
-	if (SimProxy)
+	// Cloth stats
+	if (const int32 NumActiveCloths = SimProxy.GetNumCloths())
 	{
-		// Cloth stats
-		if (const int32 NumActiveCloths = SimProxy->GetNumCloths())
-		{
-			TextValue = ConcatenateLine(TextValue, FText::Format(LOCTEXT("NumActiveCloths", "Active Cloths: {0}"), NumActiveCloths));
-		}
-		if (const int32 NumKinematicParticles = SimProxy->GetNumKinematicParticles())
-		{
-			TextValue = ConcatenateLine(TextValue, FText::Format(LOCTEXT("NumKinematicParticles", "Kinematic Particles: {0}"), NumKinematicParticles));
-		}
-		if (const int32 NumDynamicParticles = SimProxy->GetNumDynamicParticles())
-		{
-			TextValue = ConcatenateLine(TextValue, FText::Format(LOCTEXT("NumDynamicParticles", "Dynamic Particles: {0}"), NumDynamicParticles));
-		}
-		if (const int32 NumIterations = SimProxy->GetNumIterations())
-		{
-			TextValue = ConcatenateLine(TextValue, FText::Format(LOCTEXT("NumIterations", "Iterations: {0}"), NumIterations));
-		}
-		if (const int32 NumSubSteps = SimProxy->GetNumSubsteps())
-		{
-			TextValue = ConcatenateLine(TextValue, FText::Format(LOCTEXT("NumSubsteps", "Substeps: {0}"), NumSubSteps));
-		}
-		if (const int32 NumLinearSolveIterations = SimProxy->GetNumLinearSolveIterations())
-		{
-			TextValue = ConcatenateLine(TextValue, FText::Format(LOCTEXT("NumCGIterations", "CGIterations: {0}"), NumLinearSolveIterations));
-		}
-		if (const float LinearSolveError = SimProxy->GetLinearSolveError())
-		{
-			FNumberFormattingOptions NumberFormatOptions;
-			NumberFormatOptions.AlwaysSign = false;
-			NumberFormatOptions.UseGrouping = false;
-			NumberFormatOptions.RoundingMode = ERoundingMode::HalfFromZero;
-			NumberFormatOptions.MinimumIntegralDigits = 1;
-			NumberFormatOptions.MaximumIntegralDigits = 6;
-			NumberFormatOptions.MinimumFractionalDigits = 2;
-			NumberFormatOptions.MaximumFractionalDigits = 6;
-			TextValue = ConcatenateLine(TextValue, FText::Format(LOCTEXT("CGError", "CGError: {0}"), FText::AsNumber(LinearSolveError, &NumberFormatOptions)));
-		}
-		if (const float SimulationTime = SimProxy->GetSimulationTime())
-		{
-			FNumberFormattingOptions NumberFormatOptions;
-			NumberFormatOptions.AlwaysSign = false;
-			NumberFormatOptions.UseGrouping = false;
-			NumberFormatOptions.RoundingMode = ERoundingMode::HalfFromZero;
-			NumberFormatOptions.MinimumIntegralDigits = 1;
-			NumberFormatOptions.MaximumIntegralDigits = 6;
-			NumberFormatOptions.MinimumFractionalDigits = 2;
-			NumberFormatOptions.MaximumFractionalDigits = 2;
-			TextValue = ConcatenateLine(TextValue, FText::Format(LOCTEXT("SimulationTime", "Simulation Time: {0}ms"), FText::AsNumber(SimulationTime, &NumberFormatOptions)));
-		}
-		if (SimProxy->IsTeleported())
-		{
-			TextValue = ConcatenateLine(TextValue, LOCTEXT("IsTeleported", "Simulation Teleport Activated"));
-		}
+		TextValue = ConcatenateLine(TextValue, FText::Format(LOCTEXT("NumActiveCloths", "Active Cloths: {0}"), NumActiveCloths));
+	}
+	if (const int32 NumKinematicParticles = SimProxy.GetNumKinematicParticles())
+	{
+		TextValue = ConcatenateLine(TextValue, FText::Format(LOCTEXT("NumKinematicParticles", "Kinematic Particles: {0}"), NumKinematicParticles));
+	}
+	if (const int32 NumDynamicParticles = SimProxy.GetNumDynamicParticles())
+	{
+		TextValue = ConcatenateLine(TextValue, FText::Format(LOCTEXT("NumDynamicParticles", "Dynamic Particles: {0}"), NumDynamicParticles));
+	}
+	if (const int32 NumIterations = SimProxy.GetNumIterations())
+	{
+		TextValue = ConcatenateLine(TextValue, FText::Format(LOCTEXT("NumIterations", "Iterations: {0}"), NumIterations));
+	}
+	if (const int32 NumSubSteps = SimProxy.GetNumSubsteps())
+	{
+		TextValue = ConcatenateLine(TextValue, FText::Format(LOCTEXT("NumSubsteps", "Substeps: {0}"), NumSubSteps));
+	}
+	if (const int32 NumLinearSolveIterations = SimProxy.GetNumLinearSolveIterations())
+	{
+		TextValue = ConcatenateLine(TextValue, FText::Format(LOCTEXT("NumCGIterations", "CGIterations: {0}"), NumLinearSolveIterations));
+	}
+	if (const float LinearSolveError = SimProxy.GetLinearSolveError())
+	{
+		FNumberFormattingOptions NumberFormatOptions;
+		NumberFormatOptions.AlwaysSign = false;
+		NumberFormatOptions.UseGrouping = false;
+		NumberFormatOptions.RoundingMode = ERoundingMode::HalfFromZero;
+		NumberFormatOptions.MinimumIntegralDigits = 1;
+		NumberFormatOptions.MaximumIntegralDigits = 6;
+		NumberFormatOptions.MinimumFractionalDigits = 2;
+		NumberFormatOptions.MaximumFractionalDigits = 6;
+		TextValue = ConcatenateLine(TextValue, FText::Format(LOCTEXT("CGError", "CGError: {0}"), FText::AsNumber(LinearSolveError, &NumberFormatOptions)));
+	}
+	if (const float SimulationTime = SimProxy.GetSimulationTime())
+	{
+		FNumberFormattingOptions NumberFormatOptions;
+		NumberFormatOptions.AlwaysSign = false;
+		NumberFormatOptions.UseGrouping = false;
+		NumberFormatOptions.RoundingMode = ERoundingMode::HalfFromZero;
+		NumberFormatOptions.MinimumIntegralDigits = 1;
+		NumberFormatOptions.MaximumIntegralDigits = 6;
+		NumberFormatOptions.MinimumFractionalDigits = 2;
+		NumberFormatOptions.MaximumFractionalDigits = 2;
+		TextValue = ConcatenateLine(TextValue, FText::Format(LOCTEXT("SimulationTime", "Simulation Time: {0}ms"), FText::AsNumber(SimulationTime, &NumberFormatOptions)));
+	}
+	if (SimProxy.IsTeleported())
+	{
+		TextValue = ConcatenateLine(TextValue, LOCTEXT("IsTeleported", "Simulation Teleport Activated"));
 	}
 	return TextValue;
 }
 
+DECLARE_DELEGATE_ThreeParams(FClothVisualizationDebugDraw, const FClothEditorSimulationVisualization&, const ::Chaos::FClothVisualization&, FPrimitiveDrawInterface*);
+
+DECLARE_DELEGATE_FourParams(FClothVisualizationDebugDrawTexts, const FClothEditorSimulationVisualization&, const ::Chaos::FClothVisualization&, FCanvas*, const FSceneView*);
+
+DECLARE_DELEGATE_RetVal_TwoParams(FText, FLocalDebugDisplayString, const FClothEditorSimulationVisualization&, const FClothSimulationProxy& );
+
+DECLARE_DELEGATE_ThreeParams(FAdditionalMenus, FClothEditorSimulationVisualization&, FMenuBuilder&, TSharedRef<FChaosClothAssetEditor3DViewportClient>);
+
 struct FVisualizationOption
 {
 	// Actual option entries
@@ -93,90 +103,289 @@ struct FVisualizationOption
 	static const uint32 Count;
 
 	// Chaos debug draw function
-	typedef void (::Chaos::FClothVisualization::* FClothVisualizationDebugDrawFunction)(FPrimitiveDrawInterface*) const;
-	typedef void (::Chaos::FClothVisualization::* FClothVisualizationDebugDrawTextsFunction)(FCanvas*, const FSceneView*) const;
-	typedef FText (*FLocalDebugDisplayStringFunction)(const FClothSimulationProxy*);
-	FClothVisualizationDebugDrawFunction ClothVisualizationDebugDrawFunction;
-	FClothVisualizationDebugDrawTextsFunction ClothVisualizationDebugDrawTextsFunction;
-	FLocalDebugDisplayStringFunction LocalDebugDisplayStringFunction;
+	FClothVisualizationDebugDraw ClothVisualizationDebugDraw;
+	FClothVisualizationDebugDrawTexts ClothVisualizationDebugDrawTexts;
+	FLocalDebugDisplayString LocalDebugDisplayString;
+
+	// Extra menu building function
+	FAdditionalMenus AdditionalMenus;
 
 	FText DisplayName;         // Text for menu entries.
 	FText ToolTip;             // Text for menu tooltips.
 	bool bDisablesSimulation;  // Whether or not this option requires the simulation to be disabled.
 	bool bHidesClothSections;  // Hides the cloth section to avoid zfighting with the debug geometry.
 
-	// Console override
-	IConsoleVariable* const ConsoleVariable;
-
-	FVisualizationOption(FClothVisualizationDebugDrawFunction InDebugDrawFunction, const TCHAR* ConsoleName, const FText& InDisplayName, const FText& InToolTip, bool bInDisablesSimulation = false, bool bInHidesClothSections = false)
-		: ClothVisualizationDebugDrawFunction(InDebugDrawFunction)
+	FVisualizationOption(const FClothVisualizationDebugDraw& InClothVisualizationDebugDraw, const FText& InDisplayName, const FText& InToolTip, bool bInDisablesSimulation = false, bool bInHidesClothSections = false, const FAdditionalMenus& InAdditionalMenus = FAdditionalMenus())
+		: ClothVisualizationDebugDraw(InClothVisualizationDebugDraw)
+		, AdditionalMenus(InAdditionalMenus)
 		, DisplayName(InDisplayName)
 		, ToolTip(InToolTip)
 		, bDisablesSimulation(bInDisablesSimulation)
 		, bHidesClothSections(bInHidesClothSections)
-		, ConsoleVariable(IConsoleManager::Get().RegisterConsoleVariable(ConsoleName, false, *InToolTip.ToString(), 0))
 	{}
 
-	FVisualizationOption(FLocalDebugDisplayStringFunction InLocalDebugDisplayStringFunction, const TCHAR* ConsoleName, const FText& InDisplayName, const FText& InToolTip, bool bInDisablesSimulation = false, bool bInHidesClothSections = false)
-		: LocalDebugDisplayStringFunction(InLocalDebugDisplayStringFunction)
+	FVisualizationOption(const FLocalDebugDisplayString& InLocalDebugDisplayString, const FText& InDisplayName, const FText& InToolTip, bool bInDisablesSimulation = false, bool bInHidesClothSections = false)
+		: LocalDebugDisplayString(InLocalDebugDisplayString)
 		, DisplayName(InDisplayName)
 		, ToolTip(InToolTip)
 		, bDisablesSimulation(bInDisablesSimulation)
 		, bHidesClothSections(bInHidesClothSections)
-		, ConsoleVariable(IConsoleManager::Get().RegisterConsoleVariable(ConsoleName, false, *InToolTip.ToString(), 0))
 	{}
 
-	FVisualizationOption(FClothVisualizationDebugDrawTextsFunction InDebugDrawTextsFunction, const TCHAR * ConsoleName, const FText & InDisplayName, const FText & InToolTip, bool bInDisablesSimulation = false, bool bInHidesClothSections = false)
-		: ClothVisualizationDebugDrawTextsFunction(InDebugDrawTextsFunction)
+	FVisualizationOption(const FClothVisualizationDebugDrawTexts& InClothVisualizationDebugDrawTexts, const FText & InDisplayName, const FText & InToolTip, bool bInDisablesSimulation = false, bool bInHidesClothSections = false)
+		: ClothVisualizationDebugDrawTexts(InClothVisualizationDebugDrawTexts)
 		, DisplayName(InDisplayName)
 		, ToolTip(InToolTip)
 		, bDisablesSimulation(bInDisablesSimulation)
 		, bHidesClothSections(bInHidesClothSections)
-		, ConsoleVariable(IConsoleManager::Get().RegisterConsoleVariable(ConsoleName, false, *InToolTip.ToString(), 0))
+	{}
+	
+	FVisualizationOption(const FText & InDisplayName, const FText & InToolTip, bool bInDisablesSimulation = false, bool bInHidesClothSections = false)
+		: DisplayName(InDisplayName)
+		, ToolTip(InToolTip)
+		, bDisablesSimulation(bInDisablesSimulation)
+		, bHidesClothSections(bInHidesClothSections)
 	{}
 
 	~FVisualizationOption()
 	{
-		IConsoleManager::Get().UnregisterConsoleObject(ConsoleVariable);
 	}
-
-	bool IsConsoleVariableEnabled() const { return ConsoleVariable->AsVariableBool()->GetValueOnGameThread(); }
 };
 
 const FVisualizationOption FVisualizationOption::OptionData[] =
 {
-	FVisualizationOption(&::Chaos::FClothVisualization::DrawPhysMeshShaded, TEXT("p.ChaosClothAssetEditor.DebugDrawPhysMeshShaded"), LOCTEXT("ChaosVisName_PhysMesh", "Physical Mesh (Flat Shaded)"), LOCTEXT("ChaosVisName_PhysMeshShaded_ToolTip", "Draws the current physical result as a doubled sided flat shaded mesh"), /*bDisablesSimulation =*/false, /*bHidesClothSections=*/true),
-	FVisualizationOption(&::Chaos::FClothVisualization::DrawPhysMeshWired, TEXT("p.ChaosClothAssetEditor.DebugDrawPhysMeshWired"), LOCTEXT("ChaosVisName_PhysMeshWire", "Physical Mesh (Wireframe)"), LOCTEXT("ChaosVisName_PhysMeshWired_ToolTip", "Draws the current physical mesh result in wireframe")),
-	FVisualizationOption(&::Chaos::FClothVisualization::DrawAnimMeshWired, TEXT("p.ChaosClothAssetEditor.DebugDrawAnimMeshWired"), LOCTEXT("ChaosVisName_AnimMeshWire", "Animated Mesh (Wireframe)"), LOCTEXT("ChaosVisName_AnimMeshWired_ToolTip", "Draws the current animated mesh input in wireframe")),
-	FVisualizationOption((FLocalDebugDisplayStringFunction)nullptr, TEXT("p.ChaosClothAssetEditor.DebugDrawHideRenderMesh"), LOCTEXT("ChaosVisName_HideRenderMesh", "Hide Render Mesh"), LOCTEXT("ChaosVisName_HideRenderMesh_ToolTip", "Hide the render mesh."), /*bDisablesSimulation =*/false, /*bHidesClothSections=*/true),
-	FVisualizationOption(&::Chaos::FClothVisualization::DrawParticleIndices, TEXT("p.ChaosClothAssetEditor.DebugDrawParticleIndices"), LOCTEXT("ChaosVisName_ParticleIndices", "Particle Indices"), LOCTEXT("ChaosVisName_ParticleIndices_ToolTip", "Draws the particle indices as instantiated by the solver")),
-	FVisualizationOption(&::Chaos::FClothVisualization::DrawElementIndices, TEXT("p.ChaosClothAssetEditor.DebugDrawElementIndices"), LOCTEXT("ChaosVisName_ElementIndices", "Element Indices"), LOCTEXT("ChaosVisName_ElementIndices_ToolTip", "Draws the element's (triangle or other) indices as instantiated by the solver")),
-	FVisualizationOption(&::Chaos::FClothVisualization::DrawPointNormals , TEXT("p.ChaosClothAssetEditor.DebugDrawPointNormals"), LOCTEXT("ChaosVisName_PointNormals", "Physical Mesh Normals"), LOCTEXT("ChaosVisName_PointNormals_ToolTip", "Draws the current point normals for the simulation mesh")),
-	FVisualizationOption(&::Chaos::FClothVisualization::DrawPointVelocities, TEXT("p.ChaosClothAssetEditor.DebugDrawPointVelocities"), LOCTEXT("ChaosVisName_PointVelocities", "Point Velocities"), LOCTEXT("ChaosVisName_PointVelocities_ToolTip", "Draws the current point velocities for the simulation mesh")),
-	FVisualizationOption(&::Chaos::FClothVisualization::DrawAnimNormals, TEXT("p.ChaosClothAssetEditor.DebugDrawAnimNormals"), LOCTEXT("ChaosVisName_AnimNormals", "Animated Mesh Normals"), LOCTEXT("ChaosVisName_AnimNormals_ToolTip", "Draws the current point normals for the animated mesh")),
-	FVisualizationOption(&::Chaos::FClothVisualization::DrawCollision, TEXT("p.ChaosClothAssetEditor.DebugDrawCollision"), LOCTEXT("ChaosVisName_Collision", "Collisions"), LOCTEXT("ChaosVisName_Collision_ToolTip", "Draws the collision bodies the simulation is currently using")),
-	FVisualizationOption(&::Chaos::FClothVisualization::DrawBackstops , TEXT("p.ChaosClothAssetEditor.DebugDrawBackstops"), LOCTEXT("ChaosVisName_Backstop", "Backstops"), LOCTEXT("ChaosVisName_Backstop_ToolTip", "Draws the backstop radius and position for each simulation particle")),
-	FVisualizationOption(&::Chaos::FClothVisualization::DrawBackstopDistances, TEXT("p.ChaosClothAssetEditor.DebugDrawBackstopDistances"), LOCTEXT("ChaosVisName_BackstopDistance", "Backstop Distances"), LOCTEXT("ChaosVisName_BackstopDistance_ToolTip", "Draws the backstop distance offset for each simulation particle"), /*bDisablesSimulation =*/true),
-	FVisualizationOption(&::Chaos::FClothVisualization::DrawMaxDistances, TEXT("p.ChaosClothAssetEditor.DebugDrawMaxDistances"), LOCTEXT("ChaosVisName_MaxDistance", "Max Distances"), LOCTEXT("ChaosVisName_MaxDistance_ToolTip", "Draws the current max distances for the sim particles as a line along its normal"), /*bDisablesSimulation =*/true),
-	FVisualizationOption(&::Chaos::FClothVisualization::DrawMaxDistanceValues, TEXT("p.ChaosClothAssetEditor.DebugDrawMaxDistanceValues"), LOCTEXT("ChaosVisName_MaxDistanceValue", "Max Distances As Numbers"), LOCTEXT("ChaosVisName_MaxDistanceValue_ToolTip", "Draws the current max distances as numbers")),
-	FVisualizationOption(&::Chaos::FClothVisualization::DrawAnimDrive, TEXT("p.ChaosClothAssetEditor.DebugDrawAnimDrive"), LOCTEXT("ChaosVisName_AnimDrive", "Anim Drive"), LOCTEXT("ChaosVisName_AnimDrive_Tooltip", "Draws the current skinned reference mesh for the simulation which anim drive will attempt to reach if enabled")),
-	FVisualizationOption(&::Chaos::FClothVisualization::DrawEdgeConstraint, TEXT("p.ChaosClothAssetEditor.DebugDrawEdgeConstraint"), LOCTEXT("ChaosVisName_EdgeConstraint", "Edge Constraint"), LOCTEXT("ChaosVisName_EdgeConstraint_Tooltip", "Draws the edge spring constraints")),
-	FVisualizationOption(&::Chaos::FClothVisualization::DrawBendingConstraint, TEXT("p.ChaosClothAssetEditor.DebugDrawBendingConstraint"), LOCTEXT("ChaosVisName_BendingConstraint", "Bending Constraint"), LOCTEXT("ChaosVisName_BendingConstraint_Tooltip", "Draws the bending spring constraints")),
-	FVisualizationOption(&::Chaos::FClothVisualization::DrawLongRangeConstraint, TEXT("p.ChaosClothAssetEditor.DebugDrawLongRangeConstraint"), LOCTEXT("ChaosVisName_LongRangeConstraint" , "Long Range Constraint"), LOCTEXT("ChaosVisName_LongRangeConstraint_Tooltip", "Draws the long range attachment constraint distances")),
-	FVisualizationOption(&::Chaos::FClothVisualization::DrawWindAndPressureForces, TEXT("p.ChaosClothAssetEditor.DebugDrawWindAndPressureForces"), LOCTEXT("ChaosVisName_WindAndPressureForces", "Wind Aerodynamic And Pressure Forces"), LOCTEXT("ChaosVisName_WindAndPressure_Tooltip", "Draws the Wind drag and lift and pressure forces")),
-	FVisualizationOption(&::Chaos::FClothVisualization::DrawLocalSpace, TEXT("p.ChaosClothAssetEditor.DebugDrawLocalSpace"), LOCTEXT("ChaosVisName_LocalSpace", "Local Space Reference Bone"), LOCTEXT("ChaosVisName_LocalSpace_Tooltip", "Draws the local space reference bone")),
-	FVisualizationOption(&::Chaos::FClothVisualization::DrawGravity, TEXT("p.ChaosClothAssetEditor.DebugDrawGravity"), LOCTEXT("ChaosVisName_Gravity", "Gravity"), LOCTEXT("ChaosVisName_Gravity_Tooltip", "Draws gravity")),
-	FVisualizationOption(&::Chaos::FClothVisualization::DrawFictitiousAngularForces, TEXT("p.ChaosClothAssetEditor.DebugDrawFictitiousAngularForces"), LOCTEXT("ChaosVisName_FictitiousAngularForces", "Fictitious Angular Forces"), LOCTEXT("ChaosVisName_Gravity_FictitiousAngularForces", "Draws fictitious angular forces (force based solver only)")),
-	FVisualizationOption(&::Chaos::FClothVisualization::DrawSelfCollision, TEXT("p.ChaosClothAssetEditor.DebugDrawSelfCollision"), LOCTEXT("ChaosVisName_SelfCollision"       , "Self Collision"), LOCTEXT("ChaosVisName_SelfCollision_Tooltip", "Draws the self collision thickness/debugging information")),
-	FVisualizationOption(&::Chaos::FClothVisualization::DrawSelfIntersection, TEXT("p.ChaosClothAssetEditor.DebugDrawSelfIntersection"), LOCTEXT("ChaosVisName_SelfIntersection", "Self Intersection"), LOCTEXT("ChaosVisName_SelfIntersection_Tooltip", "Draws the self intersection contour/region information")),
-	FVisualizationOption(&::Chaos::FClothVisualization::DrawSelfCollisionLayers, TEXT("p.ChaosClothAssetEditor.DebugDrawSelfCollisionLayers"), LOCTEXT("ChaosVisName_SelfCollisionLayers", "Self Collision Layers"), LOCTEXT("ChaosVisName_SelfCollisionLayers_Tooltip", "Draws the self collision layers"), /*bDisablesSimulation =*/false, /*bHidesClothSections=*/true),
-	FVisualizationOption(&::Chaos::FClothVisualization::DrawSelfCollisionThickness, TEXT("p.ChaosClothAssetEditor.DebugDrawSelfCollisionThickness"), LOCTEXT("ChaosVisName_SelfCollisionThickness", "Self Collision Thickness"), LOCTEXT("ChaosVisName_SelfCollisionThickness_Tooltip", "Draws the self collision Thickness")),
-	FVisualizationOption(&::Chaos::FClothVisualization::DrawKinematicColliderShaded, TEXT("p.ChaosClothAssetEditor.DebugDrawKinematicColliderShaded"), LOCTEXT("ChaosVisName_DrawKinematicColliderShaded", "Draw Kinematic Colliders (Shaded)"), LOCTEXT("ChaosVisName_DrawKinematicColliderShaded_Tooltip", "Draw kinematic cloth colliders with flat shading.")),
-	FVisualizationOption(&::Chaos::FClothVisualization::DrawKinematicColliderWired, TEXT("p.ChaosClothAssetEditor.DebugDrawKinematicColliderWired"), LOCTEXT("ChaosVisName_DrawKinematicColliderWired", "Draw Kinematic Colliders (Wireframe)"), LOCTEXT("ChaosVisName_DrawKinematicColliderWired_Tooltip", "Draw kinematic cloth colliders in wireframe.")),
-	FVisualizationOption(&::Chaos::FClothVisualization::DrawMultiResConstraint, TEXT("p.ChaosClothAssetEditor.DebugDrawMultiResConstraint"), LOCTEXT("ChaosVisName_DrawMultiResConstraint", "Draw Multi Res Constraint"), LOCTEXT("ChaosVisName_DrawMultiResConstraint_Tooltip", "Draw multi res constraint coarse mesh and targets.")),
-	FVisualizationOption(&GetSimulationStatisticsString, TEXT("p.ChaosClothAssetEditor.DebugDrawSimulationStatistics"), LOCTEXT("ChaosVisName_SimulationStatistics"    , "Simulation Statistics"), LOCTEXT("ChaosVisName_SimulationStatistics_Tooltip", "Displays simulation statistics")),
-	FVisualizationOption(&::Chaos::FClothVisualization::DrawWeightMap, TEXT("p.ChaosClothAssetEditor.DrawWeightMap"), LOCTEXT("ChaosVisName_DrawWeightMap", "Weight Map"), LOCTEXT("ChaosVisName_DrawWeightMap_ToolTip", "Draw the weight map for the simulation mesh. You can control the name of the map to be visualized by setting the p.ChaosClothVisualization.WeightMapName console variable."), /*bDisablesSimulation =*/false, /*bHidesClothSections=*/true),
-	FVisualizationOption(&::Chaos::FClothVisualization::DrawInpaintWeightsMatched, TEXT("p.ChaosClothAssetEditor.DrawInpaintWeightsMatched"), LOCTEXT("ChaosVisName_DrawInpaintWeightsMatched"        , "Transfer Skin Weights Node: Matched Vertices"), LOCTEXT("ChaosVisName_DrawInpaintWeightsMatched_ToolTip", "When transferring weights using the InpaintWeights method, will highlight the vertices for which we copied the weights directly from the source mesh. For all other vertices, the weights were computed automatically."), /*bDisablesSimulation =*/false, /*bHidesClothSections=*/true),
+	FVisualizationOption(
+		FClothVisualizationDebugDraw::CreateLambda([](const FClothEditorSimulationVisualization&, const ::Chaos::FClothVisualization& Visualization, FPrimitiveDrawInterface* PDI)
+		{
+			Visualization.DrawPhysMeshShaded(PDI);
+		}),
+		LOCTEXT("ChaosVisName_PhysMesh", "Physical Mesh (Flat Shaded)"),
+		LOCTEXT("ChaosVisName_PhysMeshShaded_ToolTip", "Draws the current physical result as a doubled sided flat shaded mesh"),
+		/*bDisablesSimulation =*/false, /*bHidesClothSections=*/true),
+	FVisualizationOption(
+		FClothVisualizationDebugDraw::CreateLambda([](const FClothEditorSimulationVisualization&, const ::Chaos::FClothVisualization& Visualization, FPrimitiveDrawInterface* PDI)
+		{
+			Visualization.DrawPhysMeshWired(PDI);
+		}), 
+		LOCTEXT("ChaosVisName_PhysMeshWire", "Physical Mesh (Wireframe)"), 
+		LOCTEXT("ChaosVisName_PhysMeshWired_ToolTip", "Draws the current physical mesh result in wireframe")),
+	FVisualizationOption(
+		FClothVisualizationDebugDraw::CreateLambda([](const FClothEditorSimulationVisualization&, const ::Chaos::FClothVisualization& Visualization, FPrimitiveDrawInterface* PDI)
+		{
+			Visualization.DrawAnimMeshWired(PDI);
+		}),
+		LOCTEXT("ChaosVisName_AnimMeshWire", "Animated Mesh (Wireframe)"), 
+		LOCTEXT("ChaosVisName_AnimMeshWired_ToolTip", "Draws the current animated mesh input in wireframe")),
+	FVisualizationOption(
+		LOCTEXT("ChaosVisName_HideRenderMesh", "Hide Render Mesh"), 
+		LOCTEXT("ChaosVisName_HideRenderMesh_ToolTip", "Hide the render mesh."), 
+		/*bDisablesSimulation =*/false, /*bHidesClothSections=*/true),
+	FVisualizationOption(
+		FClothVisualizationDebugDrawTexts::CreateLambda([](const FClothEditorSimulationVisualization&, const ::Chaos::FClothVisualization& Visualization, FCanvas* Canvas, const FSceneView* SceneView)
+		{
+			Visualization.DrawParticleIndices(Canvas, SceneView);
+		}),
+		LOCTEXT("ChaosVisName_ParticleIndices", "Particle Indices"), 
+		LOCTEXT("ChaosVisName_ParticleIndices_ToolTip", "Draws the particle indices as instantiated by the solver")),
+	FVisualizationOption(
+		FClothVisualizationDebugDrawTexts::CreateLambda([](const FClothEditorSimulationVisualization&, const ::Chaos::FClothVisualization& Visualization, FCanvas* Canvas, const FSceneView* SceneView)
+		{
+			Visualization.DrawElementIndices(Canvas, SceneView);
+		}),
+		LOCTEXT("ChaosVisName_ElementIndices", "Element Indices"), 
+		LOCTEXT("ChaosVisName_ElementIndices_ToolTip", "Draws the element's (triangle or other) indices as instantiated by the solver")),
+	FVisualizationOption(
+		FClothVisualizationDebugDraw::CreateLambda([](const FClothEditorSimulationVisualization&, const ::Chaos::FClothVisualization& Visualization, FPrimitiveDrawInterface* PDI)
+		{
+			Visualization.DrawPointNormals(PDI);
+		}),
+		LOCTEXT("ChaosVisName_PointNormals", "Physical Mesh Normals"), 
+		LOCTEXT("ChaosVisName_PointNormals_ToolTip", "Draws the current point normals for the simulation mesh")),
+	FVisualizationOption(
+		FClothVisualizationDebugDraw::CreateLambda([](const FClothEditorSimulationVisualization&, const ::Chaos::FClothVisualization& Visualization, FPrimitiveDrawInterface* PDI)
+		{
+			Visualization.DrawPointVelocities(PDI);
+		}),
+		LOCTEXT("ChaosVisName_PointVelocities", "Point Velocities"), 
+		LOCTEXT("ChaosVisName_PointVelocities_ToolTip", "Draws the current point velocities for the simulation mesh")),
+	FVisualizationOption(
+		FClothVisualizationDebugDraw::CreateLambda([](const FClothEditorSimulationVisualization&, const ::Chaos::FClothVisualization& Visualization, FPrimitiveDrawInterface* PDI)
+		{
+			Visualization.DrawAnimNormals(PDI);
+		}),
+		LOCTEXT("ChaosVisName_AnimNormals", "Animated Mesh Normals"), 
+		LOCTEXT("ChaosVisName_AnimNormals_ToolTip", "Draws the current point normals for the animated mesh")),
+	FVisualizationOption(
+		FClothVisualizationDebugDraw::CreateLambda([](const FClothEditorSimulationVisualization&, const ::Chaos::FClothVisualization& Visualization, FPrimitiveDrawInterface* PDI)
+		{
+			Visualization.DrawCollision(PDI);
+		}),
+		LOCTEXT("ChaosVisName_Collision", "Collisions"), 
+		LOCTEXT("ChaosVisName_Collision_ToolTip", "Draws the collision bodies the simulation is currently using")),
+	FVisualizationOption(
+		FClothVisualizationDebugDraw::CreateLambda([](const FClothEditorSimulationVisualization&, const ::Chaos::FClothVisualization& Visualization, FPrimitiveDrawInterface* PDI)
+		{
+			Visualization.DrawBackstops(PDI);
+		}),
+		LOCTEXT("ChaosVisName_Backstop", "Backstops"), 
+		LOCTEXT("ChaosVisName_Backstop_ToolTip", "Draws the backstop radius and position for each simulation particle")),
+	FVisualizationOption(
+		FClothVisualizationDebugDraw::CreateLambda([](const FClothEditorSimulationVisualization&, const ::Chaos::FClothVisualization& Visualization, FPrimitiveDrawInterface* PDI)
+		{
+			Visualization.DrawBackstopDistances(PDI);
+		}),
+		LOCTEXT("ChaosVisName_BackstopDistance", "Backstop Distances"), 
+		LOCTEXT("ChaosVisName_BackstopDistance_ToolTip", "Draws the backstop distance offset for each simulation particle"), 
+		/*bDisablesSimulation =*/true),
+	FVisualizationOption(
+		FClothVisualizationDebugDraw::CreateLambda([](const FClothEditorSimulationVisualization&, const ::Chaos::FClothVisualization& Visualization, FPrimitiveDrawInterface* PDI)
+		{
+			Visualization.DrawMaxDistances(PDI);
+		}),
+		LOCTEXT("ChaosVisName_MaxDistance", "Max Distances"), 
+		LOCTEXT("ChaosVisName_MaxDistance_ToolTip", "Draws the current max distances for the sim particles as a line along its normal"), 
+		/*bDisablesSimulation =*/true),
+	FVisualizationOption(
+		FClothVisualizationDebugDrawTexts::CreateLambda([](const FClothEditorSimulationVisualization&, const ::Chaos::FClothVisualization& Visualization, FCanvas* Canvas, const FSceneView* SceneView)
+		{
+			Visualization.DrawMaxDistanceValues(Canvas, SceneView);
+		}),
+		LOCTEXT("ChaosVisName_MaxDistanceValue", "Max Distances As Numbers"), 
+		LOCTEXT("ChaosVisName_MaxDistanceValue_ToolTip", "Draws the current max distances as numbers")),
+	FVisualizationOption(
+		FClothVisualizationDebugDraw::CreateLambda([](const FClothEditorSimulationVisualization&, const ::Chaos::FClothVisualization& Visualization, FPrimitiveDrawInterface* PDI)
+		{
+			Visualization.DrawAnimDrive(PDI);
+		}),
+		LOCTEXT("ChaosVisName_AnimDrive", "Anim Drive"), 
+		LOCTEXT("ChaosVisName_AnimDrive_Tooltip", "Draws the current skinned reference mesh for the simulation which anim drive will attempt to reach if enabled")),
+	FVisualizationOption(
+		FClothVisualizationDebugDraw::CreateLambda([](const FClothEditorSimulationVisualization&, const ::Chaos::FClothVisualization& Visualization, FPrimitiveDrawInterface* PDI)
+		{
+			Visualization.DrawEdgeConstraint(PDI);
+		}),
+		LOCTEXT("ChaosVisName_EdgeConstraint", "Edge Constraint"),
+		LOCTEXT("ChaosVisName_EdgeConstraint_Tooltip", "Draws the edge spring constraints")),
+	FVisualizationOption(
+		FClothVisualizationDebugDraw::CreateLambda([](const FClothEditorSimulationVisualization&, const ::Chaos::FClothVisualization& Visualization, FPrimitiveDrawInterface* PDI)
+		{
+			Visualization.DrawBendingConstraint(PDI);
+		}),
+		LOCTEXT("ChaosVisName_BendingConstraint", "Bending Constraint"), 
+		LOCTEXT("ChaosVisName_BendingConstraint_Tooltip", "Draws the bending spring constraints")),
+	FVisualizationOption(
+		FClothVisualizationDebugDraw::CreateLambda([](const FClothEditorSimulationVisualization&, const ::Chaos::FClothVisualization& Visualization, FPrimitiveDrawInterface* PDI)
+		{
+			Visualization.DrawLongRangeConstraint(PDI);
+		}),
+		LOCTEXT("ChaosVisName_LongRangeConstraint" , "Long Range Constraint"), 
+		LOCTEXT("ChaosVisName_LongRangeConstraint_Tooltip", "Draws the long range attachment constraint distances")),
+	FVisualizationOption(
+		FClothVisualizationDebugDraw::CreateLambda([](const FClothEditorSimulationVisualization&, const ::Chaos::FClothVisualization& Visualization, FPrimitiveDrawInterface* PDI)
+		{
+			Visualization.DrawWindAndPressureForces(PDI);
+		}),
+		LOCTEXT("ChaosVisName_WindAndPressureForces", "Wind Aerodynamic And Pressure Forces"), 
+		LOCTEXT("ChaosVisName_WindAndPressure_Tooltip", "Draws the Wind drag and lift and pressure forces")),
+	FVisualizationOption(
+		FClothVisualizationDebugDraw::CreateLambda([](const FClothEditorSimulationVisualization&, const ::Chaos::FClothVisualization& Visualization, FPrimitiveDrawInterface* PDI)
+		{
+			Visualization.DrawLocalSpace(PDI);
+		}),
+		LOCTEXT("ChaosVisName_LocalSpace", "Local Space Reference Bone"), 
+		LOCTEXT("ChaosVisName_LocalSpace_Tooltip", "Draws the local space reference bone")),
+	FVisualizationOption(
+		FClothVisualizationDebugDraw::CreateLambda([](const FClothEditorSimulationVisualization&, const ::Chaos::FClothVisualization& Visualization, FPrimitiveDrawInterface* PDI)
+		{
+			Visualization.DrawGravity(PDI);
+		}),
+		LOCTEXT("ChaosVisName_Gravity", "Gravity"), 
+		LOCTEXT("ChaosVisName_Gravity_Tooltip", "Draws gravity")),
+	FVisualizationOption(
+		FClothVisualizationDebugDraw::CreateLambda([](const FClothEditorSimulationVisualization&, const ::Chaos::FClothVisualization& Visualization, FPrimitiveDrawInterface* PDI)
+		{
+			Visualization.DrawFictitiousAngularForces(PDI);
+		}),
+		LOCTEXT("ChaosVisName_FictitiousAngularForces", "Fictitious Angular Forces"),
+		LOCTEXT("ChaosVisName_Gravity_FictitiousAngularForces", "Draws fictitious angular forces (force based solver only)")),
+	FVisualizationOption(
+		FClothVisualizationDebugDraw::CreateLambda([](const FClothEditorSimulationVisualization&, const ::Chaos::FClothVisualization& Visualization, FPrimitiveDrawInterface* PDI)
+		{
+			Visualization.DrawSelfCollision(PDI);
+		}),
+		LOCTEXT("ChaosVisName_SelfCollision", "Self Collision"), 
+		LOCTEXT("ChaosVisName_SelfCollision_Tooltip", "Draws the self collision thickness/debugging information")),
+	FVisualizationOption(
+		FClothVisualizationDebugDraw::CreateLambda([](const FClothEditorSimulationVisualization&, const ::Chaos::FClothVisualization& Visualization, FPrimitiveDrawInterface* PDI)
+		{
+			Visualization.DrawSelfIntersection(PDI);
+		}),
+		LOCTEXT("ChaosVisName_SelfIntersection", "Self Intersection"), 
+		LOCTEXT("ChaosVisName_SelfIntersection_Tooltip", "Draws the self intersection contour/region information")),
+	FVisualizationOption(
+		FClothVisualizationDebugDraw::CreateLambda([](const FClothEditorSimulationVisualization&, const ::Chaos::FClothVisualization& Visualization, FPrimitiveDrawInterface* PDI)
+		{
+			Visualization.DrawSelfCollisionLayers(PDI);
+		}),
+		LOCTEXT("ChaosVisName_SelfCollisionLayers", "Self Collision Layers"), 
+		LOCTEXT("ChaosVisName_SelfCollisionLayers_Tooltip", "Draws the self collision layers"),
+		/*bDisablesSimulation =*/false, /*bHidesClothSections=*/true),
+	FVisualizationOption(
+		FClothVisualizationDebugDraw::CreateLambda([](const FClothEditorSimulationVisualization&, const ::Chaos::FClothVisualization& Visualization, FPrimitiveDrawInterface* PDI)
+		{
+			Visualization.DrawSelfCollisionThickness(PDI);
+		}),
+		LOCTEXT("ChaosVisName_SelfCollisionThickness", "Self Collision Thickness"), 
+		LOCTEXT("ChaosVisName_SelfCollisionThickness_Tooltip", "Draws the self collision Thickness")),
+	FVisualizationOption(
+		FClothVisualizationDebugDraw::CreateLambda([](const FClothEditorSimulationVisualization&, const ::Chaos::FClothVisualization& Visualization, FPrimitiveDrawInterface* PDI)
+		{
+			Visualization.DrawKinematicColliderShaded(PDI);
+		}),
+		LOCTEXT("ChaosVisName_DrawKinematicColliderShaded", "Draw Kinematic Colliders (Shaded)"), 
+		LOCTEXT("ChaosVisName_DrawKinematicColliderShaded_Tooltip", "Draw kinematic cloth colliders with flat shading.")),
+	FVisualizationOption(
+		FClothVisualizationDebugDraw::CreateLambda([](const FClothEditorSimulationVisualization&, const ::Chaos::FClothVisualization& Visualization, FPrimitiveDrawInterface* PDI)
+		{
+			Visualization.DrawKinematicColliderWired(PDI);
+		}),
+		LOCTEXT("ChaosVisName_DrawKinematicColliderWired", "Draw Kinematic Colliders (Wireframe)"), 
+		LOCTEXT("ChaosVisName_DrawKinematicColliderWired_Tooltip", "Draw kinematic cloth colliders in wireframe.")),
+	FVisualizationOption(
+		FClothVisualizationDebugDraw::CreateLambda([](const FClothEditorSimulationVisualization&, const ::Chaos::FClothVisualization& Visualization, FPrimitiveDrawInterface* PDI)
+		{
+			Visualization.DrawMultiResConstraint(PDI);
+		}),
+		LOCTEXT("ChaosVisName_DrawMultiResConstraint", "Draw Multi Res Constraint"), 
+		LOCTEXT("ChaosVisName_DrawMultiResConstraint_Tooltip", "Draw multi res constraint coarse mesh and targets.")),
+	FVisualizationOption(
+		FLocalDebugDisplayString::CreateLambda([](const FClothEditorSimulationVisualization&,const FClothSimulationProxy& Proxy)
+		{
+			return GetSimulationStatisticsString(Proxy);
+		}),
+		LOCTEXT("ChaosVisName_SimulationStatistics", "Simulation Statistics"), 
+		LOCTEXT("ChaosVisName_SimulationStatistics_Tooltip", "Displays simulation statistics")),
+	FVisualizationOption(
+		FClothVisualizationDebugDraw::CreateLambda([](const FClothEditorSimulationVisualization& EditorVisualization, const ::Chaos::FClothVisualization& Visualization, FPrimitiveDrawInterface* PDI)
+		{
+			const FString* const WeightMap = EditorVisualization.GetCurrentlySelectedWeightMap();
+			Visualization.DrawWeightMapWithName(PDI, WeightMap ? *WeightMap : FString());
+		}),
+		LOCTEXT("ChaosVisName_DrawWeightMap", "Weight Map"), 
+		LOCTEXT("ChaosVisName_DrawWeightMap_ToolTip", "Draw the weight map for the simulation mesh. You can control the name of the map to be visualized by setting the p.ChaosClothVisualization.WeightMapName console variable."), 
+		/*bDisablesSimulation =*/false, /*bHidesClothSections=*/true,
+		FAdditionalMenus::CreateLambda([](FClothEditorSimulationVisualization& EditorVisualization, FMenuBuilder& MenuBuilder, TSharedRef<FChaosClothAssetEditor3DViewportClient> ViewportClient)
+		{
+			EditorVisualization.ExtendViewportShowMenuWeightMapSelector(MenuBuilder, ViewportClient);
+		})),
+	FVisualizationOption(
+		FClothVisualizationDebugDraw::CreateLambda([](const FClothEditorSimulationVisualization&, const ::Chaos::FClothVisualization& Visualization, FPrimitiveDrawInterface* PDI)
+		{
+			Visualization.DrawInpaintWeightsMatched(PDI);
+		}),
+		LOCTEXT("ChaosVisName_DrawInpaintWeightsMatched", "Transfer Skin Weights Node: Matched Vertices"), 
+		LOCTEXT("ChaosVisName_DrawInpaintWeightsMatched_ToolTip", "When transferring weights using the InpaintWeights method, will highlight the vertices for which we copied the weights directly from the source mesh. For all other vertices, the weights were computed automatically."), 
+		/*bDisablesSimulation =*/false, /*bHidesClothSections=*/true),
 };
 const uint32 FVisualizationOption::Count = sizeof(OptionData) / sizeof(FVisualizationOption);
 
@@ -227,11 +436,65 @@ void FClothEditorSimulationVisualization::ExtendViewportShowMenu(FMenuBuilder& M
 			const FUIAction Action(ExecuteAction, FCanExecuteAction(), IsActionChecked);
 
 			MenuBuilder.AddMenuEntry(Private::FVisualizationOption::OptionData[OptionIndex].DisplayName, Private::FVisualizationOption::OptionData[OptionIndex].ToolTip, FSlateIcon(), Action, NAME_None, EUserInterfaceActionType::ToggleButton);
+			Private::FVisualizationOption::OptionData[OptionIndex].AdditionalMenus.ExecuteIfBound(*this, MenuBuilder, ViewportClient);
 		}
 	}
 	MenuBuilder.EndSection();
 }
 
+void FClothEditorSimulationVisualization::ExtendViewportShowMenuWeightMapSelector(FMenuBuilder& MenuBuilder, TSharedRef<FChaosClothAssetEditor3DViewportClient> ViewportClient)
+{
+	WeightMapSelector =
+		SNew(STextComboBox)
+		.OptionsSource(&WeightMapNames)
+		.OnSelectionChanged(STextComboBox::FOnTextSelectionChanged::CreateRaw(this, &FClothEditorSimulationVisualization::WeightMapSelectionChanged));
+
+	MenuBuilder.AddMenuEntry(FUIAction(), WeightMapSelector.ToSharedRef(),
+		NAME_None, LOCTEXT("WeightMapNameSelection", "Select weight map to draw."), EUserInterfaceActionType::None);
+}
+
+void FClothEditorSimulationVisualization::WeightMapSelectionChanged(TSharedPtr<FString> Selection, ESelectInfo::Type SelectInfo)
+{
+	CurrentlySelectedWeightMap = Selection;
+}
+
+void FClothEditorSimulationVisualization::RefreshMenusForClothComponent(const UChaosClothComponent* ClothComponent)
+{
+	TArray<FString> WeightMaps;
+	if (const ::Chaos::FClothVisualization* const Visualization = ClothComponent && ClothComponent->GetClothSimulationProxy() ? ClothComponent->GetClothSimulationProxy()->GetVisualization() : nullptr)
+	{
+		WeightMaps = Visualization->GetAllWeightMapNames();
+	}
+
+	bool bFoundCurrentlySelected = false;
+	WeightMapNames.Empty(WeightMaps.Num());
+	for (const FString& Map : WeightMaps)
+	{
+		const int32 Index = WeightMapNames.Emplace(MakeShared<FString>(Map));
+		if (!bFoundCurrentlySelected && CurrentlySelectedWeightMap && Map == *CurrentlySelectedWeightMap)
+		{
+			CurrentlySelectedWeightMap = WeightMapNames[Index];
+			bFoundCurrentlySelected = true;
+		}
+	}
+	if (!bFoundCurrentlySelected)
+	{
+		if (WeightMapNames.IsEmpty())
+		{
+			CurrentlySelectedWeightMap.Reset();
+		}
+		else
+		{
+			CurrentlySelectedWeightMap = WeightMapNames[0];
+		}
+	}
+	if (WeightMapSelector)
+	{
+		WeightMapSelector->RefreshOptions();
+		WeightMapSelector->SetSelectedItem(CurrentlySelectedWeightMap);
+	}
+}
+
 void FClothEditorSimulationVisualization::DebugDrawSimulation(const UChaosClothComponent* ClothComponent, FPrimitiveDrawInterface* PDI)
 {
 	const ::Chaos::FClothVisualization* const Visualization = ClothComponent && ClothComponent->GetClothSimulationProxy() ? ClothComponent->GetClothSimulationProxy()->GetVisualization() : nullptr;
@@ -242,10 +505,9 @@ void FClothEditorSimulationVisualization::DebugDrawSimulation(const UChaosClothC
 
 	for (int32 OptionIndex = 0; OptionIndex < Private::FVisualizationOption::Count; ++OptionIndex)
 	{
-		if (Private::FVisualizationOption::OptionData[OptionIndex].ClothVisualizationDebugDrawFunction &&
-			(Private::FVisualizationOption::OptionData[OptionIndex].IsConsoleVariableEnabled() || Flags[OptionIndex]))
+		if (Flags[OptionIndex])
 		{
-			(Visualization->*(Private::FVisualizationOption::OptionData[OptionIndex].ClothVisualizationDebugDrawFunction))(PDI);
+			Private::FVisualizationOption::OptionData[OptionIndex].ClothVisualizationDebugDraw.ExecuteIfBound(*this, *Visualization, PDI);
 		}
 	}
 }
@@ -259,10 +521,9 @@ void FClothEditorSimulationVisualization::DebugDrawSimulationTexts(const UChaosC
 	}
 	for (int32 OptionIndex = 0; OptionIndex < Private::FVisualizationOption::Count; ++OptionIndex)
 	{
-		if (Private::FVisualizationOption::OptionData[OptionIndex].ClothVisualizationDebugDrawTextsFunction &&
-			(Private::FVisualizationOption::OptionData[OptionIndex].IsConsoleVariableEnabled() || Flags[OptionIndex]))
+		if (Flags[OptionIndex])
 		{
-			(Visualization->*(Private::FVisualizationOption::OptionData[OptionIndex].ClothVisualizationDebugDrawTextsFunction))(Canvas, SceneView);
+			Private::FVisualizationOption::OptionData[OptionIndex].ClothVisualizationDebugDrawTexts.ExecuteIfBound(*this, *Visualization, Canvas, SceneView);
 		}
 	}
 }
@@ -278,10 +539,9 @@ FText FClothEditorSimulationVisualization::GetDisplayString(const UChaosClothCom
 	FText DisplayString;
 	for (int32 OptionIndex = 0; OptionIndex < Private::FVisualizationOption::Count; ++OptionIndex)
 	{
-		if (Private::FVisualizationOption::OptionData[OptionIndex].LocalDebugDisplayStringFunction &&
-			(Private::FVisualizationOption::OptionData[OptionIndex].IsConsoleVariableEnabled() || Flags[OptionIndex]))
+		if (Private::FVisualizationOption::OptionData[OptionIndex].LocalDebugDisplayString.IsBound() && Flags[OptionIndex])
 		{
-			DisplayString = Private::ConcatenateLine(DisplayString, (Private::FVisualizationOption::OptionData[OptionIndex].LocalDebugDisplayStringFunction)(SimProxy));
+			DisplayString = Private::ConcatenateLine(DisplayString, Private::FVisualizationOption::OptionData[OptionIndex].LocalDebugDisplayString.Execute(*this, *SimProxy));
 		}
 	}
 	return DisplayString;
diff --git a/Engine/Plugins/ChaosClothAssetEditor/Source/ChaosClothAssetEditor/Private/ChaosClothAsset/ClothEditorSimulationVisualization.h b/Engine/Plugins/ChaosClothAssetEditor/Source/ChaosClothAssetEditor/Private/ChaosClothAsset/ClothEditorSimulationVisualization.h
--- a/Engine/Plugins/ChaosClothAssetEditor/Source/ChaosClothAssetEditor/Private/ChaosClothAsset/ClothEditorSimulationVisualization.h
+++ b/Engine/Plugins/ChaosClothAssetEditor/Source/ChaosClothAssetEditor/Private/ChaosClothAsset/ClothEditorSimulationVisualization.h
@@ -9,6 +9,12 @@ class UChaosClothComponent;
 class FPrimitiveDrawInterface;
 class FCanvas;
 class FSceneView;
+class STextComboBox;
+
+namespace ESelectInfo
+{
+enum Type : int;
+}
 
 namespace UE::Chaos::ClothAsset
 {
@@ -24,15 +30,25 @@ public:
 	void DebugDrawSimulation(const UChaosClothComponent* ClothComponent, FPrimitiveDrawInterface* PDI);
 	void DebugDrawSimulationTexts(const UChaosClothComponent* ClothComponent, FCanvas* Canvas, const FSceneView* SceneView);
 	FText GetDisplayString(const UChaosClothComponent* ClothComponent) const;
+	void RefreshMenusForClothComponent(const UChaosClothComponent* ClothComponent);
 
+	// WeightMaps 
+	const FString* GetCurrentlySelectedWeightMap() const { return CurrentlySelectedWeightMap.Get(); }
+	void ExtendViewportShowMenuWeightMapSelector(FMenuBuilder& MenuBuilder, TSharedRef<FChaosClothAssetEditor3DViewportClient> ViewportClient);
 private:
 	/** Return whether or not - given the current enabled options - the simulation should be disabled. */
 	bool ShouldDisableSimulation() const;
 	/** Show/hide all cloth sections for the specified mesh compoment. */
 	void ShowClothSections(UChaosClothComponent* ClothComponent, bool bIsClothSectionsVisible) const;
+	/** Callback for weight map selection*/
+	void WeightMapSelectionChanged(TSharedPtr<FString> Selection, ESelectInfo::Type SelectInfo);
 
 private:
 	/** Flags used to store the checked status for the visualization options. */
 	TBitArray<> Flags;
+	TSharedPtr<STextComboBox> WeightMapSelector;
+	TArray<TSharedPtr<FString>> WeightMapNames;
+	TSharedPtr<FString> CurrentlySelectedWeightMap;
+	
 };
 } // namespace UE::Chaos::ClothAsset
