Resubmitting with manual merge fix.

Chaos Cloth Asset - Fixed typo that is preventing correct setting of the BucklingRatio property.

#tests Tested with different BucklingRatio values before and after the fix..

diff --git a/Engine/Plugins/ChaosClothAssetEditor/Source/ChaosClothAssetDataflowNodes/Private/ChaosClothAsset/SimulationBendingConfigNode.cpp b/Engine/Plugins/ChaosClothAssetEditor/Source/ChaosClothAssetDataflowNodes/Private/ChaosClothAsset/SimulationBendingConfigNode.cpp
--- a/Engine/Plugins/ChaosClothAssetEditor/Source/ChaosClothAssetDataflowNodes/Private/ChaosClothAsset/SimulationBendingConfigNode.cpp
+++ b/Engine/Plugins/ChaosClothAssetEditor/Source/ChaosClothAssetDataflowNodes/Private/ChaosClothAsset/SimulationBendingConfigNode.cpp
@@ -61,7 +61,7 @@ void FChaosClothAssetSimulationBendingConfigNode::AddProperties(FPropertyHelper&
 				
 				PropertyHelper.SetPropertyWeighted(FName(TEXT("XPBDBucklingRatio")), BucklingRatioWeighted, {
 					FName(TEXT("XPBDAnisoBucklingRatio")),
-					FName(TEXT("BucklingRatios"))});
+					FName(TEXT("BucklingRatio"))});
 			}
 			PropertyHelper.SetPropertyEnum(FName(TEXT("XPBDRestAngleType")), RestAngleType, {
 				FName(TEXT("XPBDAnisoRestAngleType")),
