Cloth Editor: roll back previous change disabling fly mode camera in both viewports. Add a Click-Drag behavior that will prevent camera movement when left-mouse dragging but allow everything else

diff --git a/Engine/Plugins/ChaosClothAssetEditor/Source/ChaosClothAssetEditor/Private/ChaosClothAsset/ClothEditor3DViewportClient.cpp b/Engine/Plugins/ChaosClothAssetEditor/Source/ChaosClothAssetEditor/Private/ChaosClothAsset/ClothEditor3DViewportClient.cpp
--- a/Engine/Plugins/ChaosClothAssetEditor/Source/ChaosClothAssetEditor/Private/ChaosClothAsset/ClothEditor3DViewportClient.cpp
+++ b/Engine/Plugins/ChaosClothAssetEditor/Source/ChaosClothAssetEditor/Private/ChaosClothAsset/ClothEditor3DViewportClient.cpp
@@ -20,6 +20,8 @@
 #include "AssetViewerSettings.h"
 #include "Editor/EditorPerProjectUserSettings.h"
 #include "Transforms/TransformGizmoDataBinder.h"
+#include "BaseBehaviors/ClickDragBehavior.h"
+#include "BaseBehaviors/SingleClickOrDragBehavior.h"
 
 namespace UE::Chaos::ClothAsset
 {
@@ -63,8 +65,19 @@ FChaosClothAssetEditor3DViewportClient::FChaosClothAssetEditor3DViewportClient(F
 	// Set correct flags according to current profile settings
 	SetAdvancedShowFlagsForScene(UAssetViewerSettings::Get()->Profiles[GetMutableDefault<UEditorPerProjectUserSettings>()->AssetViewerProfileIndex].bPostProcessingEnabled);
 
-	// Disable the non-alt-key camera controls
-	bLockFlightCamera = true;
+	//
+	// Input behaviors
+	//
+
+	InputBehaviorSet = NewObject<UInputBehaviorSet>();
+
+	// Our ClickOrDrag behavior is used to intercept non-alt left-mouse-button drag inputs, but still allow single-click for select/deselect operation
+	const TObjectPtr<USingleClickOrDragInputBehavior> ClickOrDragBehavior = NewObject<USingleClickOrDragInputBehavior>();
+	ClickOrDragBehavior->Initialize(this, this);
+	
+	InputBehaviorSet->Add(ClickOrDragBehavior);
+
+	InteractiveToolsContext->InputRouter->RegisterSource(this);
 }
 
 void FChaosClothAssetEditor3DViewportClient::RegisterDelegates()
@@ -111,6 +124,7 @@ void FChaosClothAssetEditor3DViewportClient::AddReferencedObjects(FReferenceColl
 
 	Collector.AddReferencedObject(TransformProxy);
 	Collector.AddReferencedObject(Gizmo);
+	Collector.AddReferencedObject(InputBehaviorSet);
 }
 
 void FChaosClothAssetEditor3DViewportClient::Tick(float DeltaSeconds)
@@ -302,6 +316,38 @@ void FChaosClothAssetEditor3DViewportClient::DrawCanvas(FViewport& InViewport, F
 	}
 }
 
+const UInputBehaviorSet* FChaosClothAssetEditor3DViewportClient::GetInputBehaviors() const
+{
+	return InputBehaviorSet;
+}
+
+
+// IClickBehaviorTarget
+FInputRayHit FChaosClothAssetEditor3DViewportClient::IsHitByClick(const FInputDeviceRay& ClickPos)
+{
+	// Here we are responding that we do want to handle click events, but we are only doing this so that we can also get drag events
+	// TODO: Find out if there's a way we can just intercept mouse drag events and not single-click events
+	return FInputRayHit(TNumericLimits<float>::Max());
+}
+
+void FChaosClothAssetEditor3DViewportClient::OnClicked(const FInputDeviceRay& ClickPos)
+{
+	// On a single click with no drag, respond as we would in ProcessClick()
+	if (ClickPos.bHas2D)
+	{
+		HHitProxy* const HitProxy = Viewport->GetHitProxy(ClickPos.ScreenPosition[0], ClickPos.ScreenPosition[1]);
+		UpdateSelection(HitProxy);
+	}
+}
+
+// IClickDragBehaviorTarget
+FInputRayHit FChaosClothAssetEditor3DViewportClient::CanBeginClickDragSequence(const FInputDeviceRay& PressPos)
+{
+	// We do want to handle drag events
+	return FInputRayHit(TNumericLimits<float>::Max());
+}
+
+
 FBox FChaosClothAssetEditor3DViewportClient::PreviewBoundingBox() const
 {
 	if (ClothEdMode)
@@ -340,10 +386,9 @@ const UChaosClothComponent* FChaosClothAssetEditor3DViewportClient::GetPreviewCl
 	return nullptr;
 }
 
-void FChaosClothAssetEditor3DViewportClient::ProcessClick(FSceneView& View, HHitProxy* HitProxy, FKey Key, EInputEvent Event, uint32 HitX, uint32 HitY)
-{
-	FEditorViewportClient::ProcessClick(View, HitProxy, Key, Event, HitX, HitY);
 
+void FChaosClothAssetEditor3DViewportClient::UpdateSelection(HHitProxy* HitProxy)
+{
 	const bool bIsShiftKeyDown = Viewport->KeyState(EKeys::LeftShift) || Viewport->KeyState(EKeys::RightShift);
 	const bool bIsCtrlKeyDown = Viewport->KeyState(EKeys::LeftControl) || Viewport->KeyState(EKeys::RightControl);
 
@@ -380,6 +425,12 @@ void FChaosClothAssetEditor3DViewportClient::ProcessClick(FSceneView& View, HHit
 	}
 }
 
+void FChaosClothAssetEditor3DViewportClient::ProcessClick(FSceneView& View, HHitProxy* HitProxy, FKey Key, EInputEvent Event, uint32 HitX, uint32 HitY)
+{
+	FEditorViewportClient::ProcessClick(View, HitProxy, Key, Event, HitX, HitY);
+	UpdateSelection(HitProxy);
+}
+
 
 void FChaosClothAssetEditor3DViewportClient::ClearSelectedComponents()
 {
diff --git a/Engine/Plugins/ChaosClothAssetEditor/Source/ChaosClothAssetEditor/Private/ChaosClothAsset/ClothEditor3DViewportClient.h b/Engine/Plugins/ChaosClothAssetEditor/Source/ChaosClothAssetEditor/Private/ChaosClothAsset/ClothEditor3DViewportClient.h
--- a/Engine/Plugins/ChaosClothAssetEditor/Source/ChaosClothAssetEditor/Private/ChaosClothAsset/ClothEditor3DViewportClient.h
+++ b/Engine/Plugins/ChaosClothAssetEditor/Source/ChaosClothAssetEditor/Private/ChaosClothAsset/ClothEditor3DViewportClient.h
@@ -5,6 +5,8 @@
 #include "EditorViewportClient.h"
 #include "Dataflow/DataflowNodeParameters.h"
 #include "GeometryCollection/ManagedArrayCollection.h"
+#include "InputBehaviorSet.h"
+#include "BaseBehaviors/BehaviorTargetInterfaces.h"
 
 class UChaosClothComponent;
 class UChaosClothAssetEditorMode;
@@ -23,7 +25,11 @@ class FClothEditorSimulationVisualization;
  * Viewport client for the 3d sim preview in the cloth editor. Currently same as editor viewport
  * client but doesn't allow editor gizmos/widgets.
  */
-class CHAOSCLOTHASSETEDITOR_API FChaosClothAssetEditor3DViewportClient : public FEditorViewportClient, public TSharedFromThis<FChaosClothAssetEditor3DViewportClient>
+class CHAOSCLOTHASSETEDITOR_API FChaosClothAssetEditor3DViewportClient : public FEditorViewportClient, 
+	public TSharedFromThis<FChaosClothAssetEditor3DViewportClient>, 
+	public IClickBehaviorTarget,
+	public IClickDragBehaviorTarget,
+	public IInputBehaviorSource
 {
 public:
 
@@ -89,11 +95,28 @@ private:
 	virtual void Draw(const FSceneView* View, FPrimitiveDrawInterface* PDI) override;
 	virtual void DrawCanvas(FViewport& InViewport, FSceneView& View, FCanvas& Canvas) override;
 
+	// IClickBehaviorTarget
+	virtual FInputRayHit IsHitByClick(const FInputDeviceRay& ClickPos) override;
+	virtual void OnClicked(const FInputDeviceRay& ClickPos) override;
+
+	// IClickDragBehaviorTarget
+	virtual FInputRayHit CanBeginClickDragSequence(const FInputDeviceRay& PressPos) override;
+	virtual void OnClickPress(const FInputDeviceRay& PressPos) override {}
+	virtual void OnClickDrag(const FInputDeviceRay& DragPos) override {}
+	virtual void OnClickRelease(const FInputDeviceRay& ReleasePos) override {}
+	virtual void OnTerminateDragSequence() override {}
+
+	// IInputBehaviorSource
+	virtual const UInputBehaviorSet* GetInputBehaviors() const override;
+
 	void OnAssetViewerSettingsChanged(const FName& InPropertyName);
 	void SetAdvancedShowFlagsForScene(const bool bAdvancedShowFlags);
 
 	void ComponentSelectionChanged(UObject* NewSelection);
 
+	// Update the selected components based on hitproxy
+	void UpdateSelection(HHitProxy* HitProxy);
+
 	TWeakPtr<FChaosClothPreviewScene> ClothPreviewScene;
 
 	TObjectPtr<UChaosClothAssetEditorMode> ClothEdMode;
@@ -113,5 +136,6 @@ private:
 	TObjectPtr<UCombinedTransformGizmo> Gizmo = nullptr;
 	TSharedPtr<FTransformGizmoDataBinder> DataBinder = nullptr;
 
+	TObjectPtr<UInputBehaviorSet> InputBehaviorSet;
 };
 } // namespace UE::Chaos::ClothAsset
diff --git a/Engine/Plugins/ChaosClothAssetEditor/Source/ChaosClothAssetEditor/Private/ChaosClothAsset/ClothEditorRestSpaceViewportClient.cpp b/Engine/Plugins/ChaosClothAssetEditor/Source/ChaosClothAssetEditor/Private/ChaosClothAsset/ClothEditorRestSpaceViewportClient.cpp
--- a/Engine/Plugins/ChaosClothAssetEditor/Source/ChaosClothAssetEditor/Private/ChaosClothAsset/ClothEditorRestSpaceViewportClient.cpp
+++ b/Engine/Plugins/ChaosClothAssetEditor/Source/ChaosClothAssetEditor/Private/ChaosClothAsset/ClothEditorRestSpaceViewportClient.cpp
@@ -38,7 +38,6 @@ FChaosClothEditorRestSpaceViewportClient::FChaosClothEditorRestSpaceViewportClie
 	ScrollBehavior->Initialize(ScrollBehaviorTarget.Get());
 	ScrollBehavior->SetDefaultPriority(ViewportBehaviorPriority);
 	ScrollBehavior->SetUseRightMouseButton();
-	BehaviorSet->Add(ScrollBehavior);
 	BehaviorsFor2DMode.Add(ScrollBehavior);
 
 	ZoomBehaviorTarget = MakeUnique<FEditor2DMouseWheelZoomBehaviorTarget>(this);
@@ -50,9 +49,26 @@ FChaosClothEditorRestSpaceViewportClient::FChaosClothEditorRestSpaceViewportClie
 	UMouseWheelInputBehavior* const ZoomBehavior = NewObject<UMouseWheelInputBehavior>();
 	ZoomBehavior->Initialize(ZoomBehaviorTarget.Get());
 	ZoomBehavior->SetDefaultPriority(ViewportBehaviorPriority);
-	BehaviorSet->Add(ZoomBehavior);
 	BehaviorsFor2DMode.Add(ZoomBehavior);
 
+
+	const TObjectPtr<ULocalClickDragInputBehavior> ClickDrag3DBehavior = NewObject<ULocalClickDragInputBehavior>();
+	ClickDrag3DBehavior->Initialize();
+	ClickDrag3DBehavior->SetDefaultPriority(ViewportBehaviorPriority);
+
+	ClickDrag3DBehavior->ModifierCheckFunc = [this](const FInputDeviceState& InputState) {
+		return !FInputDeviceState::IsAltKeyDown(InputState);
+	};
+
+	ClickDrag3DBehavior->CanBeginClickDragFunc = [](const FInputDeviceRay& InputDeviceRay)
+	{
+		return FInputRayHit(TNumericLimits<float>::Max()); // bHit is true. Depth is max to lose the standard tiebreaker.
+	};
+
+	BehaviorsFor3DMode.Add(ClickDrag3DBehavior);
+
+	UpdateBehaviorsForCurrentViewMode();
+
 	EngineShowFlags.SetSelectionOutline(false);
 	ModeTools->GetInteractiveToolsContext()->InputRouter->RegisterSource(this);
 
@@ -62,9 +78,6 @@ FChaosClothEditorRestSpaceViewportClient::FChaosClothEditorRestSpaceViewportClie
 	CameraPointLight->SetIntensity(3.0f);
 	CameraPointLight->SetCastShadows(false);
 	PreviewScene->AddComponent(CameraPointLight, FTransform());
-
-	// Disable the non-alt-key camera controls
-	bLockFlightCamera = true;
 }
 
 void FChaosClothEditorRestSpaceViewportClient::SetConstructionViewMode(EClothPatternVertexType InViewMode)
@@ -77,15 +90,8 @@ void FChaosClothEditorRestSpaceViewportClient::SetConstructionViewMode(EClothPat
 
 	ConstructionViewMode = InViewMode;
 
-	BehaviorSet->RemoveAll();
-
 	if (ConstructionViewMode == EClothPatternVertexType::Sim2D)
 	{
-		for (UInputBehavior* const Behavior : BehaviorsFor2DMode)
-		{
-			BehaviorSet->Add(Behavior);
-		}
-
 		const double AbsZ = FMath::Abs(ViewTransformPerspective.GetLocation().Z);
 		constexpr double CameraFarPlaneWorldZ = -10.0;
 		constexpr double CameraNearPlaneProportionZ = 0.8;
@@ -98,10 +104,31 @@ void FChaosClothEditorRestSpaceViewportClient::SetConstructionViewMode(EClothPat
 		OverrideNearClipPlane(UE_KINDA_SMALL_NUMBER);
 	}
 
+	UpdateBehaviorsForCurrentViewMode();
+
 	ModeTools->GetInteractiveToolsContext()->InputRouter->DeregisterSource(this);
 	ModeTools->GetInteractiveToolsContext()->InputRouter->RegisterSource(this);
 }
 
+void FChaosClothEditorRestSpaceViewportClient::UpdateBehaviorsForCurrentViewMode()
+{
+	BehaviorSet->RemoveAll();
+
+	if (ConstructionViewMode == EClothPatternVertexType::Sim2D)
+	{
+		for (UInputBehavior* const Behavior : BehaviorsFor2DMode)
+		{
+			BehaviorSet->Add(Behavior);
+		}
+	}
+	else
+	{
+		for (UInputBehavior* const Behavior : BehaviorsFor3DMode)
+		{
+			BehaviorSet->Add(Behavior);
+		}
+	}
+}
 
 EClothPatternVertexType FChaosClothEditorRestSpaceViewportClient::GetConstructionViewMode() const
 {
@@ -119,6 +146,7 @@ void FChaosClothEditorRestSpaceViewportClient::AddReferencedObjects(FReferenceCo
 	FEditorViewportClient::AddReferencedObjects(Collector);
 	Collector.AddReferencedObject(BehaviorSet);
 	Collector.AddReferencedObjects(BehaviorsFor2DMode);
+	Collector.AddReferencedObjects(BehaviorsFor3DMode);
 }
 
 bool FChaosClothEditorRestSpaceViewportClient::ShouldOrbitCamera() const
diff --git a/Engine/Plugins/ChaosClothAssetEditor/Source/ChaosClothAssetEditor/Private/ChaosClothAsset/ClothEditorRestSpaceViewportClient.h b/Engine/Plugins/ChaosClothAssetEditor/Source/ChaosClothAssetEditor/Private/ChaosClothAsset/ClothEditorRestSpaceViewportClient.h
--- a/Engine/Plugins/ChaosClothAssetEditor/Source/ChaosClothAssetEditor/Private/ChaosClothAsset/ClothEditorRestSpaceViewportClient.h
+++ b/Engine/Plugins/ChaosClothAssetEditor/Source/ChaosClothAssetEditor/Private/ChaosClothAsset/ClothEditorRestSpaceViewportClient.h
@@ -46,6 +46,8 @@ private:
 
 	virtual void Tick(float DeltaSeconds) override;
 
+	void UpdateBehaviorsForCurrentViewMode();
+
 	TObjectPtr<UPointLightComponent> CameraPointLight;
 
 	EClothPatternVertexType ConstructionViewMode = EClothPatternVertexType::Sim2D;
@@ -53,6 +55,7 @@ private:
 	TObjectPtr<UInputBehaviorSet> BehaviorSet;
 
 	TArray<TObjectPtr<UInputBehavior>> BehaviorsFor2DMode;
+	TArray<TObjectPtr<UInputBehavior>> BehaviorsFor3DMode;
 
 	TUniquePtr<FEditor2DScrollBehaviorTarget> ScrollBehaviorTarget;
 	TUniquePtr<FEditor2DMouseWheelZoomBehaviorTarget> ZoomBehaviorTarget;
