Chaos Cloth Editor: Add missing file from last CL to actually update the weight map menu.

diff --git a/Engine/Plugins/ChaosClothAssetEditor/Source/ChaosClothAssetEditor/Private/ChaosClothAsset/ClothEditorToolkit.cpp b/Engine/Plugins/ChaosClothAssetEditor/Source/ChaosClothAssetEditor/Private/ChaosClothAsset/ClothEditorToolkit.cpp
--- a/Engine/Plugins/ChaosClothAssetEditor/Source/ChaosClothAssetEditor/Private/ChaosClothAsset/ClothEditorToolkit.cpp
+++ b/Engine/Plugins/ChaosClothAssetEditor/Source/ChaosClothAssetEditor/Private/ChaosClothAsset/ClothEditorToolkit.cpp
@@ -1376,6 +1376,11 @@ void FChaosClothAssetEditorToolkit::OnClothAssetChanged()
 		}
 	}
 
+	if (ClothEditorSimulationVisualization)
+	{
+		ClothEditorSimulationVisualization->RefreshMenusForClothComponent(ClothPreviewScene->GetClothComponent());
+	}
+
 	if (bWasSimulationSuspended)
 	{
 		ClothMode->SuspendSimulation();
