Cloth Asset:
- Enabled multiple sections import on the SkeletalMeshImport node.
- Added a SetPhysicsAsset boolean property to the SkeletalMeshImport node, and set the behavior to false by default.

#tests Imported a skeletal mesh that had more than one section, and checked whether setting the physics asset was affected by the new checkbox.

diff --git a/Engine/Plugins/ChaosClothAssetEditor/Source/ChaosClothAssetDataflowNodes/Private/ChaosClothAsset/SkeletalMeshImportNode.cpp b/Engine/Plugins/ChaosClothAssetEditor/Source/ChaosClothAssetDataflowNodes/Private/ChaosClothAsset/SkeletalMeshImportNode.cpp
--- a/Engine/Plugins/ChaosClothAssetEditor/Source/ChaosClothAssetDataflowNodes/Private/ChaosClothAsset/SkeletalMeshImportNode.cpp
+++ b/Engine/Plugins/ChaosClothAssetEditor/Source/ChaosClothAssetDataflowNodes/Private/ChaosClothAsset/SkeletalMeshImportNode.cpp
@@ -50,35 +50,38 @@ void FChaosClothAssetSkeletalMeshImportNode::Evaluate(Dataflow::FContext& Contex
 			}
 
 			const FSkeletalMeshLODModel &LODModel = ImportedModel->LODModels[LODIndex];
-			const bool bIsValidSection = LODModel.Sections.IsValidIndex(SectionIndex);;
-			if (!bIsValidSection)
-			{
-				FClothDataflowTools::LogAndToastWarning(*this, LOCTEXT("InvalidSectionHeadline", "Invalid section."),
-					FText::Format(
-						LOCTEXT("InvalidSectionDetails", "No valid section {0} found for skeletal mesh {1}."),
-						SectionIndex,
-						FText::FromString(SkeletalMesh->GetName())));
-
-				SetValue(Context, MoveTemp(*ClothCollection), &Collection);
-				return;
-			}
-
-			const FSkelMeshSection &Section = LODModel.Sections[SectionIndex];
+			const int32 FirstSection = bImportSingleSection ? SectionIndex : 0;
+			const int32 LastSection = bImportSingleSection ? SectionIndex : LODModel.Sections.Num() - 1;
 
-			if (bImportSimMesh)
+			for (int32 Section = FirstSection; Section <= LastSection; ++Section)
 			{
-				FClothDataflowTools::AddSimPatternsFromSkeletalMeshSection(ClothCollection, LODModel, SectionIndex, UVChannel, UVScale);
+				const bool bIsValidSection = LODModel.Sections.IsValidIndex(Section);;
+				if (!bIsValidSection)
+				{
+					FClothDataflowTools::LogAndToastWarning(*this, LOCTEXT("InvalidSectionHeadline", "Invalid section."),
+						FText::Format(
+							LOCTEXT("InvalidSectionDetails", "No valid section {0} found for skeletal mesh {1}."),
+							Section,
+							FText::FromString(SkeletalMesh->GetName())));
+
+					continue;
+				}
+
+				if (bImportSimMesh)
+				{
+					FClothDataflowTools::AddSimPatternsFromSkeletalMeshSection(ClothCollection, LODModel, Section, UVChannel, UVScale);
+				}
+
+				if (bImportRenderMesh)
+				{
+					const TArray<FSkeletalMaterial>& Materials = SkeletalMesh->GetMaterials();
+					check(Section < Materials.Num());
+					const FString RenderMaterialPathName = Materials[Section].MaterialInterface ? Materials[Section].MaterialInterface->GetPathName() : "";
+					FClothDataflowTools::AddRenderPatternFromSkeletalMeshSection(ClothCollection, LODModel, Section, RenderMaterialPathName);
+				}
 			}
 
-			if (bImportRenderMesh)
-			{
-				const TArray<FSkeletalMaterial>& Materials = SkeletalMesh->GetMaterials();
-				check(SectionIndex < Materials.Num());
-				const FString RenderMaterialPathName = Materials[SectionIndex].MaterialInterface ? Materials[SectionIndex].MaterialInterface->GetPathName() : "";
-				FClothDataflowTools::AddRenderPatternFromSkeletalMeshSection(ClothCollection, LODModel, SectionIndex, RenderMaterialPathName);
-			}
-
-			if (const UPhysicsAsset* PhysicsAsset = SkeletalMesh->GetPhysicsAsset())
+			if (const UPhysicsAsset* PhysicsAsset = bSetPhysicsAsset ? SkeletalMesh->GetPhysicsAsset() : nullptr)
 			{
 				ClothFacade.SetPhysicsAssetPathName(PhysicsAsset->GetPathName());
 			}
@@ -89,4 +92,14 @@ void FChaosClothAssetSkeletalMeshImportNode::Evaluate(Dataflow::FContext& Contex
 	}
 }
 
+void FChaosClothAssetSkeletalMeshImportNode::Serialize(FArchive& Ar)
+{
+	Ar.UsingCustomVersion(FFortniteMainBranchObjectVersion::GUID);
+	if (Ar.CustomVer(FUE5MainStreamObjectVersion::GUID) < FFortniteMainBranchObjectVersion::ClothAssetSkeletalMeshMultiSectionImport)
+	{
+		bImportSingleSection = true;
+		bSetPhysicsAsset = true;
+	}
+}
+
 #undef LOCTEXT_NAMESPACE
diff --git a/Engine/Plugins/ChaosClothAssetEditor/Source/ChaosClothAssetDataflowNodes/Public/ChaosClothAsset/SkeletalMeshImportNode.h b/Engine/Plugins/ChaosClothAssetEditor/Source/ChaosClothAssetDataflowNodes/Public/ChaosClothAsset/SkeletalMeshImportNode.h
--- a/Engine/Plugins/ChaosClothAssetEditor/Source/ChaosClothAssetDataflowNodes/Public/ChaosClothAsset/SkeletalMeshImportNode.h
+++ b/Engine/Plugins/ChaosClothAssetEditor/Source/ChaosClothAssetDataflowNodes/Public/ChaosClothAsset/SkeletalMeshImportNode.h
@@ -27,8 +27,12 @@ public:
 	UPROPERTY(EditAnywhere, Category = "Skeletal Mesh Import", Meta = (ClampMin = "0", DisplayName = "LOD Index"))
 	int32 LODIndex = 0;
 
-	/** The skeletal mesh LOD section to import. */
-	UPROPERTY(EditAnywhere, Category = "Skeletal Mesh Import", Meta = (ClampMin = "0"))
+	/** Enable single import section mode. */
+	UPROPERTY(EditAnywhere, Category = "Skeletal Mesh Import")
+	bool bImportSingleSection = false;
+
+	/** The skeletal mesh LOD section to import. If not enabled, then all sections will be imported. */
+	UPROPERTY(EditAnywhere, Category = "Skeletal Mesh Import", Meta = (ClampMin = "0", EditCondition = "bImportSingleSection"))
 	int32 SectionIndex = 0;
 
 	/** Whether to import the simulation mesh from the specified skeletal mesh. */
@@ -41,7 +45,7 @@ public:
 
 	/**
 	 * UV channel of the skeletal mesh to import the 2D simulation mesh patterns from.
-	 * If set to -1, then the import will unwrap the 3D simulation mesh into 2D simulation mesh patterns.
+	 * If set to -1, or the specified UVChannel doesn't exist then the import will unwrap the 3D simulation mesh into 2D simulation mesh patterns.
 	 */
 	UPROPERTY(EditAnywhere, Category = "Skeletal Mesh Import", Meta = (ClampMin = "-1", EditCondition = "bImportSimMesh"))
 	int32 UVChannel = 0;
@@ -50,8 +54,13 @@ public:
 	UPROPERTY(EditAnywhere, Category = "Skeletal Mesh Import", Meta = (AllowPreserveRatio, EditCondition = "bImportSimMesh && UVChannel != INDEX_NONE"))
 	FVector2f UVScale = { 1.f, 1.f };
 
+	/** Set the same physics asset as the one used by the imported skeletal mesh. */
+	UPROPERTY(EditAnywhere, Category = "Skeletal Mesh Import")
+	bool bSetPhysicsAsset = false;
+
 	FChaosClothAssetSkeletalMeshImportNode(const Dataflow::FNodeParameters& InParam, FGuid InGuid = FGuid::NewGuid());
 
 private:
 	virtual void Evaluate(Dataflow::FContext& Context, const FDataflowOutput* Out) const override;
+	virtual void Serialize(FArchive& Ar) override;
 };
diff --git a/Engine/Plugins/ChaosClothAssetEditor/Source/ChaosClothAssetDataflowNodes/Public/ChaosClothAsset/StaticMeshImportNode.h b/Engine/Plugins/ChaosClothAssetEditor/Source/ChaosClothAssetDataflowNodes/Public/ChaosClothAsset/StaticMeshImportNode.h
--- a/Engine/Plugins/ChaosClothAssetEditor/Source/ChaosClothAssetDataflowNodes/Public/ChaosClothAsset/StaticMeshImportNode.h
+++ b/Engine/Plugins/ChaosClothAssetEditor/Source/ChaosClothAssetDataflowNodes/Public/ChaosClothAsset/StaticMeshImportNode.h
@@ -19,23 +19,26 @@ public:
 	UPROPERTY(Meta = (DataflowOutput))
 	FManagedArrayCollection Collection;
 
-	/* The Static Mesh to import from */
+	/* The Static Mesh to import from. */
 	UPROPERTY(EditAnywhere, Category = "Static Mesh Import")
 	TObjectPtr<const UStaticMesh> StaticMesh;
 
-	/* Which static mesh Lod to import.*/
+	/* Which static mesh Lod to import. */
 	UPROPERTY(EditAnywhere, Category = "Static Mesh Import", Meta = (ClampMin = "0", DisplayName = "LOD Index"))
 	int32 LODIndex = 0;
 
-	/* Import static mesh data as a simulation mesh data.*/
+	/* Import static mesh data as a simulation mesh data. */
 	UPROPERTY(EditAnywhere, Category = "Static Mesh Import")
 	bool bImportSimMesh = true;
 
-	/* Import static mesh data as render mesh data.*/
+	/* Import static mesh data as render mesh data. */
 	UPROPERTY(EditAnywhere, Category = "Static Mesh Import")
 	bool bImportRenderMesh = true;
 
-	/* UV Channel used to populate Sim Mesh positions */
+	/**
+	 * UV channel of the static mesh to import the 2D simulation mesh patterns from.
+	 * If set to -1, or the specified UVChannel doesn't exist then the import will unwrap the 3D simulation mesh into 2D simulation mesh patterns.
+	 */
 	UPROPERTY(EditAnywhere, Category = "Static Mesh Import", meta = (ClampMin = "-1", EditCondition = "bImportSimMesh"))
 	int32 UVChannel = 0;
 
diff --git a/Engine/Source/Runtime/Core/Public/UObject/FortniteMainBranchObjectVersion.h b/Engine/Source/Runtime/Core/Public/UObject/FortniteMainBranchObjectVersion.h
--- a/Engine/Source/Runtime/Core/Public/UObject/FortniteMainBranchObjectVersion.h
+++ b/Engine/Source/Runtime/Core/Public/UObject/FortniteMainBranchObjectVersion.h
@@ -477,6 +477,9 @@ struct FFortniteMainBranchObjectVersion
 		// Added hideable pins to dataflow
 		DataflowHideablePins,
 
+		// Added multiple section import to the cloth asset skeletal mesh import node
+		ClothAssetSkeletalMeshMultiSectionImport,
+
 		// -----<new versions can be added above this line>-------------------------------------------------
 		VersionPlusOne,
 		LatestVersion = VersionPlusOne - 1
