Chaos Cloth Asset - Fixed crash when using slider control to set values on certain nodes (e.g. UV Scale).

#tests Used the property sliders with the fix in, and observed no crash.
#rnx

diff --git a/Engine/Plugins/ChaosClothAssetEditor/Source/ChaosClothAssetEditor/Private/ChaosClothAsset/ClothEditorToolkit.cpp b/Engine/Plugins/ChaosClothAssetEditor/Source/ChaosClothAssetEditor/Private/ChaosClothAsset/ClothEditorToolkit.cpp
--- a/Engine/Plugins/ChaosClothAssetEditor/Source/ChaosClothAssetEditor/Private/ChaosClothAsset/ClothEditorToolkit.cpp
+++ b/Engine/Plugins/ChaosClothAssetEditor/Source/ChaosClothAssetEditor/Private/ChaosClothAsset/ClothEditorToolkit.cpp
@@ -253,6 +253,10 @@ UDataflow* FChaosClothAssetEditorToolkit::GetDataflow()
 
 void FChaosClothAssetEditorToolkit::Tick(float DeltaTime)
 {
+	// Execute commands that required waiting for the next tick to execute (e.g. refresh of the details view)
+	TickCommands.Broadcast();
+	TickCommands.Clear();
+
 	// Evaluate terminal node if one is assigned, but only if needed (don't force it)
 	constexpr bool bForceOperation = false;
 	EvaluateNode(nullptr, bForceOperation);
@@ -1248,6 +1252,8 @@ void FChaosClothAssetEditorToolkit::OnNodeSelectionChanged(const TSet<UObject*>&
 							OnNodeInvalidatedDelegateHandle = SelectedDataflowNode->GetOnNodeInvalidatedDelegate().AddLambda(
 								[this](FDataflowNode* InDataflowNode)
 								{
+									// Warning: Do not execute code that rebuilds the UI in this lambda as it is called by the UI!
+
 									if (SelectedDataflowNode.Get() == InDataflowNode)
 									{
 										GetClothCollectionIfPossible(SelectedDataflowNode, DataflowContext);
@@ -1256,11 +1262,15 @@ void FChaosClothAssetEditorToolkit::OnNodeSelectionChanged(const TSet<UObject*>&
 											SelectedDataflowNode->OnSelected(*DataflowContext);
 										}
 									}
-									// The detail panel won't update correctly when a change affects a node's properties and needs a refresh
-									if (NodeDetailsEditor && NodeDetailsEditor->GetDetailsView())
-									{
-										NodeDetailsEditor->GetDetailsView()->ForceRefresh();
-									}
+									// The detail panel won't update correctly when a change affects a node's properties and needs a refresh at the next tick
+									TickCommands.AddLambda([this]()
+										{
+											if (NodeDetailsEditor && NodeDetailsEditor->GetDetailsView())
+											{
+												NodeDetailsEditor->GetDetailsView()->InvalidateCachedState();
+											}
+										});
+
 									// The node has just been invalidated, best to close the current tool
 									if (const UChaosClothAssetEditorMode* const ClothMode = 
 										CastChecked<UChaosClothAssetEditorMode>(EditorModeManager->GetActiveScriptableMode(UChaosClothAssetEditorMode::EM_ChaosClothAssetEditorModeId)))
diff --git a/Engine/Plugins/ChaosClothAssetEditor/Source/ChaosClothAssetEditor/Private/ChaosClothAsset/ClothEditorToolkit.h b/Engine/Plugins/ChaosClothAssetEditor/Source/ChaosClothAssetEditor/Private/ChaosClothAsset/ClothEditorToolkit.h
--- a/Engine/Plugins/ChaosClothAssetEditor/Source/ChaosClothAssetEditor/Private/ChaosClothAsset/ClothEditorToolkit.h
+++ b/Engine/Plugins/ChaosClothAssetEditor/Source/ChaosClothAssetEditor/Private/ChaosClothAsset/ClothEditorToolkit.h
@@ -177,6 +177,8 @@ private:
 	TSharedPtr<IStructureDetailsView> NodeDetailsEditor;
 
 	FDelegateHandle OnPackageReloadedDelegateHandle;
-	
+
+	DECLARE_MULTICAST_DELEGATE(FTickCommands)
+	FTickCommands TickCommands;
 };
 } // namespace UE::Chaos::ClothAsset
