Replaced TEMPLATE_REQUIRES usage.

diff --git a/Engine/Plugins/ChaosClothAsset/Source/ChaosClothAsset/Private/ChaosClothAsset/ClothCollection.cpp b/Engine/Plugins/ChaosClothAsset/Source/ChaosClothAsset/Private/ChaosClothAsset/ClothCollection.cpp
--- a/Engine/Plugins/ChaosClothAsset/Source/ChaosClothAsset/Private/ChaosClothAsset/ClothCollection.cpp
+++ b/Engine/Plugins/ChaosClothAsset/Source/ChaosClothAsset/Private/ChaosClothAsset/ClothCollection.cpp
@@ -675,7 +675,7 @@ namespace UE::Chaos::ClothAsset
 	template CHAOSCLOTHASSET_API TTuple<int32, int32> FClothCollection::GetSubElementsStartEnd<false, true>(const TManagedArray<int32>* StartArray, const TManagedArray<int32>* EndArray, const TManagedArray<int32>* StartSubArray, const TManagedArray<int32>* EndSubArray, int32 ArrayIndex);
 	template CHAOSCLOTHASSET_API TTuple<int32, int32> FClothCollection::GetSubElementsStartEnd<true, true>(const TManagedArray<int32>* StartArray, const TManagedArray<int32>* EndArray, const TManagedArray<int32>* StartSubArray, const TManagedArray<int32>* EndSubArray, int32 ArrayIndex);
 
-	template<typename T, typename TEnableIf<TIsUserAttributeType<T>::Value, int>::type>
+	template<typename T UE_REQUIRES_DEFINITION(TIsUserAttributeType<T>::Value)>
 	TArray<FName> FClothCollection::GetUserDefinedAttributeNames(const FName& GroupName) const
 	{
 		using namespace UE::Chaos::ClothAsset::Private;
@@ -702,7 +702,7 @@ namespace UE::Chaos::ClothAsset
 	template CHAOSCLOTHASSET_API TArray<FName> FClothCollection::GetUserDefinedAttributeNames<float>(const FName& GroupName) const;
 	template CHAOSCLOTHASSET_API TArray<FName> FClothCollection::GetUserDefinedAttributeNames<FVector3f>(const FName& GroupName) const;
 
-	template<typename T, typename TEnableIf<TIsUserAttributeType<T>::Value, int>::type>
+	template<typename T UE_REQUIRES_DEFINITION(TIsUserAttributeType<T>::Value)>
 	TManagedArray<T>* FClothCollection::FindOrAddUserDefinedAttribute(const FName& Name, const FName& GroupName)
 	{
 		using namespace UE::Chaos::ClothAsset::Private;
@@ -742,7 +742,7 @@ namespace UE::Chaos::ClothAsset
 		}
 	}
 
-	template<typename T, typename TEnableIf<TIsUserAttributeType<T>::Value, int>::type>
+	template<typename T UE_REQUIRES_DEFINITION(TIsUserAttributeType<T>::Value)>
 	bool FClothCollection::HasUserDefinedAttribute(const FName& Name, const FName& GroupName) const
 	{
 		using namespace UE::Chaos::ClothAsset::Private;
@@ -769,7 +769,7 @@ namespace UE::Chaos::ClothAsset
 	template CHAOSCLOTHASSET_API bool FClothCollection::HasUserDefinedAttribute<float>(const FName& Name, const FName& GroupName) const;
 	template CHAOSCLOTHASSET_API bool FClothCollection::HasUserDefinedAttribute<FVector3f>(const FName& Name, const FName& GroupName) const;
 
-	template<typename T, typename TEnableIf<TIsUserAttributeType<T>::Value, int>::type>
+	template<typename T UE_REQUIRES_DEFINITION(TIsUserAttributeType<T>::Value)>
 	const TManagedArray<T>* FClothCollection::GetUserDefinedAttribute(const FName& Name, const FName& GroupName) const
 	{
 		using namespace UE::Chaos::ClothAsset::Private;
@@ -796,7 +796,7 @@ namespace UE::Chaos::ClothAsset
 	template CHAOSCLOTHASSET_API const TManagedArray<float>* FClothCollection::GetUserDefinedAttribute<float>(const FName& Name, const FName& GroupName) const;
 	template CHAOSCLOTHASSET_API const TManagedArray<FVector3f>* FClothCollection::GetUserDefinedAttribute<FVector3f>(const FName& Name, const FName& GroupName) const;
 
-	template<typename T, typename TEnableIf<TIsUserAttributeType<T>::Value, int>::type>
+	template<typename T UE_REQUIRES_DEFINITION(TIsUserAttributeType<T>::Value)>
 	TManagedArray<T>* FClothCollection::GetUserDefinedAttribute(const FName& Name, const FName& GroupName)
 	{
 		using namespace UE::Chaos::ClothAsset::Private;
diff --git a/Engine/Plugins/ChaosClothAsset/Source/ChaosClothAsset/Private/ChaosClothAsset/ClothCollection.h b/Engine/Plugins/ChaosClothAsset/Source/ChaosClothAsset/Private/ChaosClothAsset/ClothCollection.h
--- a/Engine/Plugins/ChaosClothAsset/Source/ChaosClothAsset/Private/ChaosClothAsset/ClothCollection.h
+++ b/Engine/Plugins/ChaosClothAsset/Source/ChaosClothAsset/Private/ChaosClothAsset/ClothCollection.h
@@ -118,7 +118,7 @@ namespace UE::Chaos::ClothAsset
 		template<typename T>
 		static inline void CopyArrayViewData(const TArrayView<T>& To, const TConstArrayView<T>& From);
 
-		template<typename T, TEMPLATE_REQUIRES(!TIsArray<T>::Value)>
+		template<typename T UE_REQUIRES(!std::is_array_v<T>)>
 		static inline void CopyArrayViewDataAndApplyOffset(const TArrayView<T>& To, const TConstArrayView<T>& From, const T Offset);
 
 		template<typename T>
@@ -128,13 +128,13 @@ namespace UE::Chaos::ClothAsset
 		static inline uint32 GetElementsTypeHash(const TManagedArray<T>* ElementArray);
 
 		//~ Weight maps
-		template<typename T, TEMPLATE_REQUIRES(TIsUserAttributeType<T>::Value)>
+		template<typename T UE_REQUIRES(TIsUserAttributeType<T>::Value)>
 		TArray<FName> GetUserDefinedAttributeNames(const FName& GroupName) const;
 
-		template<typename T, TEMPLATE_REQUIRES(TIsUserAttributeType<T>::Value)>
+		template<typename T UE_REQUIRES(TIsUserAttributeType<T>::Value)>
 		TManagedArray<T>* FindOrAddUserDefinedAttribute(const FName& Name, const FName& GroupName);
 
-		template<typename T, TEMPLATE_REQUIRES(TIsUserAttributeType<T>::Value)>
+		template<typename T UE_REQUIRES(TIsUserAttributeType<T>::Value)>
 		void AddUserDefinedAttribute(const FName& Name, const FName& GroupName)
 		{
 			FindOrAddUserDefinedAttribute<T>(Name, GroupName);
@@ -142,13 +142,13 @@ namespace UE::Chaos::ClothAsset
 
 		void RemoveUserDefinedAttribute(const FName& Name, const FName& GroupName);
 
-		template<typename T, TEMPLATE_REQUIRES(TIsUserAttributeType<T>::Value)>
+		template<typename T UE_REQUIRES(TIsUserAttributeType<T>::Value)>
 		bool HasUserDefinedAttribute(const FName& Name, const FName& GroupName) const;
 
-		template<typename T, TEMPLATE_REQUIRES(TIsUserAttributeType<T>::Value)>
+		template<typename T UE_REQUIRES(TIsUserAttributeType<T>::Value)>
 		const TManagedArray<T>* GetUserDefinedAttribute(const FName& Name, const FName& GroupName) const;
 
-		template<typename T, TEMPLATE_REQUIRES(TIsUserAttributeType<T>::Value)>
+		template<typename T UE_REQUIRES(TIsUserAttributeType<T>::Value)>
 		TManagedArray<T>* GetUserDefinedAttribute(const FName& Name, const FName& GroupName);
 
 		static bool IsValidClothCollectionGroupName(const FName& GroupName);
@@ -513,7 +513,7 @@ namespace UE::Chaos::ClothAsset
 		}
 	}
 
-	template<typename T, typename TEnableIf<!TIsArray<T>::Value, int>::type>
+	template<typename T UE_REQUIRES_DEFINITION(!std::is_array_v<T>)>
 	inline void FClothCollection::CopyArrayViewDataAndApplyOffset(const TArrayView<T>& To, const TConstArrayView<T>& From, const T Offset)
 	{
 		check(To.Num() == From.Num());
diff --git a/Engine/Plugins/ChaosClothAsset/Source/ChaosClothAsset/Private/ChaosClothAsset/CollectionClothFacade.cpp b/Engine/Plugins/ChaosClothAsset/Source/ChaosClothAsset/Private/ChaosClothAsset/CollectionClothFacade.cpp
--- a/Engine/Plugins/ChaosClothAsset/Source/ChaosClothAsset/Private/ChaosClothAsset/CollectionClothFacade.cpp
+++ b/Engine/Plugins/ChaosClothAsset/Source/ChaosClothAsset/Private/ChaosClothAsset/CollectionClothFacade.cpp
@@ -159,7 +159,7 @@ namespace UE::Chaos::ClothAsset
 		return ResultHash;
 	}
 
-	template<typename T, typename TEnableIf<TIsUserAttributeType<T>::Value, int>::type>
+	template<typename T UE_REQUIRES_DEFINITION(TIsUserAttributeType<T>::Value)>
 	uint32 FCollectionClothConstFacade::CalculateUserDefinedAttributesTypeHash(const FName& GroupName, uint32 PreviousHash) const
 	{
 		uint32 ResultHash = PreviousHash;
@@ -458,7 +458,7 @@ namespace UE::Chaos::ClothAsset
 	}
 
 
-	template<typename T, typename TEnableIf<TIsUserAttributeType<T>::Value, int>::type>
+	template<typename T UE_REQUIRES_DEFINITION(TIsUserAttributeType<T>::Value)>
 	bool FCollectionClothConstFacade::HasUserDefinedAttribute(const FName& Name, const FName& GroupName) const
 	{
 		return ClothCollection->HasUserDefinedAttribute<T>(Name, GroupName);
@@ -467,7 +467,7 @@ namespace UE::Chaos::ClothAsset
 	template CHAOSCLOTHASSET_API bool FCollectionClothConstFacade::HasUserDefinedAttribute<float>(const FName& Name, const FName& GroupName) const;
 	template CHAOSCLOTHASSET_API bool FCollectionClothConstFacade::HasUserDefinedAttribute<FVector3f>(const FName& Name, const FName& GroupName) const;
 
-	template<typename T, typename TEnableIf<TIsUserAttributeType<T>::Value, int>::type>
+	template<typename T UE_REQUIRES_DEFINITION(TIsUserAttributeType<T>::Value)>
 	TArray<FName> FCollectionClothConstFacade::GetUserDefinedAttributeNames(const FName& GroupName) const
 	{
 		return ClothCollection->GetUserDefinedAttributeNames<T>(GroupName);
@@ -476,7 +476,7 @@ namespace UE::Chaos::ClothAsset
 	template CHAOSCLOTHASSET_API TArray<FName> FCollectionClothConstFacade::GetUserDefinedAttributeNames<float>(const FName& GroupName) const;
 	template CHAOSCLOTHASSET_API TArray<FName> FCollectionClothConstFacade::GetUserDefinedAttributeNames<FVector3f>(const FName& GroupName) const;
 
-	template<typename T, typename TEnableIf<TIsUserAttributeType<T>::Value, int>::type>
+	template<typename T UE_REQUIRES_DEFINITION(TIsUserAttributeType<T>::Value)>
 	TConstArrayView<T> FCollectionClothConstFacade::GetUserDefinedAttribute(const FName& Name, const FName& GroupName) const
 	{
 		return ClothCollection->GetElements(ClothCollection->GetUserDefinedAttribute<T>(Name, GroupName));
@@ -1121,7 +1121,7 @@ namespace UE::Chaos::ClothAsset
 		return GetClothCollection()->GetElements(GetClothCollection()->GetUserDefinedAttribute<float>(Name, ClothCollectionGroup::SimVertices3D));
 	}
 
-	template<typename T, typename TEnableIf<TIsUserAttributeType<T>::Value, int>::type>
+	template<typename T UE_REQUIRES_DEFINITION(TIsUserAttributeType<T>::Value)>
 	bool FCollectionClothFacade::AddUserDefinedAttribute(const FName& Name, const FName& GroupName)
 	{
 		check(IsValid());
@@ -1137,7 +1137,7 @@ namespace UE::Chaos::ClothAsset
 		GetClothCollection()->RemoveUserDefinedAttribute(Name, GroupName);
 	}
 
-	template<typename T, typename TEnableIf<TIsUserAttributeType<T>::Value, int>::type>
+	template<typename T UE_REQUIRES_DEFINITION(TIsUserAttributeType<T>::Value)>
 	TArrayView<T> FCollectionClothFacade::GetUserDefinedAttribute(const FName& Name, const FName& GroupName)
 	{
 		return GetClothCollection()->GetElements(GetClothCollection()->GetUserDefinedAttribute<T>(Name, GroupName));
diff --git a/Engine/Plugins/ChaosClothAsset/Source/ChaosClothAsset/Private/ChaosClothAsset/CollectionClothSimPatternFacade.cpp b/Engine/Plugins/ChaosClothAsset/Source/ChaosClothAsset/Private/ChaosClothAsset/CollectionClothSimPatternFacade.cpp
--- a/Engine/Plugins/ChaosClothAsset/Source/ChaosClothAsset/Private/ChaosClothAsset/CollectionClothSimPatternFacade.cpp
+++ b/Engine/Plugins/ChaosClothAsset/Source/ChaosClothAsset/Private/ChaosClothAsset/CollectionClothSimPatternFacade.cpp
@@ -112,7 +112,7 @@ namespace UE::Chaos::ClothAsset
 		SetDefaults();
 	}
 
-	template<typename IndexType, typename TEnableIf<TIsIndexType<IndexType>::Value, int>::type>
+	template<typename IndexType UE_REQUIRES_DEFINITION(TIsIndexType<IndexType>::Value)>
 	void FCollectionClothSimPatternFacade::Initialize(const TArray<FVector2f>& Positions2D, const TArray<FVector3f>& Positions3D, const TArray<IndexType>& Indices, const int32 FabricIndex)
 	{
 		Reset();
diff --git a/Engine/Plugins/ChaosClothAsset/Source/ChaosClothAsset/Public/ChaosClothAsset/CollectionClothFacade.h b/Engine/Plugins/ChaosClothAsset/Source/ChaosClothAsset/Public/ChaosClothAsset/CollectionClothFacade.h
--- a/Engine/Plugins/ChaosClothAsset/Source/ChaosClothAsset/Public/ChaosClothAsset/CollectionClothFacade.h
+++ b/Engine/Plugins/ChaosClothAsset/Source/ChaosClothAsset/Public/ChaosClothAsset/CollectionClothFacade.h
@@ -58,7 +58,7 @@ namespace UE::Chaos::ClothAsset
 
 		uint32 CalculateTypeHash(bool bIncludeWeightMaps, uint32 PreviousHash = 0) const;
 		uint32 CalculateWeightMapTypeHash(uint32 PreviousHash = 0) const;
-		template<typename T, TEMPLATE_REQUIRES(TIsUserAttributeType<T>::Value)>
+		template<typename T UE_REQUIRES(TIsUserAttributeType<T>::Value)>
 		uint32 CalculateUserDefinedAttributesTypeHash(const FName& GroupName, uint32 PreviousHash = 0) const;
 
 		//~ LOD (single per collection) Group
@@ -170,11 +170,11 @@ namespace UE::Chaos::ClothAsset
 		TConstArrayView<float> GetWeightMap(const FName& Name) const;
 
 		//~ Other User-Defined Attributes (not instantiated for bools)
-		template<typename T, TEMPLATE_REQUIRES(TIsUserAttributeType<T>::Value)>
+		template<typename T UE_REQUIRES(TIsUserAttributeType<T>::Value)>
 		bool HasUserDefinedAttribute(const FName& Name, const FName& GroupName) const;
-		template<typename T, TEMPLATE_REQUIRES(TIsUserAttributeType<T>::Value)>
+		template<typename T UE_REQUIRES(TIsUserAttributeType<T>::Value)>
 		TArray<FName> GetUserDefinedAttributeNames(const FName& GroupName) const;
-		template<typename T, TEMPLATE_REQUIRES(TIsUserAttributeType<T>::Value)>
+		template<typename T UE_REQUIRES(TIsUserAttributeType<T>::Value)>
 		TConstArrayView<T> GetUserDefinedAttribute(const FName& Name, const FName& GroupName) const;
 		static bool IsValidClothCollectionGroupName(const FName& GroupName);
 
@@ -351,10 +351,10 @@ namespace UE::Chaos::ClothAsset
 
 		//~ Other User-Defined Attributes (not instantiated for bools)
 		/** GroupName must be an existing group as defined in ClothCollectionGroup. Returns success */
-		template<typename T, TEMPLATE_REQUIRES(TIsUserAttributeType<T>::Value)>
+		template<typename T UE_REQUIRES(TIsUserAttributeType<T>::Value)>
 		bool AddUserDefinedAttribute(const FName& Name, const FName& GroupName);
 		void RemoveUserDefinedAttribute(const FName& Name, const FName& GroupName);
-		template<typename T, TEMPLATE_REQUIRES(TIsUserAttributeType<T>::Value)>
+		template<typename T UE_REQUIRES(TIsUserAttributeType<T>::Value)>
 		TArrayView<T> GetUserDefinedAttribute(const FName& Name, const FName& GroupName);
 
 	private:
diff --git a/Engine/Plugins/ChaosClothAsset/Source/ChaosClothAsset/Public/ChaosClothAsset/CollectionClothSimPatternFacade.h b/Engine/Plugins/ChaosClothAsset/Source/ChaosClothAsset/Public/ChaosClothAsset/CollectionClothSimPatternFacade.h
--- a/Engine/Plugins/ChaosClothAsset/Source/ChaosClothAsset/Public/ChaosClothAsset/CollectionClothSimPatternFacade.h
+++ b/Engine/Plugins/ChaosClothAsset/Source/ChaosClothAsset/Public/ChaosClothAsset/CollectionClothSimPatternFacade.h
@@ -90,7 +90,7 @@ namespace UE::Chaos::ClothAsset
 		void Reset();
 
 		/** Initialize the cloth pattern using the specified 3D and 2D positions, and topology. */
-		template<typename IndexType, TEMPLATE_REQUIRES(TIsIndexType<IndexType>::Value)>
+		template<typename IndexType UE_REQUIRES(TIsIndexType<IndexType>::Value)>
 		void Initialize(const TArray<FVector2f>& Positions2D, const TArray<FVector3f>& Positions3D, const TArray<IndexType>& Indices, const int32 FabricIndex = INDEX_NONE);
 
 		/** Initialize this pattern using another pattern collection. */
diff --git a/Engine/Plugins/ChaosClothAssetEditor/Source/ChaosClothAssetDataflowNodes/Public/ChaosClothAsset/SimulationBaseConfigNode.h b/Engine/Plugins/ChaosClothAssetEditor/Source/ChaosClothAssetDataflowNodes/Public/ChaosClothAsset/SimulationBaseConfigNode.h
--- a/Engine/Plugins/ChaosClothAssetEditor/Source/ChaosClothAssetDataflowNodes/Public/ChaosClothAsset/SimulationBaseConfigNode.h
+++ b/Engine/Plugins/ChaosClothAssetEditor/Source/ChaosClothAssetDataflowNodes/Public/ChaosClothAsset/SimulationBaseConfigNode.h
@@ -41,30 +41,30 @@ protected:
 	public:
 		FPropertyHelper(const FChaosClothAssetSimulationBaseConfigNode& InConfigNode, Dataflow::FContext& InContext, ::Chaos::Softs::FCollectionPropertyMutableFacade& InProperties, const TSharedRef<FManagedArrayCollection>& InClothCollection);
 
-		template<typename PropertyType, TEMPLATE_REQUIRES(::Chaos::Softs::TIsWeightedType<PropertyType>::Value && !std::is_same<PropertyType, bool>::value)>
+		template<typename PropertyType UE_REQUIRES(::Chaos::Softs::TIsWeightedType<PropertyType>::Value && !std::is_same_v<PropertyType, bool>)>
 		inline int32 SetProperty(const FName& PropertyName, const PropertyType& PropertyValue, const TArray<FName>& SimilarPropertyNames = {}, ECollectionPropertyFlags PropertyFlags = ECollectionPropertyFlags::Animatable);
 
-		template<typename T, typename PropertyType, TEMPLATE_REQUIRES(TIsDerivedFrom<T, FChaosClothAssetSimulationBaseConfigNode>::Value && ::Chaos::Softs::TIsWeightedType<PropertyType>::Value && !std::is_same<PropertyType, bool>::value)>
+		template<typename T, typename PropertyType UE_REQUIRES(std::is_base_of_v<FChaosClothAssetSimulationBaseConfigNode, T> && ::Chaos::Softs::TIsWeightedType<PropertyType>::Value && !std::is_same_v<PropertyType, bool>)>
 		inline int32 SetProperty(const T* ConfigStruct, const PropertyType* PropertyValue, const TArray<FName>& SimilarPropertyNames = {}, ECollectionPropertyFlags PropertyFlags = ECollectionPropertyFlags::Animatable);
 
 		int32 SetPropertyBool(const FName& PropertyName, bool PropertyValue, const TArray<FName>& SimilarPropertyNames = {}, ECollectionPropertyFlags PropertyFlags = ECollectionPropertyFlags::Animatable);
 
-		template<typename T, TEMPLATE_REQUIRES(TIsDerivedFrom<T, FChaosClothAssetSimulationBaseConfigNode>::Value)>
+		template<typename T UE_REQUIRES(std::is_base_of_v<FChaosClothAssetSimulationBaseConfigNode, T>)>
 		inline int32 SetPropertyBool(const T* ConfigStruct, const bool* PropertyValue, const TArray<FName>& SimilarPropertyNames = {}, ECollectionPropertyFlags PropertyFlags = ECollectionPropertyFlags::Animatable);
 
-		template<typename EnumType, TEMPLATE_REQUIRES(TIsEnumClass<EnumType>::Value)>
+		template<typename EnumType UE_REQUIRES(TIsEnumClass<EnumType>::Value)>
 		inline int32 SetPropertyEnum(const FName& PropertyName, const EnumType& PropertyValue, const TArray<FName>& SimilarPropertyNames = {}, ECollectionPropertyFlags PropertyFlags = ECollectionPropertyFlags::Animatable);
 
-		template<typename T, typename EnumType, TEMPLATE_REQUIRES(TIsDerivedFrom<T, FChaosClothAssetSimulationBaseConfigNode>::Value && TIsEnumClass<EnumType>::Value)>
+		template<typename T, typename EnumType UE_REQUIRES(std::is_base_of_v<FChaosClothAssetSimulationBaseConfigNode, T> && TIsEnumClass<EnumType>::Value)>
 		inline int32 SetPropertyEnum(const T* ConfigStruct, const EnumType* PropertyValue, const TArray<FName>& SimilarPropertyNames = {}, ECollectionPropertyFlags PropertyFlags = ECollectionPropertyFlags::Animatable);
 
 		int32 SetPropertyString(const FName& PropertyName, const FString& PropertyValue, const TArray<FName>& SimilarPropertyNames = {}, ECollectionPropertyFlags PropertyFlags = ECollectionPropertyFlags::Animatable);
 
 		int32 SetPropertyString(const FName& PropertyName, const FChaosClothAssetConnectableIStringValue& PropertyValue, const TArray<FName>& SimilarPropertyNames = {}, ECollectionPropertyFlags PropertyFlags = ECollectionPropertyFlags::Animatable);
 
-		template<typename T, typename ConnectableStringValueType, TEMPLATE_REQUIRES(TIsDerivedFrom<T, FChaosClothAssetSimulationBaseConfigNode>::Value)>
+		template<typename T, typename ConnectableStringValueType UE_REQUIRES(std::is_base_of_v<FChaosClothAssetSimulationBaseConfigNode, T>)>
 		inline int32 SetPropertyString(const T* ConfigStruct, const ConnectableStringValueType* PropertyValue, const TArray<FName>& SimilarPropertyNames = {}, ECollectionPropertyFlags PropertyFlags = ECollectionPropertyFlags::Animatable);
-		template<typename T, TEMPLATE_REQUIRES(TIsDerivedFrom<T, FChaosClothAssetSimulationBaseConfigNode>::Value)>
+		template<typename T UE_REQUIRES(std::is_base_of_v<FChaosClothAssetSimulationBaseConfigNode, T>)>
 		inline int32 SetPropertyString(const T* ConfigStruct, const FString* PropertyValue, const TArray<FName>& SimilarPropertyNames = {}, ECollectionPropertyFlags PropertyFlags = ECollectionPropertyFlags::Animatable);
 
 		int32 SetPropertyWeighted(const FName& PropertyName, const FChaosClothAssetWeightedValue& PropertyValue, const TArray<FName>& SimilarPropertyNames = {}, ECollectionPropertyFlags PropertyFlags = ECollectionPropertyFlags::None);
@@ -75,7 +75,7 @@ protected:
 		void OverridePropertiesFloat(const TArray<FName>& PropertyNames, const EChaosClothAssetConstraintOverrideType OverrideType, const float OverrideValue);
 		void OverridePropertiesWeighted(const TArray<FName>& PropertyNames, const EChaosClothAssetConstraintOverrideType OverrideType, const FChaosClothAssetWeightedValueOverride& OverrideValue);
 
-		template<typename T, typename WeightedValueType, TEMPLATE_REQUIRES(TIsDerivedFrom<T, FChaosClothAssetSimulationBaseConfigNode>::Value)>
+		template<typename T, typename WeightedValueType UE_REQUIRES(std::is_base_of_v<FChaosClothAssetSimulationBaseConfigNode, T>)>
 		inline int32 SetPropertyWeighted(const T* ConfigStruct, const WeightedValueType* PropertyValue, const TArray<FName>& SimilarPropertyNames = {}, ECollectionPropertyFlags PropertyFlags = ECollectionPropertyFlags::None);
 
 		FString GetPropertyString(const FString* PropertyReference) const;
@@ -118,7 +118,7 @@ protected:
 		
 	private:
 		
-		template<typename T, TEMPLATE_REQUIRES(TIsDerivedFrom<T, FChaosClothAssetSimulationBaseConfigNode>::Value)>
+		template<typename T UE_REQUIRES(std::is_base_of_v<FChaosClothAssetSimulationBaseConfigNode, T>)>
 		inline FName FindPropertyNameByAddress(const T* ConfigStruct, const void* PropertyAddress);
 		
 		int32 SetPropertyWeighted(const FName& PropertyName, const bool bIsAnimatable, const float& PropertyLow, const float& PropertyHigh, const FString& WeightMap, FString& MapOverride, const TArray<FName>& SimilarPropertyNames, ECollectionPropertyFlags PropertyFlags) const;
@@ -170,7 +170,7 @@ struct TStructOpsTypeTraits<FChaosClothAssetSimulationBaseConfigNode> : public T
 	};
 };
 
-template<typename PropertyType, typename TEnableIf<::Chaos::Softs::TIsWeightedType<PropertyType>::Value && !std::is_same<PropertyType, bool>::value, int>::type>
+template<typename PropertyType UE_REQUIRES_DEFINITION(::Chaos::Softs::TIsWeightedType<PropertyType>::Value && !std::is_same_v<PropertyType, bool>)>
 inline int32 FChaosClothAssetSimulationBaseConfigNode::FPropertyHelper::SetProperty(
 	const FName& PropertyName,
 	const PropertyType& PropertyValue,
@@ -182,7 +182,7 @@ inline int32 FChaosClothAssetSimulationBaseConfigNode::FPropertyHelper::SetPrope
 	return PropertyKeyIndex;
 }
 
-template<typename T, typename PropertyType, typename TEnableIf<TIsDerivedFrom<T, FChaosClothAssetSimulationBaseConfigNode>::Value && ::Chaos::Softs::TIsWeightedType<PropertyType>::Value && !std::is_same<PropertyType, bool>::value, int>::type>
+template<typename T, typename PropertyType UE_REQUIRES_DEFINITION(std::is_base_of_v<FChaosClothAssetSimulationBaseConfigNode, T> && ::Chaos::Softs::TIsWeightedType<PropertyType>::Value && !std::is_same_v<PropertyType, bool>)>
 inline int32 FChaosClothAssetSimulationBaseConfigNode::FPropertyHelper::SetProperty(
 	const T* ConfigStruct,
 	const PropertyType* PropertyValue,
@@ -195,7 +195,7 @@ inline int32 FChaosClothAssetSimulationBaseConfigNode::FPropertyHelper::SetPrope
 	return SetProperty(PropertyName, *PropertyValue, SimilarPropertyNames, PropertyFlags);
 }
 
-template<typename T, typename TEnableIf<TIsDerivedFrom<T, FChaosClothAssetSimulationBaseConfigNode>::Value, int>::type>
+template<typename T UE_REQUIRES_DEFINITION(std::is_base_of_v<FChaosClothAssetSimulationBaseConfigNode, T>)>
 inline int32 FChaosClothAssetSimulationBaseConfigNode::FPropertyHelper::SetPropertyBool(
 	const T* ConfigStruct,
 	const bool* PropertyValue,
@@ -212,7 +212,7 @@ inline int32 FChaosClothAssetSimulationBaseConfigNode::FPropertyHelper::SetPrope
 	return SetPropertyBool(PropertyName, *PropertyValue, SimilarPropertyNames, PropertyFlags);
 }
 
-template<typename EnumType, typename TEnableIf<TIsEnumClass<EnumType>::Value, int>::type>
+template<typename EnumType UE_REQUIRES_DEFINITION(TIsEnumClass<EnumType>::Value)>
 inline int32 FChaosClothAssetSimulationBaseConfigNode::FPropertyHelper::SetPropertyEnum(
 	const FName& PropertyName,
 	const EnumType& PropertyValue,
@@ -224,7 +224,7 @@ inline int32 FChaosClothAssetSimulationBaseConfigNode::FPropertyHelper::SetPrope
 	return PropertyKeyIndex;
 }
 
-template<typename T, typename EnumType, typename TEnableIf<TIsDerivedFrom<T, FChaosClothAssetSimulationBaseConfigNode>::Value && TIsEnumClass<EnumType>::Value, int>::type>
+template<typename T, typename EnumType UE_REQUIRES_DEFINITION(std::is_base_of_v<FChaosClothAssetSimulationBaseConfigNode, T> && TIsEnumClass<EnumType>::Value)>
 inline int32 FChaosClothAssetSimulationBaseConfigNode::FPropertyHelper::SetPropertyEnum(
 	const T* ConfigStruct,
 	const EnumType* PropertyValue,
@@ -237,7 +237,7 @@ inline int32 FChaosClothAssetSimulationBaseConfigNode::FPropertyHelper::SetPrope
 	return SetPropertyEnum(PropertyName, *PropertyValue, SimilarPropertyNames, PropertyFlags);
 }
 
-template<typename T, typename TEnableIf<TIsDerivedFrom<T, FChaosClothAssetSimulationBaseConfigNode>::Value, int>::type>
+template<typename T UE_REQUIRES_DEFINITION(std::is_base_of_v<FChaosClothAssetSimulationBaseConfigNode, T>)>
 inline int32 FChaosClothAssetSimulationBaseConfigNode::FPropertyHelper::SetPropertyString(
 	const T* ConfigStruct,
 	const FString* PropertyValue,
@@ -250,7 +250,7 @@ inline int32 FChaosClothAssetSimulationBaseConfigNode::FPropertyHelper::SetPrope
 	return SetPropertyString(PropertyName, *PropertyValue, SimilarPropertyNames, PropertyFlags);
 }
 
-template<typename T, typename WeightedValueType, typename TEnableIf<TIsDerivedFrom<T, FChaosClothAssetSimulationBaseConfigNode>::Value, int>::type>
+template<typename T, typename WeightedValueType UE_REQUIRES_DEFINITION(std::is_base_of_v<FChaosClothAssetSimulationBaseConfigNode, T>)>
 inline int32 FChaosClothAssetSimulationBaseConfigNode::FPropertyHelper::SetPropertyWeighted(
 	const T* ConfigStruct,
 	const WeightedValueType* PropertyValue,
@@ -263,7 +263,7 @@ inline int32 FChaosClothAssetSimulationBaseConfigNode::FPropertyHelper::SetPrope
 	return SetPropertyWeighted(PropertyName, *PropertyValue, SimilarPropertyNames, PropertyFlags);
 }
 
-template<typename T, typename ConnectableStringValueType, typename TEnableIf<TIsDerivedFrom<T, FChaosClothAssetSimulationBaseConfigNode>::Value, int>::type>
+template<typename T, typename ConnectableStringValueType UE_REQUIRES_DEFINITION(std::is_base_of_v<FChaosClothAssetSimulationBaseConfigNode, T>)>
 inline int32 FChaosClothAssetSimulationBaseConfigNode::FPropertyHelper::SetPropertyString(
 	const T* ConfigStruct, 
 	const ConnectableStringValueType* PropertyValue,
@@ -276,7 +276,7 @@ inline int32 FChaosClothAssetSimulationBaseConfigNode::FPropertyHelper::SetPrope
 	return SetPropertyString(PropertyName, *PropertyValue, SimilarPropertyNames, PropertyFlags);
 }
 
-template<typename T, typename TEnableIf<TIsDerivedFrom<T, FChaosClothAssetSimulationBaseConfigNode>::Value, int>::type>
+template<typename T UE_REQUIRES_DEFINITION(std::is_base_of_v<FChaosClothAssetSimulationBaseConfigNode, T>)>
 inline FName FChaosClothAssetSimulationBaseConfigNode::FPropertyHelper::FindPropertyNameByAddress(const T* ConfigStruct, const void* PropertyAddress)
 {
 	for (TFieldIterator<FProperty> FieldIt(T::StaticStruct()); FieldIt; ++FieldIt)
diff --git a/Engine/Plugins/Experimental/AVCodecs/AVCodecsCore/Source/AVCodecsCore/Public/AVExtension.h b/Engine/Plugins/Experimental/AVCodecs/AVCodecsCore/Source/AVCodecsCore/Public/AVExtension.h
--- a/Engine/Plugins/Experimental/AVCodecs/AVCodecsCore/Source/AVCodecsCore/Public/AVExtension.h
+++ b/Engine/Plugins/Experimental/AVCodecs/AVCodecsCore/Source/AVCodecsCore/Public/AVExtension.h
@@ -72,13 +72,12 @@ public:
 	 * @param InResource The resource to read from.
 	 * @return Result of the operation, @see FAVResult.
 	 */
-	template <typename TOutput, typename TInput, TEMPLATE_REQUIRES(!TIsDerivedFrom<TInput, TOutput>::Value)>
-	static FAVResult TransformResource(TSharedPtr<TOutput>& OutResource, TSharedPtr<TInput> const& InResource) = delete;
-
 	// Shortcut case for above for when the input and output types match. This does NOT duplicate the resource.
-	template <typename TOutput, typename TInput, TEMPLATE_REQUIRES(TIsDerivedFrom<TInput, TOutput>::Value)>
+	template <typename TOutput, typename TInput>
 	static FAVResult TransformResource(TSharedPtr<TOutput>& OutResource, TSharedPtr<TInput> const& InResource)
 	{
+		static_assert(std::is_base_of_v<TOutput, TInput>, "Input resource must be inherited from Output resource");
+
 		OutResource = InResource;
 
 		return EAVResult::Success;
diff --git a/Engine/Plugins/Experimental/AVCodecs/AVCodecsCore/Source/AVCodecsCore/Public/AVUtility.h b/Engine/Plugins/Experimental/AVCodecs/AVCodecsCore/Source/AVCodecsCore/Public/AVUtility.h
--- a/Engine/Plugins/Experimental/AVCodecs/AVCodecsCore/Source/AVCodecsCore/Public/AVUtility.h
+++ b/Engine/Plugins/Experimental/AVCodecs/AVCodecsCore/Source/AVCodecsCore/Public/AVUtility.h
@@ -239,7 +239,7 @@ public:
 	 * @tparam TType Type of object to test.
 	 * @return True if an object of the type exists.
 	 */
-	template <typename TType, TEMPLATE_REQUIRES(TIsDerivedFrom<TType, TBaseType>::Value)>
+	template <typename TType UE_REQUIRES(std::is_convertible_v<TType*, TBaseType*>)>
 	bool Contains() const
 	{
 		return Get<TType>().IsValid();
@@ -251,7 +251,7 @@ public:
 	 * @tparam TType Type of object to find.
 	 * @return Typed object if it exists in the map, invalid if it does not.
 	 */
-	template <typename TType, TEMPLATE_REQUIRES(TIsDerivedFrom<TType, TBaseType>::Value)>
+	template <typename TType UE_REQUIRES(std::is_convertible_v<TType*, TBaseType*>)>
 	TSharedPtr<TType> const& Get() const
 	{
 		// Can't reinterpret cast, but we know this is safe
@@ -270,7 +270,7 @@ public:
 	 * @tparam TType Type of object to add.
 	 * @param NewObject Object to add.
 	 */
-	template <typename TType, TEMPLATE_REQUIRES(TIsDerivedFrom<TType, TBaseType>::Value)>
+	template <typename TType UE_REQUIRES(std::is_convertible_v<TType*, TBaseType*>)>
 	void Set(TSharedPtr<TType> const& NewObject)
 	{
 		Objects.Add(FTypeID::Get<TType>(), NewObject);
@@ -282,7 +282,7 @@ public:
 	 * @tparam TType Type of the object to edit.
 	 * @return Shared pointer reference to the mapping, invalid if the object did not already exist.
 	 */
-	template <typename TType, TEMPLATE_REQUIRES(TIsDerivedFrom<TType, TBaseType>::Value)>
+	template <typename TType UE_REQUIRES(std::is_convertible_v<TType*, TBaseType*>)>
 	TSharedPtr<TType>& Edit()
 	{
 		// Can't reinterpret cast, but we know this is safe
@@ -294,7 +294,7 @@ public:
 	 *
 	 * @tparam TType Type of object to add.
 	 */
-	template <typename TType, TEMPLATE_REQUIRES(TIsDerivedFrom<TType, TBaseType>::Value)>
+	template <typename TType UE_REQUIRES(std::is_convertible_v<TType*, TBaseType*>)>
 	void Remove(TSharedPtr<TType> const& NewContext)
 	{
 		Objects.Remove(FTypeID::Get<TType>());
@@ -319,7 +319,7 @@ public:
 	 * @tparam TAPI Type of API to find, which must inherit from FAPI.
 	 * @return Direct reference to the typed API.
 	 */
-	template <typename TAPI, TEMPLATE_REQUIRES(TIsDerivedFrom<TAPI, FAPI>::Value)>
+	template <typename TAPI UE_REQUIRES(std::is_base_of_v<FAPI, TAPI>)>
 	static TAPI const& Get()
 	{
 		TSharedPtr<TAPI>& Singleton = Singletons.Edit<TAPI>();
@@ -336,7 +336,7 @@ public:
 	 *
 	 * @tparam TAPI Type of the API to clean up, which must inherit from FAPI.
 	 */
-	template <typename TAPI, TEMPLATE_REQUIRES(TIsDerivedFrom<TAPI, FAPI>::Value)>
+	template <typename TAPI UE_REQUIRES(std::is_base_of_v<FAPI, TAPI>)>
 	static void Shutdown()
 	{
 		Singletons.Edit<TAPI>().Reset();
diff --git a/Engine/Plugins/Runtime/MassEntity/Source/MassEntity/Public/MassEntityView.h b/Engine/Plugins/Runtime/MassEntity/Source/MassEntity/Public/MassEntityView.h
--- a/Engine/Plugins/Runtime/MassEntity/Source/MassEntity/Public/MassEntityView.h
+++ b/Engine/Plugins/Runtime/MassEntity/Source/MassEntity/Public/MassEntityView.h
@@ -48,9 +48,9 @@ struct MASSENTITY_API FMassEntityView
 	template<typename T>
 	T& GetFragmentData() const
 	{
-		static_assert(!TIsDerivedFrom<T, FMassTag>::IsDerived,
+		static_assert(!std::is_base_of_v<FMassTag, T>,
 			"Given struct doesn't represent a valid fragment type but a tag. Use HasTag instead.");
-		static_assert(TIsDerivedFrom<T, FMassTag>::IsDerived || TIsDerivedFrom<T, FMassFragment>::IsDerived,
+		static_assert(std::is_base_of_v<FMassTag, T> || std::is_base_of_v<FMassFragment, T>,
 			"Given struct doesn't represent a valid fragment type. Make sure to inherit from FMassFragment or one of its child-types.");
 
 		return *((T*)GetFragmentPtrChecked(*T::StaticStruct()));
@@ -60,9 +60,9 @@ struct MASSENTITY_API FMassEntityView
 	template<typename T>
 	T* GetFragmentDataPtr() const
 	{
-		static_assert(!TIsDerivedFrom<T, FMassTag>::IsDerived,
+		static_assert(!std::is_base_of_v<FMassTag, T>,
 			"Given struct doesn't represent a valid fragment type but a tag. Use HasTag instead.");
-		static_assert(TIsDerivedFrom<T, FMassTag>::IsDerived || TIsDerivedFrom<T, FMassFragment>::IsDerived,
+		static_assert(std::is_base_of_v<FMassTag, T> || std::is_base_of_v<FMassFragment, T>,
 			"Given struct doesn't represent a valid fragment type. Make sure to inherit from FMassFragment or one of its child-types.");
 
 		return (T*)GetFragmentPtr(*T::StaticStruct());
@@ -78,7 +78,7 @@ struct MASSENTITY_API FMassEntityView
 	template<typename T>
 	const T* GetConstSharedFragmentDataPtr() const
 	{
-		static_assert(TIsDerivedFrom<T, FMassConstSharedFragment>::IsDerived,
+		static_assert(std::is_base_of_v<FMassConstSharedFragment, T>,
 			"Given struct doesn't represent a valid shared fragment type. Make sure to inherit from FMassConstSharedFragment or one of its child-types.");
 
 		return (const T*)GetConstSharedFragmentPtr(*T::StaticStruct());
@@ -88,7 +88,7 @@ struct MASSENTITY_API FMassEntityView
 	template<typename T>
 	const T& GetConstSharedFragmentData() const
 	{
-		static_assert(TIsDerivedFrom<T, FMassConstSharedFragment>::IsDerived,
+		static_assert(std::is_base_of_v<FMassConstSharedFragment, T>,
 			"Given struct doesn't represent a valid const shared fragment type. Make sure to inherit from FMassConstSharedFragment or one of its child-types.");
 
 		return *((const T*)GetConstSharedFragmentPtrChecked(*T::StaticStruct()));
@@ -101,27 +101,27 @@ struct MASSENTITY_API FMassEntityView
 	}
 
 	/** will fail a check if the viewed entity doesn't have the given shared fragment */
-	template<typename T, TEMPLATE_REQUIRES(TIsDerivedFrom<T, FMassSharedFragment>::Value)>
+	template<typename T UE_REQUIRES(std::is_base_of_v<FMassSharedFragment, T>)>
 	T& GetSharedFragmentData() const
 	{
 		return *((T*)GetSharedFragmentPtrChecked(*T::StaticStruct()));
 	}
 
 	/** if the viewed entity doesn't have the given shared fragment the function will return null */
-	template<typename T, TEMPLATE_REQUIRES(TIsDerivedFrom<T, FMassSharedFragment>::Value)>
+	template<typename T UE_REQUIRES(std::is_base_of_v<FMassSharedFragment, T>)>
 	T* GetSharedFragmentDataPtr() const
 	{
 		return (T*)GetSharedFragmentPtr(*T::StaticStruct());
 	}
 
-	template<typename T, TEMPLATE_REQUIRES(TIsDerivedFrom<T, FMassConstSharedFragment>::Value)>
+	template<typename T UE_REQUIRES(std::is_base_of_v<FMassConstSharedFragment, T>)>
 	UE_DEPRECATED(5.5, "Using GetSharedFragmentDataPtr with const shared fragments is deprecated. Use GetConstSharedFragmentDataPtr instead")
 	T* GetSharedFragmentDataPtr() const
 	{
 		return const_cast<T*>(GetConstSharedFragmentDataPtr<T>());
 	}
 
-	template<typename T, TEMPLATE_REQUIRES(TIsDerivedFrom<T, FMassConstSharedFragment>::Value)>
+	template<typename T UE_REQUIRES(std::is_base_of_v<FMassConstSharedFragment, T>)>
 	UE_DEPRECATED(5.5, "Using GetSharedFragmentDataPtr with const shared fragments is deprecated. Use GetConstSharedFragmentData instead")
 	T& GetSharedFragmentData() const
 	{
@@ -139,7 +139,7 @@ struct MASSENTITY_API FMassEntityView
 	template<typename T>
 	bool HasTag() const
 	{
-		static_assert(TIsDerivedFrom<T, FMassTag>::IsDerived, "Given struct doesn't represent a valid tag type. Make sure to inherit from FMassTag or one of its child-types.");
+		static_assert(std::is_base_of_v<FMassTag, T>, "Given struct doesn't represent a valid tag type. Make sure to inherit from FMassTag or one of its child-types.");
 		return HasTag(*T::StaticStruct());
 	}
 
diff --git a/Engine/Plugins/VirtualProduction/RemoteControl/Source/RemoteControl/Public/IRemoteControlModule.h b/Engine/Plugins/VirtualProduction/RemoteControl/Source/RemoteControl/Public/IRemoteControlModule.h
--- a/Engine/Plugins/VirtualProduction/RemoteControl/Source/RemoteControl/Public/IRemoteControlModule.h
+++ b/Engine/Plugins/VirtualProduction/RemoteControl/Source/RemoteControl/Public/IRemoteControlModule.h
@@ -614,7 +614,7 @@ public:
 	virtual bool CanBeAccessedRemotely(UObject* InObject) const = 0;
 
 	template<class InPropertyIdPropertyHandlerType, typename... InArgTypes
-		, TEMPLATE_REQUIRES(TIsDerivedFrom<InPropertyIdPropertyHandlerType, IPropertyIdHandler>::Value)>
+		UE_REQUIRES(std::is_base_of_v<IPropertyIdHandler, InPropertyIdPropertyHandlerType>)>
 	TSharedRef<InPropertyIdPropertyHandlerType> RegisterPropertyIdPropertyHandler(InArgTypes&&... InArgs)
 	{
 		TSharedRef<InPropertyIdPropertyHandlerType> KeyPropertyHandler = MakeShared<InPropertyIdPropertyHandlerType>(Forward<InArgTypes>(InArgs)...);
diff --git a/Engine/Source/Runtime/ApplicationCore/Public/GenericPlatform/GenericPlatformApplicationMisc.h b/Engine/Source/Runtime/ApplicationCore/Public/GenericPlatform/GenericPlatformApplicationMisc.h
--- a/Engine/Source/Runtime/ApplicationCore/Public/GenericPlatform/GenericPlatformApplicationMisc.h
+++ b/Engine/Source/Runtime/ApplicationCore/Public/GenericPlatform/GenericPlatformApplicationMisc.h
@@ -311,7 +311,7 @@ struct FGenericPlatformApplicationMisc
 	 * If we know or can approximate the pixel density of the screen we will convert the incoming inches
 	 * to pixels on the device.  If the accuracy is unknown OutPixels will be set to 0.
 	 */
-	template<typename T, typename T2, TEMPLATE_REQUIRES(TIsFloatingPoint<T>::Value && TIsFloatingPoint<T2>::Value)>
+	template<typename T, typename T2 UE_REQUIRES(std::is_floating_point_v<T> && std::is_floating_point_v<T2>)>
 	static EScreenPhysicalAccuracy ConvertInchesToPixels(T Inches, T2& OutPixels)
 	{
 		int32 ScreenDensity = 0;
@@ -333,7 +333,7 @@ struct FGenericPlatformApplicationMisc
 	 * If we know or can approximate the pixel density of the screen we will convert the incoming pixels
 	 * to inches on the device.  If the accuracy is unknown OutInches will be set to 0.
 	 */
-	template<typename T, typename T2, TEMPLATE_REQUIRES(TIsFloatingPoint<T>::Value && TIsFloatingPoint<T2>::Value)>
+	template<typename T, typename T2 UE_REQUIRES(std::is_floating_point_v<T> && std::is_floating_point_v<T2>)>
 	static EScreenPhysicalAccuracy ConvertPixelsToInches(T Pixels, T2& OutInches)
 	{
 		int32 ScreenDensity = 0;
diff --git a/Engine/Source/Runtime/Core/Public/Math/Box.h b/Engine/Source/Runtime/Core/Public/Math/Box.h
--- a/Engine/Source/Runtime/Core/Public/Math/Box.h
+++ b/Engine/Source/Runtime/Core/Public/Math/Box.h
@@ -158,8 +158,11 @@ public:
 	explicit TBox(const TArray<TVector<T>>& Points) : TBox<T>(&Points[0], Points.Num()) {};
 
 	// Conversion from other type.
-	template<typename FArg, TEMPLATE_REQUIRES(!std::is_same_v<T, FArg>)>
-	explicit TBox(const TBox<FArg>& From) : TBox<T>(TVector<T>(From.Min), TVector<T>(From.Max)) {}
+	template<typename FArg UE_REQUIRES(!std::is_same_v<T, FArg>)>
+	explicit TBox(const TBox<FArg>& From)
+		: TBox<T>(TVector<T>(From.Min), TVector<T>(From.Max))
+	{
+	}
 
 public:
 
diff --git a/Engine/Source/Runtime/Core/Public/Math/Box2D.h b/Engine/Source/Runtime/Core/Public/Math/Box2D.h
--- a/Engine/Source/Runtime/Core/Public/Math/Box2D.h
+++ b/Engine/Source/Runtime/Core/Public/Math/Box2D.h
@@ -416,8 +416,13 @@ public:
 	bool SerializeFromMismatchedTag(FName StructTag, FArchive& Ar);
 
 	// Conversion from other type. 
-	template<typename FArg, TEMPLATE_REQUIRES(!std::is_same_v<T, FArg>)>
-	explicit TBox2(const TBox2<FArg>& From) : Min((TVector2<T>)From.Min), Max((TVector2<T>)From.Max), bIsValid(From.bIsValid) {}
+	template<typename FArg UE_REQUIRES(!std::is_same_v<T, FArg>)>
+	explicit TBox2(const TBox2<FArg>& From)
+		: Min((TVector2<T>)From.Min)
+		, Max((TVector2<T>)From.Max)
+		, bIsValid(From.bIsValid)
+	{
+	}
 };
 
 
diff --git a/Engine/Source/Runtime/Core/Public/Math/BoxSphereBounds.h b/Engine/Source/Runtime/Core/Public/Math/BoxSphereBounds.h
--- a/Engine/Source/Runtime/Core/Public/Math/BoxSphereBounds.h
+++ b/Engine/Source/Runtime/Core/Public/Math/BoxSphereBounds.h
@@ -124,8 +124,11 @@ public:
 	TBoxSphereBounds( const TVector<T>* Points, uint32 NumPoints );
 
 	// Conversion to other type.
-	template<typename TFrom, typename TExtentFrom, TEMPLATE_REQUIRES(!(std::is_same_v<T, TFrom> && std::is_same_v<TExtent, TExtentFrom>))>
-	explicit TBoxSphereBounds(const TBoxSphereBounds<TFrom, TExtentFrom>& From) : TBoxSphereBounds<T, TExtent>(TVector<T>(From.Origin), TVector<TExtent>(From.BoxExtent), (TExtent)From.SphereRadius) {}
+	template<typename TFrom, typename TExtentFrom UE_REQUIRES(!(std::is_same_v<T, TFrom> && std::is_same_v<TExtent, TExtentFrom>))>
+	explicit TBoxSphereBounds(const TBoxSphereBounds<TFrom, TExtentFrom>& From)
+		: TBoxSphereBounds<T, TExtent>(TVector<T>(From.Origin), TVector<TExtent>(From.BoxExtent), (TExtent)From.SphereRadius)
+	{
+	}
 
 public:
 
diff --git a/Engine/Source/Runtime/Core/Public/Math/CapsuleShape.h b/Engine/Source/Runtime/Core/Public/Math/CapsuleShape.h
--- a/Engine/Source/Runtime/Core/Public/Math/CapsuleShape.h
+++ b/Engine/Source/Runtime/Core/Public/Math/CapsuleShape.h
@@ -53,8 +53,11 @@ public:
 	{ }
 
 	// Conversion to other type.
-	template<typename FArg, TEMPLATE_REQUIRES(!std::is_same_v<T, FArg>)>
-	explicit TCapsuleShape(const TCapsuleShape<FArg>& From) : TCapsuleShape<T>(TVector<T>(From.Center), (T)From.Radius, TVector<T>(From.Orientation), (T)From.Length) {}
+	template<typename FArg UE_REQUIRES(!std::is_same_v<T, FArg>)>
+	explicit TCapsuleShape(const TCapsuleShape<FArg>& From)
+		: TCapsuleShape<T>(TVector<T>(From.Center), (T)From.Radius, TVector<T>(From.Orientation), (T)From.Length)
+	{
+	}
 };
 
 }	// namespace UE::Math
diff --git a/Engine/Source/Runtime/Core/Public/Math/ClipProjectionMatrix.h b/Engine/Source/Runtime/Core/Public/Math/ClipProjectionMatrix.h
--- a/Engine/Source/Runtime/Core/Public/Math/ClipProjectionMatrix.h
+++ b/Engine/Source/Runtime/Core/Public/Math/ClipProjectionMatrix.h
@@ -29,8 +29,11 @@ public:
 	TClipProjectionMatrix( const TMatrix<T>& SrcProjMat, const TPlane<T>& Plane );
 
 	// Conversion to other type.
-	template<typename FArg, TEMPLATE_REQUIRES(!std::is_same_v<T, FArg>)>
-	explicit TClipProjectionMatrix(const TClipProjectionMatrix<FArg>& From) : TMatrix<T>(From) {}
+	template<typename FArg UE_REQUIRES(!std::is_same_v<T, FArg>)>
+	explicit TClipProjectionMatrix(const TClipProjectionMatrix<FArg>& From)
+		: TMatrix<T>(From)
+	{
+	}
 	
 private:
 	/** return sign of a number */
diff --git a/Engine/Source/Runtime/Core/Public/Math/DualQuat.h b/Engine/Source/Runtime/Core/Public/Math/DualQuat.h
--- a/Engine/Source/Runtime/Core/Public/Math/DualQuat.h
+++ b/Engine/Source/Runtime/Core/Public/Math/DualQuat.h
@@ -38,8 +38,11 @@ public:
 	}
 
 	// Conversion to other type.
-	template<typename FArg, TEMPLATE_REQUIRES(!std::is_same_v<T, FArg>)>
-	explicit TDualQuat(const TDualQuat<FArg>& From) : TDualQuat<T>(TQuat<T>(From.R), TQuat<T>(From.D)) {}
+	template<typename FArg UE_REQUIRES(!std::is_same_v<T, FArg>)>
+	explicit TDualQuat(const TDualQuat<FArg>& From)
+		: TDualQuat<T>(TQuat<T>(From.R), TQuat<T>(From.D))
+	{
+	}
 
 	/** Dual quat addition */
 	TDualQuat<T> operator+(const TDualQuat<T> &B) const
diff --git a/Engine/Source/Runtime/Core/Public/Math/InverseRotationMatrix.h b/Engine/Source/Runtime/Core/Public/Math/InverseRotationMatrix.h
--- a/Engine/Source/Runtime/Core/Public/Math/InverseRotationMatrix.h
+++ b/Engine/Source/Runtime/Core/Public/Math/InverseRotationMatrix.h
@@ -24,8 +24,11 @@ public:
 	TInverseRotationMatrix(const TRotator<T>& Rot);
 	
 	// Conversion to other type.
-    template<typename FArg, TEMPLATE_REQUIRES(!std::is_same_v<T, FArg>)>
-    explicit TInverseRotationMatrix(const TInverseRotationMatrix<FArg>& From) : TMatrix<T>(From) {}	
+	template<typename FArg UE_REQUIRES(!std::is_same_v<T, FArg>)>
+	explicit TInverseRotationMatrix(const TInverseRotationMatrix<FArg>& From)
+		: TMatrix<T>(From)
+	{
+	}
 };
 
 template<typename T>
diff --git a/Engine/Source/Runtime/Core/Public/Math/Matrix.h b/Engine/Source/Runtime/Core/Public/Math/Matrix.h
--- a/Engine/Source/Runtime/Core/Public/Math/Matrix.h
+++ b/Engine/Source/Runtime/Core/Public/Math/Matrix.h
@@ -407,7 +407,7 @@ public:
 	bool SerializeFromMismatchedTag(FName StructTag, FArchive& Ar);
 
 	// Conversion to other type.
-	template<typename FArg, TEMPLATE_REQUIRES(!std::is_same_v<T, FArg>)>
+	template<typename FArg UE_REQUIRES(!std::is_same_v<T, FArg>)>
 	explicit TMatrix(const TMatrix<FArg>& From)
 	{
 		// TODO: SIMD this?
@@ -496,8 +496,11 @@ struct TBasisVectorMatrix : public TMatrix<T>
 	TBasisVectorMatrix(const TVector<T>& XAxis,const TVector<T>& YAxis,const TVector<T>& ZAxis,const TVector<T>& Origin);
 
 	// Conversion to other type.
-	template<typename FArg, TEMPLATE_REQUIRES(!std::is_same_v<T, FArg>)>
-	explicit TBasisVectorMatrix(const TBasisVectorMatrix<FArg>& From) : TMatrix<T>(From) {}
+	template<typename FArg UE_REQUIRES(!std::is_same_v<T, FArg>)>
+	explicit TBasisVectorMatrix(const TBasisVectorMatrix<FArg>& From)
+		: TMatrix<T>(From)
+	{
+	}
 };
 
 
@@ -515,8 +518,11 @@ struct TLookFromMatrix : public TMatrix<T>
 	TLookFromMatrix(const TVector<T>& EyePosition, const TVector<T>& LookDirection, const TVector<T>& UpVector);
 
 	// Conversion to other type.
-	template<typename FArg, TEMPLATE_REQUIRES(!std::is_same_v<T, FArg>)>
-	explicit TLookFromMatrix(const TLookFromMatrix<FArg>& From) : TMatrix<T>(From) {}
+	template<typename FArg UE_REQUIRES(!std::is_same_v<T, FArg>)>
+	explicit TLookFromMatrix(const TLookFromMatrix<FArg>& From)
+		: TMatrix<T>(From)
+	{
+	}
 };
 
 
@@ -534,8 +540,11 @@ struct TLookAtMatrix : public TLookFromMatrix<T>
 	TLookAtMatrix(const TVector<T>& EyePosition, const TVector<T>& LookAtPosition, const TVector<T>& UpVector);
 
 	// Conversion to other type.
-	template<typename FArg, TEMPLATE_REQUIRES(!std::is_same_v<T, FArg>)>
-	explicit TLookAtMatrix(const TLookAtMatrix<FArg>& From) : TLookFromMatrix<T>(From) {}
+	template<typename FArg UE_REQUIRES(!std::is_same_v<T, FArg>)>
+	explicit TLookAtMatrix(const TLookAtMatrix<FArg>& From)
+		: TLookFromMatrix<T>(From)
+	{
+	}
 };
 
 } // namespace UE::Math
@@ -566,7 +575,7 @@ template<> CORE_API FQuat4d FMatrix44d::ToQuat() const;
 
 // very high quality 4x4 matrix inverse
 // @todo: this is redundant with FMatrix44d::Inverse and should be removed ; seems to be unused
-template<typename FArg, TEMPLATE_REQUIRES(std::is_floating_point<FArg>::value)>
+template<typename FArg UE_REQUIRES(std::is_floating_point_v<FArg>)>
 static inline bool Inverse4x4( double* dst, const FArg* src )
 {
 	const double s0  = (double)(src[ 0]); const double s1  = (double)(src[ 1]); const double s2  = (double)(src[ 2]); const double s3  = (double)(src[ 3]);
diff --git a/Engine/Source/Runtime/Core/Public/Math/MirrorMatrix.h b/Engine/Source/Runtime/Core/Public/Math/MirrorMatrix.h
--- a/Engine/Source/Runtime/Core/Public/Math/MirrorMatrix.h
+++ b/Engine/Source/Runtime/Core/Public/Math/MirrorMatrix.h
@@ -26,8 +26,11 @@ public:
 	TMirrorMatrix( const TPlane<T>& Plane );
 
 	// Conversion to other type.
-	template<typename FArg, TEMPLATE_REQUIRES(!std::is_same_v<T, FArg>)>
-	explicit TMirrorMatrix(const TMirrorMatrix<FArg>& From) : TMatrix<T>(From) {}
+	template<typename FArg UE_REQUIRES(!std::is_same_v<T, FArg>)>
+	explicit TMirrorMatrix(const TMirrorMatrix<FArg>& From)
+		: TMatrix<T>(From)
+	{
+	}
 };
 
 template<typename T>
diff --git a/Engine/Source/Runtime/Core/Public/Math/OrthoMatrix.h b/Engine/Source/Runtime/Core/Public/Math/OrthoMatrix.h
--- a/Engine/Source/Runtime/Core/Public/Math/OrthoMatrix.h
+++ b/Engine/Source/Runtime/Core/Public/Math/OrthoMatrix.h
@@ -27,8 +27,11 @@ public:
 	TOrthoMatrix(T Width,T Height,T ZScale,T ZOffset);
 	
 	// Conversion to other type.
-	template<typename FArg, TEMPLATE_REQUIRES(!std::is_same_v<T, FArg>)>
-	explicit TOrthoMatrix(const TOrthoMatrix<FArg>& From) : TMatrix<T>(From) {}
+	template<typename FArg UE_REQUIRES(!std::is_same_v<T, FArg>)>
+	explicit TOrthoMatrix(const TOrthoMatrix<FArg>& From)
+		: TMatrix<T>(From)
+	{
+	}
 };
 
 
@@ -40,8 +43,11 @@ public:
 	TReversedZOrthoMatrix(T Left, T Right, T Bottom, T Top, T ZScale, T ZOffset);
 
 	// Conversion to other type.
-	template<typename FArg, TEMPLATE_REQUIRES(!std::is_same_v<T, FArg>)>
-	explicit TReversedZOrthoMatrix(const TReversedZOrthoMatrix<FArg>& From) : TMatrix<T>(From) {}
+	template<typename FArg UE_REQUIRES(!std::is_same_v<T, FArg>)>
+	explicit TReversedZOrthoMatrix(const TReversedZOrthoMatrix<FArg>& From)
+		: TMatrix<T>(From)
+	{
+	}
 };
 
 template<typename T>
diff --git a/Engine/Source/Runtime/Core/Public/Math/PerspectiveMatrix.h b/Engine/Source/Runtime/Core/Public/Math/PerspectiveMatrix.h
--- a/Engine/Source/Runtime/Core/Public/Math/PerspectiveMatrix.h
+++ b/Engine/Source/Runtime/Core/Public/Math/PerspectiveMatrix.h
@@ -55,8 +55,11 @@ public:
 	TPerspectiveMatrix(T HalfFOV, T Width, T Height, T MinZ);
 
 	// Conversion to other type.
-	template<typename FArg, TEMPLATE_REQUIRES(!std::is_same_v<T, FArg>)>
-	explicit TPerspectiveMatrix(const TPerspectiveMatrix<FArg>& From) : TMatrix<T>(From) {}
+	template<typename FArg UE_REQUIRES(!std::is_same_v<T, FArg>)>
+	explicit TPerspectiveMatrix(const TPerspectiveMatrix<FArg>& From)
+		: TMatrix<T>(From)
+	{
+	}
 };
 
 
@@ -69,8 +72,11 @@ public:
 	TReversedZPerspectiveMatrix(T HalfFOV, T Width, T Height, T MinZ);
 	
 	// Conversion to other type.
-	template<typename FArg, TEMPLATE_REQUIRES(!std::is_same_v<T, FArg>)>
-	explicit TReversedZPerspectiveMatrix(const TReversedZPerspectiveMatrix<FArg>& From) : TMatrix<T>(From) {}
+	template<typename FArg UE_REQUIRES(!std::is_same_v<T, FArg>)>
+	explicit TReversedZPerspectiveMatrix(const TReversedZPerspectiveMatrix<FArg>& From)
+		: TMatrix<T>(From)
+	{
+	}
 };
 
 
diff --git a/Engine/Source/Runtime/Core/Public/Math/Plane.h b/Engine/Source/Runtime/Core/Public/Math/Plane.h
--- a/Engine/Source/Runtime/Core/Public/Math/Plane.h
+++ b/Engine/Source/Runtime/Core/Public/Math/Plane.h
@@ -336,8 +336,11 @@ public:
 	}
 
 	// Conversion to other type.
-	template<typename FArg, TEMPLATE_REQUIRES(!std::is_same_v<T, FArg>)>
-	explicit TPlane(const TPlane<FArg>& From) : TPlane<T>((T)From.X, (T)From.Y, (T)From.Z, (T)From.W) {}
+	template<typename FArg UE_REQUIRES(!std::is_same_v<T, FArg>)>
+	explicit TPlane(const TPlane<FArg>& From)
+		: TPlane<T>((T)From.X, (T)From.Y, (T)From.Z, (T)From.W)
+	{
+	}
 };
 
 
diff --git a/Engine/Source/Runtime/Core/Public/Math/Quat.h b/Engine/Source/Runtime/Core/Public/Math/Quat.h
--- a/Engine/Source/Runtime/Core/Public/Math/Quat.h
+++ b/Engine/Source/Runtime/Core/Public/Math/Quat.h
@@ -36,8 +36,8 @@ namespace Math
 template<typename T>
 struct alignas(16) TQuat
 {
-	// Can't have a TEMPLATE_REQUIRES in the declaration because of the forward declarations, so check for allowed types here.
-	static_assert(TIsFloatingPoint<T>::Value, "TQuat only supports float and double types.");
+	// Can't have a UE_REQUIRES in the declaration because of the forward declarations, so check for allowed types here.
+	static_assert(std::is_floating_point_v<T>, "TQuat only supports float and double types.");
 
 public:
 	/** Type of the template param (float or double) */
@@ -86,9 +86,12 @@ public:
 	/**
 	 * Initializes all elements to V
 	 */
-	template<TEMPLATE_REQUIRES(std::is_arithmetic<T>::value)>
+	template<typename DummyT = T UE_REQUIRES(std::is_arithmetic_v<T>)>
 	explicit FORCEINLINE TQuat(T V)
-	: X(V), Y(V), Z(V), W(V)
+		: X(V)
+		, Y(V)
+		, Z(V)
+		, W(V)
 	{
 		DiagnosticCheckNaN();
 	}
@@ -240,7 +243,7 @@ public:
 	 * @param Scale The scaling factor.
 	 * @return a reference to this after scaling.
 	 */
-	template<typename FArg, TEMPLATE_REQUIRES(std::is_arithmetic<FArg>::value)>
+	template<typename FArg UE_REQUIRES(std::is_arithmetic_v<FArg>)>
 	FORCEINLINE TQuat<T> operator*=(const FArg Scale)
 	{
 #if PLATFORM_ENABLE_VECTORINTRINSICS
@@ -264,7 +267,7 @@ public:
 	 * @param Scale The scaling factor.
 	 * @return The result of scaling.
 	 */
-	template<typename FArg, TEMPLATE_REQUIRES(std::is_arithmetic<FArg>::value)>
+	template<typename FArg UE_REQUIRES(std::is_arithmetic_v<FArg>)>
 	FORCEINLINE TQuat<T> operator*(const FArg Scale) const
 	{
 #if PLATFORM_ENABLE_VECTORINTRINSICS
@@ -282,7 +285,7 @@ public:
 	 * @param Scale What to divide by.
 	 * @return a reference to this after scaling.
 	 */
-	template<typename FArg, TEMPLATE_REQUIRES(std::is_arithmetic<FArg>::value)>
+	template<typename FArg UE_REQUIRES(std::is_arithmetic_v<FArg>)>
 	FORCEINLINE TQuat<T> operator/=(const FArg Scale)
 	{
 #if PLATFORM_ENABLE_VECTORINTRINSICS
@@ -307,7 +310,7 @@ public:
 	 * @param Scale What to divide by.
 	 * @return new Quaternion of this after division by scale.
 	 */
-	template<typename FArg, TEMPLATE_REQUIRES(std::is_arithmetic<FArg>::value)>
+	template<typename FArg UE_REQUIRES(std::is_arithmetic_v<FArg>)>
 	FORCEINLINE TQuat<T> operator/(const FArg Scale) const
 	{
 #if PLATFORM_ENABLE_VECTORINTRINSICS
@@ -688,8 +691,11 @@ public:
 	bool SerializeFromMismatchedTag(FName StructTag, FArchive& Ar);
 
 	// Conversion to other type.
-	template<typename FArg, TEMPLATE_REQUIRES(!std::is_same_v<T, FArg>)>
-	explicit TQuat(const TQuat<FArg>& From) : TQuat<T>((T)From.X, (T)From.Y, (T)From.Z, (T)From.W) {}
+	template<typename FArg UE_REQUIRES(!std::is_same_v<T, FArg>)>
+	explicit TQuat(const TQuat<FArg>& From)
+		: TQuat<T>((T)From.X, (T)From.Y, (T)From.Z, (T)From.W)
+	{
+	}
 };
 
 /**
diff --git a/Engine/Source/Runtime/Core/Public/Math/QuatRotationTranslationMatrix.h b/Engine/Source/Runtime/Core/Public/Math/QuatRotationTranslationMatrix.h
--- a/Engine/Source/Runtime/Core/Public/Math/QuatRotationTranslationMatrix.h
+++ b/Engine/Source/Runtime/Core/Public/Math/QuatRotationTranslationMatrix.h
@@ -27,8 +27,11 @@ public:
 	TQuatRotationTranslationMatrix(const TQuat<T>& Q, const TVector<T>& Origin);
 
 	// Conversion to other type.
-	template<typename FArg, TEMPLATE_REQUIRES(!std::is_same_v<T, FArg>)>
-	explicit TQuatRotationTranslationMatrix(const TQuatRotationTranslationMatrix<FArg>& From) : TMatrix<T>(From) {}
+	template<typename FArg UE_REQUIRES(!std::is_same_v<T, FArg>)>
+	explicit TQuatRotationTranslationMatrix(const TQuatRotationTranslationMatrix<FArg>& From)
+		: TMatrix<T>(From)
+	{
+	}
 
 	/** Matrix factory. Return an FMatrix so we don't have type conversion issues in expressions. */
 	static TMatrix<T> Make(const TQuat<T>& Q, const TVector<T>& Origin)
@@ -56,8 +59,11 @@ public:
 	}
 
 	// Conversion to other type.
-	template<typename FArg, TEMPLATE_REQUIRES(!std::is_same_v<T, FArg>)>
-	explicit TQuatRotationMatrix(const TQuatRotationMatrix<FArg>& From) : TQuatRotationTranslationMatrix<T>(From) {}
+	template<typename FArg UE_REQUIRES(!std::is_same_v<T, FArg>)>
+	explicit TQuatRotationMatrix(const TQuatRotationMatrix<FArg>& From)
+		: TQuatRotationTranslationMatrix<T>(From)
+	{
+	}
 	
 	/** Matrix factory. Return an FMatrix so we don't have type conversion issues in expressions. */
 	static TMatrix<T> Make(const TQuat<T>& Q)
diff --git a/Engine/Source/Runtime/Core/Public/Math/Ray.h b/Engine/Source/Runtime/Core/Public/Math/Ray.h
--- a/Engine/Source/Runtime/Core/Public/Math/Ray.h
+++ b/Engine/Source/Runtime/Core/Public/Math/Ray.h
@@ -203,8 +203,11 @@ public:
 	bool SerializeFromMismatchedTag(FName StructTag, FArchive& Ar);
 
 	// Conversion to other type.
-	template<typename FArg, TEMPLATE_REQUIRES(!std::is_same_v<T, FArg>)>
-	explicit TRay(const TRay<FArg>& From) : TRay<T>(TVector<T>(From.Origin), TVector<T>(From.Direction), true) {}
+	template<typename FArg UE_REQUIRES(!std::is_same_v<T, FArg>)>
+	explicit TRay(const TRay<FArg>& From)
+		: TRay<T>(TVector<T>(From.Origin), TVector<T>(From.Direction), true)
+	{
+	}
 };
 
 }	// namespace UE::Math
diff --git a/Engine/Source/Runtime/Core/Public/Math/RotationAboutPointMatrix.h b/Engine/Source/Runtime/Core/Public/Math/RotationAboutPointMatrix.h
--- a/Engine/Source/Runtime/Core/Public/Math/RotationAboutPointMatrix.h
+++ b/Engine/Source/Runtime/Core/Public/Math/RotationAboutPointMatrix.h
@@ -28,9 +28,12 @@ public:
 	TRotationAboutPointMatrix(const TRotator<T>& Rot, const TVector<T>& Origin);
 
 	// Conversion to other type.
-	template<typename FArg, TEMPLATE_REQUIRES(!std::is_same_v<T, FArg>)>
-	explicit TRotationAboutPointMatrix(const TRotationAboutPointMatrix<FArg>& From) : TRotationTranslationMatrix<T>(From) {}	
-	
+	template<typename FArg UE_REQUIRES(!std::is_same_v<T, FArg>)>
+	explicit TRotationAboutPointMatrix(const TRotationAboutPointMatrix<FArg>& From)
+		: TRotationTranslationMatrix<T>(From)
+	{
+	}
+
 	/** Matrix factory. Return an TMatrix<T> so we don't have type conversion issues in expressions. */
 	static TMatrix<T> Make(const TRotator<T>& Rot, const TVector<T>& Origin)
 	{
diff --git a/Engine/Source/Runtime/Core/Public/Math/RotationTranslationMatrix.h b/Engine/Source/Runtime/Core/Public/Math/RotationTranslationMatrix.h
--- a/Engine/Source/Runtime/Core/Public/Math/RotationTranslationMatrix.h
+++ b/Engine/Source/Runtime/Core/Public/Math/RotationTranslationMatrix.h
@@ -27,8 +27,11 @@ public:
 	TRotationTranslationMatrix(const TRotator<T>& Rot, const TVector<T>& Origin);
 
 	// Conversion to other type.
-	template<typename FArg, TEMPLATE_REQUIRES(!std::is_same_v<T, FArg>)>
-	explicit TRotationTranslationMatrix(const TRotationTranslationMatrix<FArg>& From) : TMatrix<T>(From) {}
+	template<typename FArg UE_REQUIRES(!std::is_same_v<T, FArg>)>
+	explicit TRotationTranslationMatrix(const TRotationTranslationMatrix<FArg>& From)
+		: TMatrix<T>(From)
+	{
+	}
 	
 	/** Matrix factory. Return an TMatrix<T> so we don't have type conversion issues in expressions. */
 	static TMatrix<T> Make(const TRotator<T>& Rot, const TVector<T>& Origin)
diff --git a/Engine/Source/Runtime/Core/Public/Math/Rotator.h b/Engine/Source/Runtime/Core/Public/Math/Rotator.h
--- a/Engine/Source/Runtime/Core/Public/Math/Rotator.h
+++ b/Engine/Source/Runtime/Core/Public/Math/Rotator.h
@@ -35,8 +35,8 @@ template<typename T>
 struct TRotator
 {
 
-	// Can't have a TEMPLATE_REQUIRES in the declaration because of the forward declarations, so check for allowed types here.
-	static_assert(TIsFloatingPoint<T>::Value, "TRotator only supports float and double types.");
+	// Can't have a UE_REQUIRES in the declaration because of the forward declarations, so check for allowed types here.
+	static_assert(std::is_floating_point_v<T>, "TRotator only supports float and double types.");
 
 public:
 	using FReal = T;
@@ -141,7 +141,7 @@ public:
 	 * @param Scale The scaling factor.
 	 * @return The result of scaling.
 	 */
-	template<typename FArg, TEMPLATE_REQUIRES(std::is_arithmetic<FArg>::value)>
+	template<typename FArg UE_REQUIRES(std::is_arithmetic_v<FArg>)>
 	FORCEINLINE TRotator operator*( FArg Scale ) const
 	{
 		return TRotator(Pitch * Scale, Yaw * Scale, Roll * Scale);
@@ -153,7 +153,7 @@ public:
 	 * @param Scale The scaling factor.
 	 * @return Copy of the rotator after scaling.
 	 */
-	template<typename FArg, TEMPLATE_REQUIRES(std::is_arithmetic<FArg>::value)>
+	template<typename FArg UE_REQUIRES(std::is_arithmetic_v<FArg>)>
 	FORCEINLINE TRotator operator*=( FArg Scale )
 	{
 		Pitch = Pitch * Scale; Yaw = Yaw * Scale; Roll = Roll * Scale;
@@ -484,8 +484,11 @@ public:
 	bool SerializeFromMismatchedTag(FName StructTag, FArchive& Ar);
 
 	// Conversion from other type.
-	template<typename FArg, TEMPLATE_REQUIRES(!std::is_same_v<T, FArg>)>
-	explicit TRotator(const TRotator<FArg>& From) : TRotator<T>((T)From.Pitch, (T)From.Yaw, (T)From.Roll) {}
+	template<typename FArg UE_REQUIRES(!std::is_same_v<T, FArg>)>
+	explicit TRotator(const TRotator<FArg>& From)
+		: TRotator<T>((T)From.Pitch, (T)From.Yaw, (T)From.Roll)
+	{
+	}
 };
 
 #if !defined(_MSC_VER) || defined(__clang__)  // MSVC can't forward declare explicit specializations
@@ -537,7 +540,7 @@ inline FArchive& operator<<(FArchive& Ar, TRotator<double>& R)
  * @param R rotator to be scaled.
  * @return Scaled rotator.
  */
-template<typename T, typename FArg, TEMPLATE_REQUIRES(std::is_arithmetic<FArg>::value)>
+template<typename T, typename FArg UE_REQUIRES(std::is_arithmetic_v<FArg>)>
 FORCEINLINE TRotator<T> operator*(FArg Scale, const TRotator<T>& R )
 {
 	return R.operator*( Scale );
diff --git a/Engine/Source/Runtime/Core/Public/Math/ScaleMatrix.h b/Engine/Source/Runtime/Core/Public/Math/ScaleMatrix.h
--- a/Engine/Source/Runtime/Core/Public/Math/ScaleMatrix.h
+++ b/Engine/Source/Runtime/Core/Public/Math/ScaleMatrix.h
@@ -29,8 +29,11 @@ public:
 	TScaleMatrix( const TVector<T>& Scale );
 
 	// Conversion to other type.
-	template<typename FArg, TEMPLATE_REQUIRES(!std::is_same_v<T, FArg>)>
-	explicit TScaleMatrix(const TScaleMatrix<FArg>& From) : TMatrix<T>(From) {}
+	template<typename FArg UE_REQUIRES(!std::is_same_v<T, FArg>)>
+	explicit TScaleMatrix(const TScaleMatrix<FArg>& From)
+		: TMatrix<T>(From)
+	{
+	}
 	
 	/** Matrix factory. Return an TMatrix<T> so we don't have type conversion issues in expressions. */
 	static TMatrix<T> Make(T Scale)
diff --git a/Engine/Source/Runtime/Core/Public/Math/ScaleRotationTranslationMatrix.h b/Engine/Source/Runtime/Core/Public/Math/ScaleRotationTranslationMatrix.h
--- a/Engine/Source/Runtime/Core/Public/Math/ScaleRotationTranslationMatrix.h
+++ b/Engine/Source/Runtime/Core/Public/Math/ScaleRotationTranslationMatrix.h
@@ -28,8 +28,11 @@ public:
 	TScaleRotationTranslationMatrix(const TVector<T>& Scale, const TRotator<T>& Rot, const TVector<T>& Origin);
 	
 	// Conversion to other type.
-	template<typename FArg, TEMPLATE_REQUIRES(!std::is_same_v<T, FArg>)>
-	explicit TScaleRotationTranslationMatrix(const TScaleRotationTranslationMatrix<FArg>& From) : TMatrix<T>(From) {}
+	template<typename FArg UE_REQUIRES(!std::is_same_v<T, FArg>)>
+	explicit TScaleRotationTranslationMatrix(const TScaleRotationTranslationMatrix<FArg>& From)
+		: TMatrix<T>(From)
+	{
+	}
 };
 
 namespace
diff --git a/Engine/Source/Runtime/Core/Public/Math/Sphere.h b/Engine/Source/Runtime/Core/Public/Math/Sphere.h
--- a/Engine/Source/Runtime/Core/Public/Math/Sphere.h
+++ b/Engine/Source/Runtime/Core/Public/Math/Sphere.h
@@ -89,8 +89,11 @@ public:
 	CORE_API TSphere(const TSphere<T>* Spheres, int32 Count);
 
 	// Conversion from other variant type.
-	template<typename FArg, TEMPLATE_REQUIRES(!std::is_same_v<T, FArg>)>
-	explicit TSphere(const TSphere<FArg>& From) : TSphere<T>(TVector<T>(From.Center), T(From.W)) {}
+	template<typename FArg UE_REQUIRES(!std::is_same_v<T, FArg>)>
+	explicit TSphere(const TSphere<FArg>& From)
+		: TSphere<T>(TVector<T>(From.Center), T(From.W))
+	{
+	}
 
 	/**
 	 * Check whether two spheres are the same within specified tolerance.
diff --git a/Engine/Source/Runtime/Core/Public/Math/TransformCalculus.h b/Engine/Source/Runtime/Core/Public/Math/TransformCalculus.h
--- a/Engine/Source/Runtime/Core/Public/Math/TransformCalculus.h
+++ b/Engine/Source/Runtime/Core/Public/Math/TransformCalculus.h
@@ -417,7 +417,7 @@ inline VectorType TransformVector(const TransformType& Transform, const VectorTy
  * @param RHS Scale that goes from space B to space C.
  * @return a new Scale representing the transformation from the input space of LHS to the output space of RHS.
  */
-template<typename FloatType, TEMPLATE_REQUIRES(TIsFloatingPoint<FloatType>::Value)>
+template<typename FloatType UE_REQUIRES(std::is_floating_point_v<FloatType>)>
 inline FloatType Concatenate(FloatType LHS, FloatType RHS)
 {
 	return LHS * RHS;
@@ -430,7 +430,7 @@ inline FloatType Concatenate(FloatType LHS, FloatType RHS)
  * @param Transform Input transform from space A to space B.
  * @return Inverted transform from space B to space A.
  */
-template<typename FloatType, TEMPLATE_REQUIRES(TIsFloatingPoint<FloatType>::Value)>
+template<typename FloatType UE_REQUIRES(std::is_floating_point_v<FloatType>)>
 inline FloatType Inverse(FloatType Scale)
 {
 	return 1.0f / Scale;
diff --git a/Engine/Source/Runtime/Core/Public/Math/TransformNonVectorized.h b/Engine/Source/Runtime/Core/Public/Math/TransformNonVectorized.h
--- a/Engine/Source/Runtime/Core/Public/Math/TransformNonVectorized.h
+++ b/Engine/Source/Runtime/Core/Public/Math/TransformNonVectorized.h
@@ -1191,8 +1191,11 @@ private:
 
 public:
 	// Conversion to other type.
-	template<typename FArg, TEMPLATE_REQUIRES(!std::is_same_v<T, FArg>)>
-	explicit TTransform(const TTransform<FArg>& From) : TTransform<T>((TQuat<T>)From.GetRotation(), (TVector<T>)From.GetTranslation(), (TVector<T>)From.GetScale3D()) {}
+	template<typename FArg UE_REQUIRES(!std::is_same_v<T, FArg>)>
+	explicit TTransform(const TTransform<FArg>& From)
+		: TTransform<T>((TQuat<T>)From.GetRotation(), (TVector<T>)From.GetTranslation(), (TVector<T>)From.GetScale3D())
+	{
+	}
 };
 
 #if !defined(_MSC_VER) || defined(__clang__)  // MSVC can't forward declare explicit specializations
diff --git a/Engine/Source/Runtime/Core/Public/Math/TransformVectorized.h b/Engine/Source/Runtime/Core/Public/Math/TransformVectorized.h
--- a/Engine/Source/Runtime/Core/Public/Math/TransformVectorized.h
+++ b/Engine/Source/Runtime/Core/Public/Math/TransformVectorized.h
@@ -62,8 +62,8 @@ namespace Math
 template<typename T>
 struct alignas( TAlignOfTransform<T>::Value ) TTransform
 {
-	// Can't have a TEMPLATE_REQUIRES in the declaration because of the forward declarations, so check for allowed types here.
-	static_assert(TIsFloatingPoint<T>::Value, "TTransform only supports float and double types.");
+	// Can't have a UE_REQUIRES in the declaration because of the forward declarations, so check for allowed types here.
+	static_assert(std::is_floating_point_v<T>, "TTransform only supports float and double types.");
 
 	friend Z_Construct_UScriptStruct_FTransform3f_Statics;
 	friend Z_Construct_UScriptStruct_FTransform3d_Statics;
@@ -1454,12 +1454,18 @@ public:
 
 	// Conversion to other type.
 	friend struct TTransform<double>;
-	template<typename FArg, TEMPLATE_REQUIRES(!std::is_same_v<T, FArg> && std::is_same_v<T, float>)>
-	explicit TTransform(const TTransform<FArg>& From) : TTransform(MakeVectorRegisterFloatFromDouble(From.Rotation), MakeVectorRegisterFloatFromDouble(From.Translation), MakeVectorRegisterFloatFromDouble(From.Scale3D)) {}
+	template<typename FArg UE_REQUIRES(!std::is_same_v<T, FArg> && std::is_same_v<T, float>)>
+	explicit TTransform(const TTransform<FArg>& From)
+		: TTransform(MakeVectorRegisterFloatFromDouble(From.Rotation), MakeVectorRegisterFloatFromDouble(From.Translation), MakeVectorRegisterFloatFromDouble(From.Scale3D))
+	{
+	}
 
 	friend struct TTransform<float>;
-	template<typename FArg, TEMPLATE_REQUIRES(!std::is_same_v<T, FArg> && std::is_same_v<T, double>)>
-	explicit TTransform(const TTransform<FArg>& From) : TTransform(MakeVectorRegisterDouble(From.Rotation), MakeVectorRegisterDouble(From.Translation), MakeVectorRegisterDouble(From.Scale3D)) {}
+	template<typename FArg UE_REQUIRES(!std::is_same_v<T, FArg> && std::is_same_v<T, double>)>
+	explicit TTransform(const TTransform<FArg>& From)
+		: TTransform(MakeVectorRegisterDouble(From.Rotation), MakeVectorRegisterDouble(From.Translation), MakeVectorRegisterDouble(From.Scale3D))
+	{
+	}
 };
 
 #if !defined(_MSC_VER) || defined(__clang__)  // MSVC can't forward declare explicit specializations
diff --git a/Engine/Source/Runtime/Core/Public/Math/TranslationMatrix.h b/Engine/Source/Runtime/Core/Public/Math/TranslationMatrix.h
--- a/Engine/Source/Runtime/Core/Public/Math/TranslationMatrix.h
+++ b/Engine/Source/Runtime/Core/Public/Math/TranslationMatrix.h
@@ -20,8 +20,11 @@ public:
 	TTranslationMatrix(const TVector<T>& Delta);
 
 	// Conversion to other type.
-	template<typename FArg, TEMPLATE_REQUIRES(!std::is_same_v<T, FArg>)>
-	explicit TTranslationMatrix(const TTranslationMatrix<FArg>& From) : TMatrix<T>(From) {}
+	template<typename FArg UE_REQUIRES(!std::is_same_v<T, FArg>)>
+	explicit TTranslationMatrix(const TTranslationMatrix<FArg>& From)
+		: TMatrix<T>(From)
+	{
+	}
 	
 	/** Matrix factory. Return an FMatrix so we don't have type conversion issues in expressions. */
 	static TMatrix<T> Make(TVector<T> const& Delta)
diff --git a/Engine/Source/Runtime/Core/Public/Math/Vector.h b/Engine/Source/Runtime/Core/Public/Math/Vector.h
--- a/Engine/Source/Runtime/Core/Public/Math/Vector.h
+++ b/Engine/Source/Runtime/Core/Public/Math/Vector.h
@@ -283,7 +283,7 @@ public:
      * @param Bias How much to subtract from each component.
      * @return The result of subtraction.
      */
-	template<typename FArg, TEMPLATE_REQUIRES(std::is_arithmetic<FArg>::value)>
+	template<typename FArg UE_REQUIRES(std::is_arithmetic_v<FArg>)>
 	FORCEINLINE TVector<T> operator-(FArg Bias) const
 	{
 		return TVector<T>(X - (T)Bias, Y - (T)Bias, Z - (T)Bias);
@@ -295,7 +295,7 @@ public:
      * @param Bias How much to add to each component.
      * @return The result of addition.
      */
-	template<typename FArg, TEMPLATE_REQUIRES(std::is_arithmetic<FArg>::value)>
+	template<typename FArg UE_REQUIRES(std::is_arithmetic_v<FArg>)>
 	FORCEINLINE TVector<T> operator+(FArg Bias) const
 	{
 		return TVector<T>(X + (T)Bias, Y + (T)Bias, Z + (T)Bias);
@@ -307,7 +307,7 @@ public:
      * @param Scale What to multiply each component by.
      * @return The result of multiplication.
      */
-	template<typename FArg, TEMPLATE_REQUIRES(std::is_arithmetic<FArg>::value)>
+	template<typename FArg UE_REQUIRES(std::is_arithmetic_v<FArg>)>
 	FORCEINLINE TVector<T> operator*(FArg Scale) const
 	{
 		return TVector<T>(X * (T)Scale, Y * (T)Scale, Z * (T)Scale);
@@ -319,7 +319,7 @@ public:
      * @param Scale What to divide each component by.
      * @return The result of division.
      */
-	template<typename FArg, TEMPLATE_REQUIRES(std::is_arithmetic<FArg>::value)>
+	template<typename FArg UE_REQUIRES(std::is_arithmetic_v<FArg>)>
 	TVector<T> operator/(FArg Scale) const
 	{
 		const T RScale = T(1) / Scale;
@@ -408,7 +408,7 @@ public:
      * @param Scale Amount to scale this vector by.
      * @return Copy of the vector after scaling.
      */
-	template<typename FArg, TEMPLATE_REQUIRES(std::is_arithmetic<FArg>::value)>
+	template<typename FArg UE_REQUIRES(std::is_arithmetic_v<FArg>)>
 	FORCEINLINE TVector<T> operator*=(FArg Scale)
 	{
 		X *= Scale; Y *= Scale; Z *= Scale;
@@ -422,7 +422,7 @@ public:
      * @param V What to divide this vector by.
      * @return Copy of the vector after division.
      */
-	template<typename FArg, TEMPLATE_REQUIRES(std::is_arithmetic<FArg>::value)>
+	template<typename FArg UE_REQUIRES(std::is_arithmetic_v<FArg>)>
 	TVector<T> operator/=(FArg Scale)
 	{
 		const T RV = (T)1 / Scale;
@@ -1185,7 +1185,7 @@ public:
 	}
 
 	// Conversion from other type.
-	template<typename FArg, TEMPLATE_REQUIRES(!std::is_same_v<T, FArg>)>
+	template<typename FArg UE_REQUIRES(!std::is_same_v<T, FArg>)>
 	explicit TVector(const TVector<FArg>& From) : TVector<T>((T)From.X, (T)From.Y, (T)From.Z) {}
 };
 
@@ -2529,7 +2529,7 @@ template<> CORE_API const FVector3d FVector3d::ZAxisVector;
  * @param V Vector to scale.
  * @return Result of multiplication.
  */
-template<typename T, typename T2, TEMPLATE_REQUIRES(std::is_arithmetic<T2>::value)>
+template<typename T, typename T2 UE_REQUIRES(std::is_arithmetic_v<T2>)>
 FORCEINLINE TVector<T> operator*(T2 Scale, const TVector<T>& V)
 {
 	return V.operator*(Scale);
diff --git a/Engine/Source/Runtime/Core/Public/Math/Vector2D.h b/Engine/Source/Runtime/Core/Public/Math/Vector2D.h
--- a/Engine/Source/Runtime/Core/Public/Math/Vector2D.h
+++ b/Engine/Source/Runtime/Core/Public/Math/Vector2D.h
@@ -16,6 +16,8 @@
 #include "Math/IntPoint.h"
 #include "Logging/LogMacros.h"
 
+#include <type_traits>
+
 #ifdef _MSC_VER
 #pragma warning (push)
 // Ensure template functions don't generate shadowing warnings against global variables at the point of instantiation.
@@ -715,8 +717,11 @@ public:
 	
 	
 	// Conversion from other type.
-	template<typename FArg, TEMPLATE_REQUIRES(!std::is_same_v<T, FArg>)>
-	explicit TVector2(const TVector2<FArg>& From) : TVector2<T>((T)From.X, (T)From.Y) {}
+	template<typename FArg UE_REQUIRES(!std::is_same_v<T, FArg>)>
+	explicit TVector2(const TVector2<FArg>& From)
+		: TVector2<T>((T)From.X, (T)From.Y)
+	{
+	}
 };
 
 /**
@@ -814,7 +819,7 @@ template<> CORE_API const FVector2d FVector2d::Unit45Deg;
  * @param V Vector2 to scale.
  * @return Result of multiplication.
  */
-template<typename T, typename T2, TEMPLATE_REQUIRES(std::is_arithmetic<T2>::value)>
+template<typename T, typename T2 UE_REQUIRES(std::is_arithmetic_v<T2>)>
 FORCEINLINE TVector2<T> operator*(T2 Scale, const TVector2<T>& V)
 {
 	return V.operator*(Scale);
diff --git a/Engine/Source/Runtime/Core/Public/Math/Vector4.h b/Engine/Source/Runtime/Core/Public/Math/Vector4.h
--- a/Engine/Source/Runtime/Core/Public/Math/Vector4.h
+++ b/Engine/Source/Runtime/Core/Public/Math/Vector4.h
@@ -12,6 +12,9 @@
 #include "Math/Vector2D.h"
 #include "Math/Vector.h"
 #include "Serialization/MemoryLayout.h"
+#include "Templates/Requires.h"
+
+#include <type_traits>
 
 namespace UE
 {
@@ -24,8 +27,8 @@ namespace Math
 template<typename T>
 struct alignas(16) TVector4
 {
-	// Can't have a TEMPLATE_REQUIRES in the declaration because of the forward declarations, so check for allowed types here.
-	static_assert(TIsFloatingPoint<T>::Value, "TVector4 only supports float and double types.");
+	// Can't have a UE_REQUIRES in the declaration because of the forward declarations, so check for allowed types here.
+	static_assert(std::is_floating_point_v<T>, "TVector4 only supports float and double types.");
 
 
 public:
@@ -67,7 +70,7 @@ public:
 	 * @param InVector 3D Vector to set first three components.
 	 * @param InW W Coordinate.
 	 */
-	template<typename FArg, TEMPLATE_REQUIRES(std::is_arithmetic<FArg>::value)>
+	template<typename FArg UE_REQUIRES(std::is_arithmetic_v<FArg>)>
 	FORCEINLINE TVector4(const UE::Math::TVector<T>& InVector, FArg InW)
 		: X(InVector.X)
 		, Y(InVector.Y)
@@ -83,7 +86,7 @@ public:
 	 * @param InOverrideW Replaces W Coordinate of InVector.
 	 */
 
-	template<typename FArg, TEMPLATE_REQUIRES(std::is_arithmetic<FArg>::value)>
+	template<typename FArg UE_REQUIRES(std::is_arithmetic_v<FArg>)>
 	FORCEINLINE TVector4(const UE::Math::TVector4<T>& InVector, FArg OverrideW)
 		: X(InVector.X)
 		, Y(InVector.Y)
@@ -242,7 +245,7 @@ public:
 	 * @param Scale The scaling factor.
 	 * @return The result of vector scaling.
 	 */
-	template<typename FArg, TEMPLATE_REQUIRES(std::is_arithmetic<FArg>::value)>
+	template<typename FArg UE_REQUIRES(std::is_arithmetic_v<FArg>)>
 	FORCEINLINE TVector4<T> operator*(FArg Scale) const
 	{
 		return TVector4(X * Scale, Y * Scale, Z * Scale, W * Scale);
@@ -254,7 +257,7 @@ public:
 	 * @param Scale What to divide by.
 	 * @return The result of division.
 	 */
-	template<typename FArg, TEMPLATE_REQUIRES(std::is_arithmetic<FArg>::value)>
+	template<typename FArg UE_REQUIRES(std::is_arithmetic_v<FArg>)>
 	FORCEINLINE TVector4<T> operator/(FArg Scale) const
 	{
 		const T RScale = T(1.0f) / Scale;
@@ -299,7 +302,7 @@ public:
 	 * @param Scale The scaling factor.
 	 * @return The result of vector scaling.
 	 */
-	template<typename FArg, TEMPLATE_REQUIRES(std::is_arithmetic<FArg>::value)>
+	template<typename FArg UE_REQUIRES(std::is_arithmetic_v<FArg>)>
 	FORCEINLINE TVector4<T> operator*=(FArg Scale)
 	{
 		X *= Scale; Y *= Scale; Z *= Scale; W *= Scale;
@@ -313,7 +316,7 @@ public:
 	 * @param Scale The inverse scaling factor.
 	 * @return The result of vector scaling by 1/Scale.
 	 */
-	template<typename FArg, TEMPLATE_REQUIRES(std::is_arithmetic<FArg>::value)>
+	template<typename FArg UE_REQUIRES(std::is_arithmetic_v<FArg>)>
 	FORCEINLINE TVector4<T> operator/=(FArg Scale)
 	{
 		const T RV = T(1.0f) / Scale;
@@ -519,11 +522,11 @@ public:
 	bool SerializeFromMismatchedTag(FName StructTag, FArchive& Ar);
 	
 	// Conversion from other type: double->float
-	template<typename FArg, TEMPLATE_REQUIRES(std::is_same_v<FArg, double> && std::is_same_v<T, float>)>
+	template<typename FArg UE_REQUIRES(std::is_same_v<FArg, double> && std::is_same_v<T, float>)>
 	explicit TVector4(const TVector4<FArg>& From) : TVector4<T>((T)From.X, (T)From.Y, (T)From.Z, (T)From.W) {}
 
 	// Conversion from other type: float->double
-	template<typename FArg, TEMPLATE_REQUIRES(std::is_same_v<FArg, float> && std::is_same_v<T, double>)>
+	template<typename FArg UE_REQUIRES(std::is_same_v<FArg, float> && std::is_same_v<T, double>)>
 	explicit TVector4(const TVector4<FArg>& From) : TVector4<T>((T)From.X, (T)From.Y, (T)From.Z, (T)From.W) {}
 
 	/**
@@ -999,7 +1002,7 @@ namespace Math
 	 * @param V The vector to scale.
 	 * @return The result of scaling.
 	 */
-	template<typename T, typename T2, TEMPLATE_REQUIRES(std::is_arithmetic<T2>::value)>
+	template<typename T, typename T2 UE_REQUIRES(std::is_arithmetic_v<T2>)>
 	FORCEINLINE TVector4<T> operator*(T2 Scale, const TVector4<T>& V)
 	{
 		return V.operator*(Scale);
diff --git a/Engine/Source/Runtime/Engine/Classes/Engine/NetSerialization.h b/Engine/Source/Runtime/Engine/Classes/Engine/NetSerialization.h
--- a/Engine/Source/Runtime/Engine/Classes/Engine/NetSerialization.h
+++ b/Engine/Source/Runtime/Engine/Classes/Engine/NetSerialization.h
@@ -325,7 +325,7 @@ bool WriteFixedCompressedFloat(const T Value, FArchive& Ar)
 	return !clamp;
 }
 
-template<int32 MaxValue, uint32 NumBits, typename T, TEMPLATE_REQUIRES(TIsFloatingPoint<T>::Value), TEMPLATE_REQUIRES(NumBits < 32)>
+template<int32 MaxValue, uint32 NumBits, typename T UE_REQUIRES(std::is_floating_point_v<T> && NumBits < 32)>
 bool ReadFixedCompressedFloat(T& Value, FArchive& Ar)
 {
 	using Details = TFixedCompressedFloatDetails<MaxValue, NumBits>;
