Chaos Cloth Asset - Fixed crash when duplicating the SkeletalMeshImport node.

#rnx

diff --git a/Engine/Plugins/ChaosClothAssetEditor/Source/ChaosClothAssetDataflowNodes/Private/ChaosClothAsset/SkeletalMeshImportNode.cpp b/Engine/Plugins/ChaosClothAssetEditor/Source/ChaosClothAssetDataflowNodes/Private/ChaosClothAsset/SkeletalMeshImportNode.cpp
--- a/Engine/Plugins/ChaosClothAssetEditor/Source/ChaosClothAssetDataflowNodes/Private/ChaosClothAsset/SkeletalMeshImportNode.cpp
+++ b/Engine/Plugins/ChaosClothAssetEditor/Source/ChaosClothAssetDataflowNodes/Private/ChaosClothAsset/SkeletalMeshImportNode.cpp
@@ -95,7 +95,7 @@ void FChaosClothAssetSkeletalMeshImportNode::Evaluate(Dataflow::FContext& Contex
 void FChaosClothAssetSkeletalMeshImportNode::Serialize(FArchive& Ar)
 {
 	Ar.UsingCustomVersion(FFortniteMainBranchObjectVersion::GUID);
-	if (Ar.CustomVer(FUE5MainStreamObjectVersion::GUID) < FFortniteMainBranchObjectVersion::ClothAssetSkeletalMeshMultiSectionImport)
+	if (Ar.IsLoading() && Ar.CustomVer(FUE5MainStreamObjectVersion::GUID) < FFortniteMainBranchObjectVersion::ClothAssetSkeletalMeshMultiSectionImport)
 	{
 		bImportSingleSection = true;
 		bSetPhysicsAsset = true;
