Fixed a bunch of Private namespaces which will be ambiguous after adding UE::Geometry::Private.

diff --git a/Engine/Plugins/ChaosClothAssetEditor/Source/ChaosClothAssetDataflowNodes/Private/ChaosClothAsset/RemeshNode.cpp b/Engine/Plugins/ChaosClothAssetEditor/Source/ChaosClothAssetDataflowNodes/Private/ChaosClothAsset/RemeshNode.cpp
--- a/Engine/Plugins/ChaosClothAssetEditor/Source/ChaosClothAssetDataflowNodes/Private/ChaosClothAsset/RemeshNode.cpp
+++ b/Engine/Plugins/ChaosClothAssetEditor/Source/ChaosClothAssetDataflowNodes/Private/ChaosClothAsset/RemeshNode.cpp
@@ -668,7 +668,7 @@ namespace UE::Chaos::ClothAsset::Private
 			const double EdgeLength = (Mesh.GetVertex(EdgeVerts[0]) - Mesh.GetVertex(EdgeVerts[1])).Length();
 			if (EdgeLength < MinLength)
 			{
-				const Private::FEdgeCollapseParameters CollapseParams = Private::GetBoundaryEdgeCollapseParameters(Mesh, Seams, EdgeVerts);
+				const UE::Chaos::ClothAsset::Private::FEdgeCollapseParameters CollapseParams = Private::GetBoundaryEdgeCollapseParameters(Mesh, Seams, EdgeVerts);
 
 				if (CollapseParams.bCanCollapse)
 				{
@@ -1105,21 +1105,21 @@ void FChaosClothAssetRemeshNode::RemeshSimMesh(const TSharedRef<const FManagedAr
 
 	for (int32 ResampleIter = 0; ResampleIter < IterationsSim; ++ResampleIter)
 	{
-		Private::RemeshSeams(Mesh2D, Seams, TargetEdgeLength);
+		UE::Chaos::ClothAsset::Private::RemeshSeams(Mesh2D, Seams, TargetEdgeLength);
 	}
 
 	// Remesh boundaries
 
 	for (int32 ResampleIter = 0; ResampleIter < IterationsSim; ++ResampleIter)
 	{
-		Private::RemeshBoundaries(Mesh2D, Seams, TargetEdgeLength);
+		UE::Chaos::ClothAsset::Private::RemeshBoundaries(Mesh2D, Seams, TargetEdgeLength);
 	}
 
 	// Do the remeshing of the rest of the mesh
 
 	UE::Geometry::FCompactMaps CompactMaps;
 	constexpr bool bUniformSmoothing = true;
-	Private::Remesh(Mesh2D, TargetEdgeLength, IterationsSim, SmoothingSim, bUniformSmoothing, Seams, &CompactMaps);
+	UE::Chaos::ClothAsset::Private::Remesh(Mesh2D, TargetEdgeLength, IterationsSim, SmoothingSim, bUniformSmoothing, Seams, &CompactMaps);
 
 	// Update stitches
 	for (TArray<FIntVector2>& Seam : Seams)
@@ -1142,7 +1142,7 @@ void FChaosClothAssetRemeshNode::RemeshSimMesh(const TSharedRef<const FManagedAr
 	Converter.Convert(ClothCollection, INDEX_NONE, EClothPatternVertexType::Sim3D, SourceMesh3D);
 
 	FDynamicMesh3 Mesh3D;
-	Private::ProjectTo3D(Mesh2D, SourceMesh2D, SourceMesh3D, Mesh3D);
+	UE::Chaos::ClothAsset::Private::ProjectTo3D(Mesh2D, SourceMesh2D, SourceMesh3D, Mesh3D);
 
 
 	// Build the output cloth sim mesh
@@ -1301,19 +1301,19 @@ void FChaosClothAssetRemeshNode::RemeshRenderMesh(const TSharedRef<const FManage
 		// Create pseudo-stitches based on boundary vertex proximity. These stitches aren't going to actually weld vertices together, but they will guide boundary remeshing.
 		// The goal is to maintain a vertex pairing along boundaries in order to avoid holes opening up when the mesh deforms due to skinning.
 		TArray<FIntVector2> Stitches;
-		Private::FindCoincidentBoundaryVertices(DynamicMesh, Stitches);
+		UE::Chaos::ClothAsset::Private::FindCoincidentBoundaryVertices(DynamicMesh, Stitches);
 
 		FClothGeometryTools::BuildConnectedSeams(Stitches, DynamicMesh, Seams);
 
 		for (int RemeshPass = 0; RemeshPass < RenderSeamRemeshIterations; ++RemeshPass)
 		{
-			Private::RemeshSeams(DynamicMesh, Seams, TargetEdgeLength);
+			UE::Chaos::ClothAsset::Private::RemeshSeams(DynamicMesh, Seams, TargetEdgeLength);
 		}
 
 		// Also remesh the open boundaries that are not constrained by seams
 		for (int RemeshPass = 0; RemeshPass < RenderSeamRemeshIterations; ++RemeshPass)
 		{
-			Private::RemeshBoundaries(DynamicMesh, Seams, TargetEdgeLength);
+			UE::Chaos::ClothAsset::Private::RemeshBoundaries(DynamicMesh, Seams, TargetEdgeLength);
 		}
 	}
 
@@ -1322,14 +1322,14 @@ void FChaosClothAssetRemeshNode::RemeshRenderMesh(const TSharedRef<const FManage
 	if (RemeshMethodRender == EChaosClothAssetRemeshMethod::Remesh)
 	{
 		constexpr bool bUniformSmoothing = false;	// uniform smoothing can distort the UV layer pretty badly
-		const bool bSuccess = Private::Remesh(DynamicMesh, TargetEdgeLength, IterationsRender, SmoothingRender, bUniformSmoothing, Seams, &CompactMaps);
+		const bool bSuccess = UE::Chaos::ClothAsset::Private::Remesh(DynamicMesh, TargetEdgeLength, IterationsRender, SmoothingRender, bUniformSmoothing, Seams, &CompactMaps);
 		check(bSuccess);
 	}
 	else
 	{
 		const bool bCoarsenBoundariesDuringSimplify = !bRemeshRenderSeams;
 		const int TargetVertexCount = FMath::RoundToInt(static_cast<float>(TargetPercentRender) / 100.0f * static_cast<float>(InputMeshVertexCount));
-		Private::Simplify(DynamicMesh, TargetVertexCount, bCoarsenBoundariesDuringSimplify, &CompactMaps);
+		UE::Chaos::ClothAsset::Private::Simplify(DynamicMesh, TargetVertexCount, bCoarsenBoundariesDuringSimplify, &CompactMaps);
 	}
 
 	// Collect outputs
diff --git a/Engine/Plugins/ChaosClothAssetEditor/Source/ChaosClothAssetEditor/Private/ChaosClothAsset/ClothEditorMode.cpp b/Engine/Plugins/ChaosClothAssetEditor/Source/ChaosClothAssetEditor/Private/ChaosClothAsset/ClothEditorMode.cpp
--- a/Engine/Plugins/ChaosClothAssetEditor/Source/ChaosClothAssetEditor/Private/ChaosClothAsset/ClothEditorMode.cpp
+++ b/Engine/Plugins/ChaosClothAssetEditor/Source/ChaosClothAssetEditor/Private/ChaosClothAsset/ClothEditorMode.cpp
@@ -848,7 +848,7 @@ void UChaosClothAssetEditorMode::ReinitializeDynamicMeshComponents()
 		for (int32 PatternID = 0; PatternID < ClothFacade.GetNumSimPatterns(); ++PatternID)
 		{
 			const FCollectionClothSimPatternConstFacade Pattern = ClothFacade.GetSimPattern(PatternID);
-			const FLinearColor PatternColor = Private::PseudoRandomColor(PatternID);
+			const FLinearColor PatternColor = UE::Chaos::ClothAsset::Private::PseudoRandomColor(PatternID);
 
 			for (int32 TriID = 0; TriID < Pattern.GetNumSimFaces(); ++TriID)
 			{
diff --git a/Engine/Plugins/Experimental/Dataflow/Source/DataflowEditor/Private/Dataflow/AssetDefinition_DataflowContext.cpp b/Engine/Plugins/Experimental/Dataflow/Source/DataflowEditor/Private/Dataflow/AssetDefinition_DataflowContext.cpp
--- a/Engine/Plugins/Experimental/Dataflow/Source/DataflowEditor/Private/Dataflow/AssetDefinition_DataflowContext.cpp
+++ b/Engine/Plugins/Experimental/Dataflow/Source/DataflowEditor/Private/Dataflow/AssetDefinition_DataflowContext.cpp
@@ -175,8 +175,8 @@ namespace DataflowContextDefinitionHelpers
 		check(ContentOwner.Get());
 		bool bNeedsNewAsset = true;
 		TObjectPtr<UObject> Asset = nullptr;
-		
-		UDataflow* DataflowAsset = Private::GetDataflowAssetFrom(ContentOwner);
+
+		UDataflow* DataflowAsset = UE::Dataflow::Private::GetDataflowAssetFrom(ContentOwner);
 		if (bDataflowEnableContextCaching)
 		{
 			// Setup an asset that lives in the content broswer.
diff --git a/Engine/Plugins/Experimental/Dataflow/Source/DataflowEditor/Private/Dataflow/DataflowEditorUtil.cpp b/Engine/Plugins/Experimental/Dataflow/Source/DataflowEditor/Private/Dataflow/DataflowEditorUtil.cpp
--- a/Engine/Plugins/Experimental/Dataflow/Source/DataflowEditor/Private/Dataflow/DataflowEditorUtil.cpp
+++ b/Engine/Plugins/Experimental/Dataflow/Source/DataflowEditor/Private/Dataflow/DataflowEditorUtil.cpp
@@ -16,7 +16,7 @@
 
 using namespace UE::Geometry;
 
-namespace Private
+namespace UE::Dataflow::Private
 {
 	bool HasSkeletalMesh(UObject* InObject)
 	{
diff --git a/Engine/Plugins/Experimental/Dataflow/Source/DataflowEditor/Private/Dataflow/DataflowEditorUtil.h b/Engine/Plugins/Experimental/Dataflow/Source/DataflowEditor/Private/Dataflow/DataflowEditorUtil.h
--- a/Engine/Plugins/Experimental/Dataflow/Source/DataflowEditor/Private/Dataflow/DataflowEditorUtil.h
+++ b/Engine/Plugins/Experimental/Dataflow/Source/DataflowEditor/Private/Dataflow/DataflowEditorUtil.h
@@ -17,7 +17,7 @@ class USkeleton;
 class UAnimationAsset;
 class UMaterial;
 
-namespace Private
+namespace UE::Dataflow::Private
 {
 	bool HasSkeletalMesh(UObject* InObject);
 
