Cloth Editor: add a new "WeightMapNode" that will supersede AddWeightMapNode. Deprecate the old AddWeightMapNode. The new Node has only one set of vertex weights rather than two. The deprecated AddWeightMapNode still works in the graph but its properties cannot be edited and is no longer linked to the Weight Map Paint Tool. Users should migrate their weight map data out of the AddWeightMap nodes and into WeightMap nodes.

diff --git a/Engine/Plugins/ChaosClothAssetEditor/Source/ChaosClothAssetDataflowNodes/Private/ChaosClothAsset/ChaosClothAssetDataflowNodesModule.cpp b/Engine/Plugins/ChaosClothAssetEditor/Source/ChaosClothAssetDataflowNodes/Private/ChaosClothAsset/ChaosClothAssetDataflowNodesModule.cpp
--- a/Engine/Plugins/ChaosClothAssetEditor/Source/ChaosClothAssetDataflowNodes/Private/ChaosClothAsset/ChaosClothAssetDataflowNodesModule.cpp
+++ b/Engine/Plugins/ChaosClothAssetEditor/Source/ChaosClothAssetDataflowNodes/Private/ChaosClothAsset/ChaosClothAssetDataflowNodesModule.cpp
@@ -58,7 +58,9 @@
 #include "ChaosClothAsset/TransformPositionsNode.h"
 #include "ChaosClothAsset/TransformUVsNode.h"
 #include "ChaosClothAsset/USDImportNode.h"
+#include "ChaosClothAsset/WeightMapNode.h"
 #include "ChaosClothAsset/WeightedValueCustomization.h"
+
 #include "ChaosClothAsset/ImportedValueCustomization.h"
 #include "ChaosClothAsset/WeightMapToSelectionNode.h"
 #include "ChaosClothAsset/ConnectableValueCustomization.h"
@@ -75,6 +77,7 @@ namespace UE::Chaos::ClothAsset
 		{
 			DATAFLOW_NODE_REGISTER_CREATION_FACTORY(FChaosClothAssetAddStitchNode);
 			DATAFLOW_NODE_REGISTER_CREATION_FACTORY(FChaosClothAssetAddWeightMapNode);
+			DATAFLOW_NODE_REGISTER_CREATION_FACTORY(FChaosClothAssetWeightMapNode);
 			DATAFLOW_NODE_REGISTER_CREATION_FACTORY(FChaosClothAssetAttributeNode);
 			DATAFLOW_NODE_REGISTER_CREATION_FACTORY(FChaosClothAssetBindToRootBoneNode);
 			DATAFLOW_NODE_REGISTER_CREATION_FACTORY(FChaosClothAssetCopySimulationToRenderMeshNode);
diff --git a/Engine/Plugins/ChaosClothAssetEditor/Source/ChaosClothAssetDataflowNodes/Private/ChaosClothAsset/WeightMapNode.cpp b/Engine/Plugins/ChaosClothAssetEditor/Source/ChaosClothAssetDataflowNodes/Private/ChaosClothAsset/WeightMapNode.cpp
new file mode 100644
--- /dev/null
+++ b/Engine/Plugins/ChaosClothAssetEditor/Source/ChaosClothAssetDataflowNodes/Private/ChaosClothAsset/WeightMapNode.cpp
@@ -0,0 +1,305 @@
+// Copyright Epic Games, Inc. All Rights Reserved.
+
+#include "ChaosClothAsset/WeightMapNode.h"
+#include "ChaosClothAsset/ClothAsset.h"
+#include "ChaosClothAsset/ClothCollectionGroup.h"
+#include "ChaosClothAsset/ClothDataflowTools.h"
+#include "ChaosClothAsset/ClothGeometryTools.h"
+#include "ChaosClothAsset/CollectionClothFacade.h"
+#include "ChaosClothAsset/WeightedValue.h"
+#include "Dataflow/DataflowInputOutput.h"
+#include "Dataflow/DataflowObject.h"
+#include "InteractiveToolChange.h"
+
+#include UE_INLINE_GENERATED_CPP_BY_NAME(WeightMapNode)
+
+#define LOCTEXT_NAMESPACE "ChaosClothAssetWeightMapNode"
+
+namespace UE::Chaos::ClothAsset::Private
+{
+	// These are defined in AddWeightMapNode.cpp
+
+	void TransferWeightMap(
+		const TConstArrayView<FVector2f>& InSourcePositions,
+		const TConstArrayView<FIntVector3>& SourceIndices,
+		const TConstArrayView<int32> SourceWeightsLookup,
+		const TConstArrayView<float>& InSourceWeights,
+		const TConstArrayView<FVector2f>& InTargetPositions,
+		const TConstArrayView<FIntVector3>& TargetIndices,
+		const TConstArrayView<int32> TargetWeightsLookup,
+		TArray<float>& OutTargetWeights);
+
+	void SetVertexWeights(const TConstArrayView<float> InputMap, const TArray<float>& FinalValues, EChaosClothAssetWeightMapOverrideType OverrideType, TArray<float>& SourceVertexWeights);
+
+	void CalculateFinalVertexWeightValues(const TConstArrayView<float> InputMap, TArrayView<float> FinalOutputMap, EChaosClothAssetWeightMapOverrideType OverrideType, const TArray<float>& SourceVertexWeights);
+}
+
+
+FChaosClothAssetWeightMapNode::FChaosClothAssetWeightMapNode(const Dataflow::FNodeParameters& InParam, FGuid InGuid)
+	: FDataflowTerminalNode(InParam, InGuid)
+{
+	RegisterInputConnection(&Collection);
+	if (FDataflowInput* const Input = RegisterInputConnection(&InputName.StringValue, GET_MEMBER_NAME_CHECKED(FChaosClothAssetConnectableIStringValue, StringValue)))
+	{
+		Input->SetCanHidePin(true);
+		Input->SetPinIsHidden(true);
+	}
+	if (FDataflowInput* const Input = RegisterInputConnection(&TransferCollection))
+	{
+		Input->SetCanHidePin(true);
+		Input->SetPinIsHidden(true);
+	}
+	RegisterOutputConnection(&Collection, &Collection);
+	RegisterOutputConnection(&Name);
+}
+
+void FChaosClothAssetWeightMapNode::SetAssetValue(TObjectPtr<UObject> Asset, Dataflow::FContext& Context) const
+{
+	using namespace UE::Chaos::ClothAsset;
+
+	if (UChaosClothAsset* const ClothAsset = Cast<UChaosClothAsset>(Asset.Get()))
+	{
+PRAGMA_DISABLE_DEPRECATION_WARNINGS  // TODO: Remove after adding the getter function (currently shelved!)  
+		if (UDataflow* const DataflowAsset = ClothAsset->DataflowAsset)
+PRAGMA_ENABLE_DEPRECATION_WARNINGS
+		{
+			const TSharedPtr<Dataflow::FGraph, ESPMode::ThreadSafe> Dataflow = DataflowAsset->GetDataflow();
+			if (const TSharedPtr<FDataflowNode> BaseNode = Dataflow->FindBaseNode(this->GetGuid()))  // This is basically a safe const_cast
+			{
+				FChaosClothAssetWeightMapNode* const MutableThis = static_cast<FChaosClothAssetWeightMapNode*>(BaseNode.Get());
+				check(MutableThis == this);
+
+				// Make the name a valid attribute name, and replace the value in the UI
+				FWeightMapTools::MakeWeightMapName(MutableThis->Name);
+
+				// Transfer weight map if the transfer collection input has changed and is valid
+				FManagedArrayCollection InCollection = GetValue<FManagedArrayCollection>(Context, &Collection);
+				const TSharedRef<FManagedArrayCollection> ClothCollection = MakeShared<FManagedArrayCollection>(MoveTemp(InCollection));
+				FCollectionClothConstFacade ClothFacade(ClothCollection);
+				if (ClothFacade.IsValid())  // Can only act on the collection if it is a valid cloth collection
+				{
+					FManagedArrayCollection InTransferCollection = GetValue<FManagedArrayCollection>(Context, &TransferCollection);
+					const TSharedRef<const FManagedArrayCollection> TransferClothCollection = MakeShared<const FManagedArrayCollection>(MoveTemp(InTransferCollection));
+					FCollectionClothConstFacade TransferClothFacade(TransferClothCollection);
+					const FName InInputName = GetInputName(Context);
+
+					const uint32 NameTypeHash = HashCombineFast(GetTypeHash(InInputName), (uint32)TransferType);
+					const uint32 InTransferCollectionHash = (TransferClothFacade.HasValidSimulationData() && InInputName != NAME_None) ?
+						HashCombineFast(TransferClothFacade.CalculateWeightMapTypeHash(), NameTypeHash) : 0;  // TODO: Remove after adding the function (currently shelved!)  
+
+					if (TransferCollectionHash != InTransferCollectionHash)
+					{
+						MutableThis->TransferCollectionHash = InTransferCollectionHash;
+
+						if (TransferCollectionHash)
+						{
+							if (TransferClothFacade.HasWeightMap(InInputName))
+							{
+								// Remap the weights
+								TArray<float> RemappedWeights;
+								RemappedWeights.SetNumZeroed(ClothFacade.GetNumSimVertices3D());
+
+								switch (TransferType)
+								{
+								case EChaosClothAssetWeightMapTransferType::Use2DSimMesh:
+									Private::TransferWeightMap(
+										TransferClothFacade.GetSimPosition2D(),
+										TransferClothFacade.GetSimIndices2D(),
+										TransferClothFacade.GetSimVertex3DLookup(),
+										TransferClothFacade.GetWeightMap(InInputName),
+										ClothFacade.GetSimPosition2D(),
+										ClothFacade.GetSimIndices2D(),
+										ClothFacade.GetSimVertex3DLookup(),
+										RemappedWeights);
+									break;
+								case EChaosClothAssetWeightMapTransferType::Use3DSimMesh:
+									FClothGeometryTools::TransferWeightMap(
+										TransferClothFacade.GetSimPosition3D(),
+										TransferClothFacade.GetSimIndices3D(),
+										TransferClothFacade.GetWeightMap(InInputName),
+										ClothFacade.GetSimPosition3D(),
+										ClothFacade.GetSimNormal(),
+										ClothFacade.GetSimIndices3D(),
+										TArrayView<float>(RemappedWeights));
+									break;
+								default: unimplemented();
+								}
+
+								// TODO: Allow for transferring render weights from InTransferCollection?
+
+								MutableThis->SetVertexWeights(ClothFacade.GetWeightMap(InInputName), RemappedWeights);
+							}
+						}
+					}
+				}
+			}
+		}
+	}
+}
+
+void FChaosClothAssetWeightMapNode::Evaluate(Dataflow::FContext& Context, const FDataflowOutput* Out) const
+{
+	using namespace UE::Chaos::ClothAsset;
+
+	auto CheckSourceVertexWeights = [this](TArrayView<float>& ClothWeights, const TArray<float>& SourceVertexWeights, bool bIsSim)
+	{
+		if (SourceVertexWeights.Num() > 0 && SourceVertexWeights.Num() != ClothWeights.Num())
+		{
+			FClothDataflowTools::LogAndToastWarning(*this,
+				LOCTEXT("VertexCountMismatchHeadline", "Vertex count mismatch."),
+				FText::Format(LOCTEXT("VertexCountMismatchDetails", "{0} vertex weights in the node: {1}\n{0} vertices in the cloth: {2}"),
+					bIsSim ? FText::FromString("Sim") : FText::FromString("Render"),
+					SourceVertexWeights.Num(),
+					ClothWeights.Num()));
+		}
+	};
+
+	if (Out->IsA<FManagedArrayCollection>(&Collection))
+	{
+		// Evaluate InputName
+		const FName InInputName = GetInputName(Context);
+		InputName.StringValue_Override = GetValue<FString>(Context, &InputName.StringValue, UE::Chaos::ClothAsset::FWeightMapTools::NotOverridden);
+
+		// Evaluate in collection
+		FManagedArrayCollection InCollection = GetValue<FManagedArrayCollection>(Context, &Collection);
+		const TSharedRef<FManagedArrayCollection> ClothCollection = MakeShared<FManagedArrayCollection>(MoveTemp(InCollection));
+		FCollectionClothFacade ClothFacade(ClothCollection);
+		if (ClothFacade.IsValid())  // Can only act on the collection if it is a valid cloth collection
+		{
+			const FName InName(Name.IsEmpty() ? InInputName : FName(Name));
+
+			// Copy simulation weights into cloth collection
+
+			if (MeshTarget == EChaosClothAssetWeightMapMeshTarget::Simulation)
+			{
+				ClothFacade.AddWeightMap(InName);		// Does nothing if weight map already exists
+				TArrayView<float> ClothSimWeights = ClothFacade.GetWeightMap(InName);
+
+				if (ClothSimWeights.Num() != ClothFacade.GetNumSimVertices3D())
+				{
+					check(ClothSimWeights.Num() == 0);
+					FClothDataflowTools::LogAndToastWarning(*this,
+						LOCTEXT("InvalidSimWeightMapNameHeadline", "Invalid weight map name."),
+						FText::Format(LOCTEXT("InvalidSimWeightMapNameDetails", "Could not create a sim weight map with name \"{0}\" (reserved name? wrong type?)."),
+							FText::FromName(InName)));
+				}
+				else
+				{
+					constexpr bool bIsSim = true;
+					CheckSourceVertexWeights(ClothSimWeights, GetVertexWeights(), bIsSim);
+					CalculateFinalVertexWeightValues(ClothFacade.GetWeightMap(InInputName), ClothSimWeights);
+				}
+			}
+			else 
+			{
+				check(MeshTarget == EChaosClothAssetWeightMapMeshTarget::Render);
+
+				ClothFacade.AddUserDefinedAttribute<float>(InName, ClothCollectionGroup::RenderVertices);
+				TArrayView<float> ClothRenderWeights = ClothFacade.GetUserDefinedAttribute<float>(InName, ClothCollectionGroup::RenderVertices);
+
+				if (ClothRenderWeights.Num() != ClothFacade.GetNumRenderVertices())
+				{
+					check(ClothRenderWeights.Num() == 0);
+					FClothDataflowTools::LogAndToastWarning(*this,
+						LOCTEXT("InvalidRenderWeightMapNameHeadline", "Invalid weight map name."),
+						FText::Format(LOCTEXT("InvalidRenderWeightMapNameDetails", "Could not create a render weight map with name \"{0}\" (reserved name? wrong type?)."),
+							FText::FromName(InName)));
+				}
+				else
+				{
+					constexpr bool bIsSim = false;
+					CheckSourceVertexWeights(ClothRenderWeights, GetVertexWeights(), bIsSim);
+					CalculateFinalVertexWeightValues(ClothFacade.GetUserDefinedAttribute<float>(InInputName, ClothCollectionGroup::RenderVertices), ClothRenderWeights);
+				}
+			}
+		}
+
+		SetValue(Context, MoveTemp(*ClothCollection), &Collection);
+	}
+	else if (Out->IsA<FString>(&Name))
+	{
+		FString InputNameString = GetValue<FString>(Context, &InputName.StringValue);
+		UE::Chaos::ClothAsset::FWeightMapTools::MakeWeightMapName(InputNameString);
+		SetValue(Context, Name.IsEmpty() ? InputNameString : Name, &Name);
+	}
+}
+
+FName FChaosClothAssetWeightMapNode::GetInputName(Dataflow::FContext& Context) const
+{
+	FString InputNameString = GetValue<FString>(Context, &InputName.StringValue);
+	UE::Chaos::ClothAsset::FWeightMapTools::MakeWeightMapName(InputNameString);
+	const FName InInputName(*InputNameString);
+	return InInputName != NAME_None ? InInputName : FName(Name);
+}
+
+void FChaosClothAssetWeightMapNode::SetVertexWeights(const TConstArrayView<float> InputMap, const TArray<float>& FinalValues)
+{
+	UE::Chaos::ClothAsset::Private::SetVertexWeights(InputMap, FinalValues, MapOverrideType, GetVertexWeights());
+}
+
+void FChaosClothAssetWeightMapNode::CalculateFinalVertexWeightValues(const TConstArrayView<float> InputMap, TArrayView<float> FinalOutputMap) const
+{
+	UE::Chaos::ClothAsset::Private::CalculateFinalVertexWeightValues(InputMap, FinalOutputMap, MapOverrideType, GetVertexWeights());
+}
+
+// Object encapsulating a change to the WeightMap node's values. Used for Undo/Redo.
+class FChaosClothAssetWeightMapNode::FWeightMapNodeChange final : public FToolCommandChange
+{
+
+public:
+
+	FWeightMapNodeChange(const FChaosClothAssetWeightMapNode& Node) :
+		NodeGuid(Node.GetGuid()),
+		SavedWeights(Node.GetVertexWeights()),
+		SavedMapOverrideType(Node.MapOverrideType),
+		SavedWeightMapName(Node.Name)
+	{}
+
+private:
+
+	FGuid NodeGuid;
+	TArray<float> SavedWeights;
+
+	EChaosClothAssetWeightMapOverrideType SavedMapOverrideType;
+	FString SavedWeightMapName;
+
+	virtual FString ToString() const final
+	{
+		return TEXT("FChaosClothAssetWeightMapNode::FWeightMapNodeChange");
+	}
+
+	virtual void Apply(UObject* Object) final
+	{
+		SwapApplyRevert(Object);
+	}
+
+	virtual void Revert(UObject* Object) final
+	{
+		SwapApplyRevert(Object);
+	}
+
+	void SwapApplyRevert(UObject* Object)
+	{
+		if (UDataflow* const Dataflow = Cast<UDataflow>(Object))
+		{
+			if (const TSharedPtr<FDataflowNode> BaseNode = Dataflow->GetDataflow()->FindBaseNode(NodeGuid))
+			{
+				if (FChaosClothAssetWeightMapNode* const Node = BaseNode->AsType<FChaosClothAssetWeightMapNode>())
+				{
+					Swap(Node->GetVertexWeights(), SavedWeights);
+					Swap(Node->MapOverrideType, SavedMapOverrideType);
+					Swap(Node->Name, SavedWeightMapName);
+
+					Node->Invalidate();
+				}
+			}
+		}
+	}
+};
+
+TUniquePtr<FToolCommandChange> FChaosClothAssetWeightMapNode::MakeWeightMapNodeChange(const FChaosClothAssetWeightMapNode& Node)
+{
+	return MakeUnique<FChaosClothAssetWeightMapNode::FWeightMapNodeChange>(Node);
+}
+
+#undef LOCTEXT_NAMESPACE
diff --git a/Engine/Plugins/ChaosClothAssetEditor/Source/ChaosClothAssetDataflowNodes/Public/ChaosClothAsset/AddWeightMapNode.h b/Engine/Plugins/ChaosClothAssetEditor/Source/ChaosClothAssetDataflowNodes/Public/ChaosClothAsset/AddWeightMapNode.h
--- a/Engine/Plugins/ChaosClothAssetEditor/Source/ChaosClothAssetDataflowNodes/Public/ChaosClothAsset/AddWeightMapNode.h
+++ b/Engine/Plugins/ChaosClothAssetEditor/Source/ChaosClothAssetDataflowNodes/Public/ChaosClothAsset/AddWeightMapNode.h
@@ -38,9 +38,12 @@ enum class EChaosClothAssetWeightMapOverrideType : uint8
 };
 
 
-/** Painted weight map attributes node. */
-PRAGMA_DISABLE_DEPRECATION_WARNINGS  // For deprecated VertexWeights in copy constructor (Clang)
-USTRUCT(Meta = (DataflowCloth))
+/** 
+* Painted weight map attributes node.
+* Deprecated, use FChaosClothAssetWeightMapNode instead.
+*/
+PRAGMA_DISABLE_DEPRECATION_WARNINGS
+USTRUCT(Meta = (DataflowCloth, Deprecated="5.5"))
 struct FChaosClothAssetAddWeightMapNode : public FDataflowTerminalNode
 {
 	GENERATED_USTRUCT_BODY()
@@ -59,34 +62,34 @@ public:
 	 * Note this operation only happens once when the TransferCollection is first connected, or updated.
 	 * Changing the InputName or the TransferType will also redo the transfer operation.
 	 */
-	UPROPERTY(Meta = (DataflowInput))
+	UPROPERTY(VisibleAnywhere, Category = "Add Weight Map", Meta = (DataflowInput))
 	FManagedArrayCollection TransferCollection;
 
 	/** The name to be set as a weight map attribute. */
-	UPROPERTY(EditAnywhere, Category = "Add Weight Map", Meta = (DataflowOutput))
+	UPROPERTY(VisibleAnywhere, Category = "Add Weight Map", Meta = (DataflowOutput))
 	FString Name;
 
 	/** The name to populate this map from and override based on Map Override Type. Name will be used if Input Name is empty.*/
-	UPROPERTY(EditAnywhere, Category = "Add Weight Map")
+	UPROPERTY(VisibleAnywhere, Category = "Add Weight Map")
 	FChaosClothAssetConnectableIStringValue InputName = {TEXT("")};
 
 	/** How to apply this node's weight values onto existing maps. Changing this value will change the output map. 
 	 *  To change how the node's stored weights are calculated, change the equivalent value on the Weight Map Paint Tool context.*/
-	UPROPERTY(EditAnywhere, Category = "Add Weight Map")
+	UPROPERTY(VisibleAnywhere, Category = "Add Weight Map")
 	EChaosClothAssetWeightMapOverrideType MapOverrideType = EChaosClothAssetWeightMapOverrideType::ReplaceAll;
 
 	/**
 	 * The type of transfer used to transfer the weight map when a TransferCollection is connected.
 	 * This property is disabled when no TransferCollection input has been conencted.
 	 */
-	UPROPERTY(EditAnywhere, Category = "Add Weight Map", Meta = (EditCondition = "TransferCollectionHash != 0"))
+	UPROPERTY(VisibleAnywhere, Category = "Add Weight Map", Meta = (EditCondition = "TransferCollectionHash != 0"))
 	EChaosClothAssetWeightMapTransferType TransferType = EChaosClothAssetWeightMapTransferType::Use2DSimMesh;
 
 	UE_DEPRECATED(5.4, "This property will be made private.")
 	UPROPERTY()
 	TArray<float> VertexWeights;
 
-	UPROPERTY(EditAnywhere, Category = "Add Weight Map")
+	UPROPERTY(VisibleAnywhere, Category = "Add Weight Map")
 	EChaosClothAssetWeightMapMeshType MeshTarget = EChaosClothAssetWeightMapMeshType::Simulation;
 
 	FChaosClothAssetAddWeightMapNode(const Dataflow::FNodeParameters& InParam, FGuid InGuid = FGuid::NewGuid());
diff --git a/Engine/Plugins/ChaosClothAssetEditor/Source/ChaosClothAssetDataflowNodes/Public/ChaosClothAsset/WeightMapNode.h b/Engine/Plugins/ChaosClothAssetEditor/Source/ChaosClothAssetDataflowNodes/Public/ChaosClothAsset/WeightMapNode.h
new file mode 100644
--- /dev/null
+++ b/Engine/Plugins/ChaosClothAssetEditor/Source/ChaosClothAssetDataflowNodes/Public/ChaosClothAsset/WeightMapNode.h
@@ -0,0 +1,93 @@
+// Copyright Epic Games, Inc. All Rights Reserved.
+
+#pragma once 
+
+#include "ChaosClothAsset/AddWeightMapNode.h"
+#include "GeometryCollection/ManagedArrayCollection.h"
+#include "ChaosClothAsset/ConnectableValue.h"
+#include "WeightMapNode.generated.h"
+
+/** Which mesh to update with the node's weight map */
+UENUM()
+enum class EChaosClothAssetWeightMapMeshTarget : uint8
+{
+	Simulation,
+	Render
+};
+
+/** Painted weight map attributes node. */
+USTRUCT(Meta = (DataflowCloth))
+struct FChaosClothAssetWeightMapNode : public FDataflowTerminalNode
+{
+	GENERATED_USTRUCT_BODY()
+	DATAFLOW_NODE_DEFINE_INTERNAL(FChaosClothAssetWeightMapNode, "WeightMap", "Cloth", "Cloth Weight Map")
+
+public:
+	static constexpr float ReplaceChangedPassthroughValue = UE_BIG_NUMBER;
+
+	UPROPERTY(Meta = (DataflowInput, DataflowOutput, DataflowPassthrough = "Collection"))
+	FManagedArrayCollection Collection;
+
+	/**
+	 * The collection used to transfer weight map from.
+	 * Connecting a collection containing a weight map with Input Name (or Name if Input Name is empty)
+	 * will transfer the weights to the input collection vertices.
+	 * Note this operation only happens once when the TransferCollection is first connected, or updated.
+	 * Changing the InputName or the TransferType will also redo the transfer operation.
+	 */
+	UPROPERTY(Meta = (DataflowInput))
+	FManagedArrayCollection TransferCollection;
+
+	/** The name to be set as a weight map attribute. */
+	UPROPERTY(EditAnywhere, Category = "Weight Map", Meta = (DataflowOutput))
+	FString Name;
+
+	/** The name to populate this map from and override based on Map Override Type. Name will be used if Input Name is empty.*/
+	UPROPERTY(EditAnywhere, Category = "Weight Map")
+	FChaosClothAssetConnectableIStringValue InputName = { TEXT("") };
+
+	/** How to apply this node's weight values onto existing maps. Changing this value will change the output map.
+	 *  To change how the node's stored weights are calculated, change the equivalent value on the Weight Map Paint Tool context.*/
+	UPROPERTY(EditAnywhere, Category = "Weight Map")
+	EChaosClothAssetWeightMapOverrideType MapOverrideType = EChaosClothAssetWeightMapOverrideType::ReplaceChanged;
+
+	/**
+	 * The type of transfer used to transfer the weight map when a TransferCollection is connected.
+	 * This property is disabled when no TransferCollection input has been conencted.
+	 */
+	UPROPERTY(EditAnywhere, Category = "Weight Map", Meta = (EditCondition = "TransferCollectionHash != 0"))
+	EChaosClothAssetWeightMapTransferType TransferType = EChaosClothAssetWeightMapTransferType::Use2DSimMesh;
+
+	UPROPERTY(EditAnywhere, Category = "Weight Map")
+	EChaosClothAssetWeightMapMeshTarget MeshTarget = EChaosClothAssetWeightMapMeshTarget::Simulation;
+
+	FChaosClothAssetWeightMapNode(const Dataflow::FNodeParameters& InParam, FGuid InGuid = FGuid::NewGuid());
+
+private:
+	friend class UClothEditorWeightMapPaintTool;
+
+	class FWeightMapNodeChange;
+	static TUniquePtr<class FToolCommandChange> CHAOSCLOTHASSETDATAFLOWNODES_API MakeWeightMapNodeChange(const FChaosClothAssetWeightMapNode& Node);
+
+	const TArray<float>& GetVertexWeights() const { return VertexWeights; }
+	TArray<float>& GetVertexWeights() { return VertexWeights; }
+
+	UPROPERTY()
+	TArray<float> VertexWeights;
+
+	// These methods are exported for UClothEditorWeightMapPaintTool which lives in a different module.
+	FName CHAOSCLOTHASSETDATAFLOWNODES_API GetInputName(Dataflow::FContext& Context) const;
+
+	void CHAOSCLOTHASSETDATAFLOWNODES_API SetVertexWeights(const TConstArrayView<float> InputMap, const TArray<float>& FinalValues);
+
+	// Input and FinalOutputMap can be the same array, but should not otherwise be interleaved.
+	void CHAOSCLOTHASSETDATAFLOWNODES_API CalculateFinalVertexWeightValues(const TConstArrayView<float> InputMap, TArrayView<float> FinalOutputMap) const;
+
+	//~ Begin FDataflowNode interface
+	virtual void SetAssetValue(TObjectPtr<UObject> Asset, Dataflow::FContext& Context) const override;
+	virtual void Evaluate(Dataflow::FContext& Context, const FDataflowOutput* Out) const override;
+	//~ End FDataflowNode interface
+
+	UPROPERTY()
+	uint32 TransferCollectionHash = 0;
+};
diff --git a/Engine/Plugins/ChaosClothAssetEditor/Source/ChaosClothAssetEditor/Private/ChaosClothAsset/ClothEditorCommands.cpp b/Engine/Plugins/ChaosClothAssetEditor/Source/ChaosClothAssetEditor/Private/ChaosClothAsset/ClothEditorCommands.cpp
--- a/Engine/Plugins/ChaosClothAssetEditor/Source/ChaosClothAssetEditor/Private/ChaosClothAsset/ClothEditorCommands.cpp
+++ b/Engine/Plugins/ChaosClothAssetEditor/Source/ChaosClothAssetEditor/Private/ChaosClothAsset/ClothEditorCommands.cpp
@@ -42,8 +42,8 @@ void FChaosClothAssetEditorCommands::RegisterCommands()
 	UI_COMMAND(BeginRemeshTool, "Remesh", "Remesh the selected mesh", EUserInterfaceActionType::Button, FInputChord());
 	UI_COMMAND(BeginAttributeEditorTool, "AttrEd", "Edit/configure mesh attributes", EUserInterfaceActionType::Button, FInputChord());
 	
-	UI_COMMAND(BeginWeightMapPaintTool, "Add Weight Map", "Paint weight maps on the mesh", EUserInterfaceActionType::None, FInputChord());
-	UI_COMMAND(AddWeightMapNode, "Add Weight Map", "Paint weight maps on the mesh", EUserInterfaceActionType::Button, FInputChord());
+	UI_COMMAND(BeginWeightMapPaintTool, "Weight Map", "Paint weight maps on the mesh", EUserInterfaceActionType::None, FInputChord());
+	UI_COMMAND(AddWeightMapNode, "Weight Map", "Paint weight maps on the mesh", EUserInterfaceActionType::Button, FInputChord());
 	UI_COMMAND(BeginMeshSelectionTool, "Select", "Select mesh elements", EUserInterfaceActionType::None, FInputChord());
 	UI_COMMAND(AddMeshSelectionNode, "Select", "Select mesh elements", EUserInterfaceActionType::Button, FInputChord());
 
diff --git a/Engine/Plugins/ChaosClothAssetEditor/Source/ChaosClothAssetEditor/Private/ChaosClothAsset/ClothEditorMode.cpp b/Engine/Plugins/ChaosClothAssetEditor/Source/ChaosClothAssetEditor/Private/ChaosClothAsset/ClothEditorMode.cpp
--- a/Engine/Plugins/ChaosClothAssetEditor/Source/ChaosClothAssetEditor/Private/ChaosClothAsset/ClothEditorMode.cpp
+++ b/Engine/Plugins/ChaosClothAssetEditor/Source/ChaosClothAssetEditor/Private/ChaosClothAsset/ClothEditorMode.cpp
@@ -11,7 +11,7 @@
 #include "ChaosClothAsset/ClothEditorPreviewScene.h"
 #include "ChaosClothAsset/ClothEditorContextObject.h"
 #include "ChaosClothAsset/ClothEditorToolBuilders.h"
-#include "ChaosClothAsset/AddWeightMapNode.h"
+#include "ChaosClothAsset/WeightMapNode.h"
 #include "ChaosClothAsset/TransferSkinWeightsNode.h"
 #include "ChaosClothAsset/SelectionNode.h"
 #include "AssetEditorModeManager.h"
@@ -338,7 +338,7 @@ void UChaosClothAssetEditorMode::RegisterTools()
 
 	UClothEditorWeightMapPaintToolBuilder* WeightMapPaintToolBuilder = NewObject<UClothEditorWeightMapPaintToolBuilder>();
 	RegisterClothTool(CommandInfos.BeginWeightMapPaintTool, FChaosClothAssetEditorCommands::BeginWeightMapPaintToolIdentifier, WeightMapPaintToolBuilder, WeightMapPaintToolBuilder, ConstructionViewportToolsContext);
-	RegisterAddNodeCommand(CommandInfos.AddWeightMapNode, FChaosClothAssetAddWeightMapNode::StaticType(), CommandInfos.BeginWeightMapPaintTool);
+	RegisterAddNodeCommand(CommandInfos.AddWeightMapNode, FChaosClothAssetWeightMapNode::StaticType(), CommandInfos.BeginWeightMapPaintTool);
 
 	UClothTransferSkinWeightsToolBuilder* TransferToolBuilder = NewObject<UClothTransferSkinWeightsToolBuilder>();
 	RegisterClothTool(CommandInfos.BeginTransferSkinWeightsTool, FChaosClothAssetEditorCommands::BeginTransferSkinWeightsToolIdentifier, TransferToolBuilder, TransferToolBuilder, ConstructionViewportToolsContext);
diff --git a/Engine/Plugins/ChaosClothAssetEditor/Source/ChaosClothAssetEditorTools/Private/ChaosClothAsset/ClothWeightMapPaintTool.cpp b/Engine/Plugins/ChaosClothAssetEditor/Source/ChaosClothAssetEditorTools/Private/ChaosClothAsset/ClothWeightMapPaintTool.cpp
--- a/Engine/Plugins/ChaosClothAssetEditor/Source/ChaosClothAssetEditorTools/Private/ChaosClothAsset/ClothWeightMapPaintTool.cpp
+++ b/Engine/Plugins/ChaosClothAssetEditor/Source/ChaosClothAssetEditorTools/Private/ChaosClothAsset/ClothWeightMapPaintTool.cpp
@@ -39,7 +39,7 @@
 #include "ContextObjectStore.h"
 #include "Dataflow/DataflowEdNode.h"
 #include "Dataflow/DataflowObject.h"
-#include "ChaosClothAsset/AddWeightMapNode.h"
+#include "ChaosClothAsset/WeightMapNode.h"
 #include "GraphEditor.h"
 #include "Dataflow/DataflowSNode.h"
 #include "Dataflow/DataflowGraphEditor.h"
@@ -109,7 +109,7 @@ void UClothEditorWeightMapPaintTool::Setup()
 	UMeshSculptToolBase::Setup();
 
 	// Get the selected weight map node
-	WeightMapNodeToUpdate = ClothEditorContextObject->GetSingleSelectedNodeOfType<FChaosClothAssetAddWeightMapNode>();
+	WeightMapNodeToUpdate = ClothEditorContextObject->GetSingleSelectedNodeOfType<FChaosClothAssetWeightMapNode>();
 	checkf(WeightMapNodeToUpdate, TEXT("No Weight Map Node is currently selected, or more than one node is selected"));
 
 	SetToolDisplayName(LOCTEXT("ToolName", "Paint Weight Maps"));
@@ -464,14 +464,7 @@ void UClothEditorWeightMapPaintTool::Setup()
 	const int32 NumExpectedWeights = bHaveDynamicMeshToWeightConversion ? WeightToDynamicMesh.Num() : Mesh->MaxVertexID();
 	TArray<float> CurrentWeights;
 	CurrentWeights.SetNumZeroed(NumExpectedWeights);
-	if (bIsRenderMode)
-	{
-		WeightMapNodeToUpdate->CalculateFinalRenderVertexWeightValues(InputWeightMap, TArrayView<float>(CurrentWeights));
-	}
-	else
-	{
-		WeightMapNodeToUpdate->CalculateFinalVertexWeightValues(InputWeightMap, TArrayView<float>(CurrentWeights));
-	}
+	WeightMapNodeToUpdate->CalculateFinalVertexWeightValues(InputWeightMap, TArrayView<float>(CurrentWeights));
 	
 	if (bHaveDynamicMeshToWeightConversion)
 	{
@@ -1980,7 +1973,7 @@ void UClothEditorWeightMapPaintTool::UpdateSelectedNode()
 	if (UDataflow* const Dataflow = ClothEditorContextObject->GetDataflowAsset())
 	{
 		GetToolManager()->GetContextTransactionsAPI()->AppendChange(Dataflow, 
-			FChaosClothAssetAddWeightMapNode::MakeWeightMapNodeChange(*WeightMapNodeToUpdate),
+			FChaosClothAssetWeightMapNode::MakeWeightMapNodeChange(*WeightMapNodeToUpdate),
 			LOCTEXT("WeightMapNodeChangeDescription", "Update Weight Map Node"));
 	}
 
@@ -1998,26 +1991,11 @@ void UClothEditorWeightMapPaintTool::UpdateSelectedNode()
 		{
 			NodeWeights[DynamicMeshToWeight[DynamicMeshIdx]] = CurrentWeights[DynamicMeshIdx];
 		}
-
-		if (bIsRenderMode)
-		{
-			WeightMapNodeToUpdate->SetRenderVertexWeights(InputWeightMap, NodeWeights);
-		}
-		else
-		{
-			WeightMapNodeToUpdate->SetVertexWeights(InputWeightMap, NodeWeights);
-		}
+		WeightMapNodeToUpdate->SetVertexWeights(InputWeightMap, NodeWeights);
 	}
 	else
 	{
-		if (bIsRenderMode)
-		{
-			WeightMapNodeToUpdate->SetRenderVertexWeights(InputWeightMap, CurrentWeights);
-		}
-		else
-		{
-			WeightMapNodeToUpdate->SetVertexWeights(InputWeightMap, CurrentWeights);
-		}
+		WeightMapNodeToUpdate->SetVertexWeights(InputWeightMap, CurrentWeights);
 	}
 	
 	WeightMapNodeToUpdate->Invalidate();
diff --git a/Engine/Plugins/ChaosClothAssetEditor/Source/ChaosClothAssetEditorTools/Private/ChaosClothAsset/ClothWeightMapPaintTool.h b/Engine/Plugins/ChaosClothAssetEditor/Source/ChaosClothAssetEditorTools/Private/ChaosClothAsset/ClothWeightMapPaintTool.h
--- a/Engine/Plugins/ChaosClothAssetEditor/Source/ChaosClothAssetEditorTools/Private/ChaosClothAsset/ClothWeightMapPaintTool.h
+++ b/Engine/Plugins/ChaosClothAssetEditor/Source/ChaosClothAssetEditorTools/Private/ChaosClothAsset/ClothWeightMapPaintTool.h
@@ -30,7 +30,7 @@ class UWeightMapPaintBrushOpProps;
 class UWeightMapSmoothBrushOpProps;
 class UClothEditorContextObject;
 class UPolygonSelectionMechanic;
-struct FChaosClothAssetAddWeightMapNode;
+struct FChaosClothAssetWeightMapNode;
 enum class EChaosClothAssetWeightMapOverrideType : uint8;
 
 DECLARE_STATS_GROUP(TEXT("WeightMapPaintTool"), STATGROUP_WeightMapPaintTool, STATCAT_Advanced);
@@ -555,7 +555,7 @@ protected:
 
 	// Node graph editor support
 
-	FChaosClothAssetAddWeightMapNode* WeightMapNodeToUpdate = nullptr;
+	FChaosClothAssetWeightMapNode* WeightMapNodeToUpdate = nullptr;
 
 	void UpdateSelectedNode();
 
