Cloth Editor: when clicking on a WeightMap node to start the Paint tool, change the construction viewport view mode to the appropriate Sim/Render mode based on the node's MeshTarget value. Don't allow changing between Render and Sim modes while the tool is active (but do allow changing 2D/3D if we are in Sim mode)

diff --git a/Engine/Plugins/ChaosClothAssetEditor/Source/ChaosClothAssetEditor/Private/ChaosClothAsset/ClothEditorMode.cpp b/Engine/Plugins/ChaosClothAssetEditor/Source/ChaosClothAssetEditor/Private/ChaosClothAsset/ClothEditorMode.cpp
--- a/Engine/Plugins/ChaosClothAssetEditor/Source/ChaosClothAssetEditor/Private/ChaosClothAsset/ClothEditorMode.cpp
+++ b/Engine/Plugins/ChaosClothAssetEditor/Source/ChaosClothAssetEditor/Private/ChaosClothAsset/ClothEditorMode.cpp
@@ -229,7 +229,15 @@ void UChaosClothAssetEditorMode::RegisterClothTool(TSharedPtr<FUICommandInfo> UI
 		{
 			// Check if we need to switch view modes before starting the tool
 			TArray<UE::Chaos::ClothAsset::EClothPatternVertexType> SupportedModes;
-			ClothToolBuilder->GetSupportedViewModes(SupportedModes);
+
+			const UClothEditorContextObject* EditorContextObject = ToolsContext->ContextObjectStore->FindContext<UClothEditorContextObject>();
+			if (!EditorContextObject)
+			{
+				InitializeContextObject();
+				EditorContextObject = ToolsContext->ContextObjectStore->FindContext<UClothEditorContextObject>();
+			}
+			checkf(EditorContextObject, TEXT("Failed to find or create UClothEditorContextObject"));
+			ClothToolBuilder->GetSupportedViewModes(*EditorContextObject, SupportedModes);
 
 			bDynamicMeshUseInputCollection = true;
 			if (SupportedModes.Num() > 0 && !SupportedModes.Contains(this->GetConstructionViewMode()))
@@ -246,14 +254,8 @@ void UChaosClothAssetEditorMode::RegisterClothTool(TSharedPtr<FUICommandInfo> UI
 			}
 			else
 			{
-				bool bCurrentDynamicMeshIsInput = false;
-				if (const UEditorInteractiveToolsContext* const RestSpaceToolsContext = GetInteractiveToolsContext())
-				{
-					if (const UClothEditorContextObject* const EditorContextObject = RestSpaceToolsContext->ContextObjectStore->FindContext<UClothEditorContextObject>())
-					{
-						bCurrentDynamicMeshIsInput = EditorContextObject->IsUsingInputCollection();
-					}
-				}
+				const bool bCurrentDynamicMeshIsInput = EditorContextObject->IsUsingInputCollection();
+
 				if (!bCurrentDynamicMeshIsInput || bDynamicMeshComponentInitDeferred)
 				{
 					ReinitializeDynamicMeshComponents();
@@ -1360,8 +1362,15 @@ bool UChaosClothAssetEditorMode::CanChangeConstructionViewModeTo(UE::Chaos::Clot
 	const IChaosClothAssetEditorToolBuilder* const ClothToolBuilder = Cast<const IChaosClothAssetEditorToolBuilder>(ActiveToolBuilder);
 	checkf(ClothToolBuilder, TEXT("Cloth Editor has an active Tool Builder that does not implement IChaosClothAssetEditorToolBuilder"));
 
+	const UEditorInteractiveToolsContext* const RestSpaceToolsContext = GetInteractiveToolsContext();
+	checkf(RestSpaceToolsContext, TEXT("Cloth Editor Mode doesn't have a valid InteractiveToolsContext"));
+
+	const UClothEditorContextObject* const EditorContextObject = RestSpaceToolsContext->ContextObjectStore->FindContext<UClothEditorContextObject>();
+	checkf(EditorContextObject, TEXT("UClothEditorContextObject not found in ContextObjectStore despite having an active tool. This should have been created by the time a tool is activated"));
+
 	TArray<UE::Chaos::ClothAsset::EClothPatternVertexType> SupportedViewModes;
-	ClothToolBuilder->GetSupportedViewModes(SupportedViewModes);
+	ClothToolBuilder->GetSupportedViewModes(*EditorContextObject, SupportedViewModes);
+
 	return SupportedViewModes.Contains(NewViewMode);
 }
 
diff --git a/Engine/Plugins/ChaosClothAssetEditor/Source/ChaosClothAssetEditorTools/Private/ChaosClothAsset/ClothEditorToolBuilders.cpp b/Engine/Plugins/ChaosClothAssetEditor/Source/ChaosClothAssetEditorTools/Private/ChaosClothAsset/ClothEditorToolBuilders.cpp
--- a/Engine/Plugins/ChaosClothAssetEditor/Source/ChaosClothAssetEditorTools/Private/ChaosClothAsset/ClothEditorToolBuilders.cpp
+++ b/Engine/Plugins/ChaosClothAssetEditor/Source/ChaosClothAssetEditorTools/Private/ChaosClothAsset/ClothEditorToolBuilders.cpp
@@ -1,10 +1,12 @@
 // Copyright Epic Games, Inc. All Rights Reserved.
 
 #include "ChaosClothAsset/ClothEditorToolBuilders.h"
+#include "ChaosClothAsset/ClothCollectionGroup.h"
 #include "ChaosClothAsset/ClothPatternVertexType.h"
 #include "TargetInterfaces/PrimitiveComponentBackedTarget.h"
 #include "ChaosClothAsset/ClothEditorContextObject.h"
 #include "ChaosClothAsset/SelectionNode.h"
+#include "ChaosClothAsset/WeightMapNode.h"
 #include "ToolContextInterfaces.h"
 #include "ToolTargetManager.h"
 #include "ContextObjectStore.h"
@@ -17,11 +19,39 @@
 
 // ------------------- Weight Map Paint Tool -------------------
 
-void UClothEditorWeightMapPaintToolBuilder::GetSupportedViewModes(TArray<UE::Chaos::ClothAsset::EClothPatternVertexType>& Modes) const
+void UClothEditorWeightMapPaintToolBuilder::GetSupportedViewModes(const UClothEditorContextObject& ContextObject, TArray<UE::Chaos::ClothAsset::EClothPatternVertexType>& Modes) const
 {
-	Modes.Add(UE::Chaos::ClothAsset::EClothPatternVertexType::Sim3D);
-	Modes.Add(UE::Chaos::ClothAsset::EClothPatternVertexType::Sim2D);
-	Modes.Add(UE::Chaos::ClothAsset::EClothPatternVertexType::Render);
+	using namespace UE::Chaos::ClothAsset;
+	const FChaosClothAssetWeightMapNode* const WeightMapNode = ContextObject.GetSingleSelectedNodeOfType<FChaosClothAssetWeightMapNode>();
+	if (WeightMapNode)
+	{
+		if (WeightMapNode->MeshTarget == EChaosClothAssetWeightMapMeshTarget::Simulation)
+		{
+			Modes.Add(EClothPatternVertexType::Sim3D);
+			Modes.Add(EClothPatternVertexType::Sim2D);
+		}
+		else
+		{
+			check(WeightMapNode->MeshTarget == EChaosClothAssetWeightMapMeshTarget::Render);
+			Modes.Add(EClothPatternVertexType::Render);
+		}
+	}
+	else
+	{
+		// No node selected. This happens if we start the tool due to pushing the button in the toolbar -- the tool starts before the node selection can change.
+		// In this case lock to either sim or render mode, whatever is current.
+		// TODO: See if we can have the button action select the node before attempting to start the tool.
+		const EClothPatternVertexType CurrentViewMode = ContextObject.GetConstructionViewMode();
+		if (CurrentViewMode == EClothPatternVertexType::Render)
+		{
+			Modes.Add(EClothPatternVertexType::Render);
+		}
+		else
+		{
+			Modes.Add(EClothPatternVertexType::Sim3D);
+			Modes.Add(EClothPatternVertexType::Sim2D);
+		}
+	}
 }
 
 UMeshSurfacePointTool* UClothEditorWeightMapPaintToolBuilder::CreateNewTool(const FToolBuilderState& SceneState) const
@@ -39,8 +69,9 @@ UMeshSurfacePointTool* UClothEditorWeightMapPaintToolBuilder::CreateNewTool(cons
 
 // ------------------- Selection Tool -------------------
 
-void UClothMeshSelectionToolBuilder::GetSupportedViewModes(TArray<UE::Chaos::ClothAsset::EClothPatternVertexType>& Modes) const
+void UClothMeshSelectionToolBuilder::GetSupportedViewModes(const UClothEditorContextObject& ContextObject, TArray<UE::Chaos::ClothAsset::EClothPatternVertexType>& Modes) const
 {
+	// TODO: When the Secondary Selection set is removed, update this function to be similar to UClothEditorWeightMapPaintToolBuilder::GetSupportedViewModes above
 	Modes.Add(UE::Chaos::ClothAsset::EClothPatternVertexType::Sim3D);
 	Modes.Add(UE::Chaos::ClothAsset::EClothPatternVertexType::Sim2D);
 	Modes.Add(UE::Chaos::ClothAsset::EClothPatternVertexType::Render);
@@ -81,7 +112,7 @@ UInteractiveTool* UClothMeshSelectionToolBuilder::BuildTool(const FToolBuilderSt
 
 // ------------------- Skin Weight Transfer Tool -------------------
 
-void UClothTransferSkinWeightsToolBuilder::GetSupportedViewModes(TArray<UE::Chaos::ClothAsset::EClothPatternVertexType>& Modes) const
+void UClothTransferSkinWeightsToolBuilder::GetSupportedViewModes(const UClothEditorContextObject& ContextObject, TArray<UE::Chaos::ClothAsset::EClothPatternVertexType>& Modes) const
 {
 	Modes.Add(UE::Chaos::ClothAsset::EClothPatternVertexType::Sim3D);
 }
diff --git a/Engine/Plugins/ChaosClothAssetEditor/Source/ChaosClothAssetEditorTools/Public/ChaosClothAsset/ClothEditorToolBuilders.h b/Engine/Plugins/ChaosClothAssetEditor/Source/ChaosClothAssetEditorTools/Public/ChaosClothAsset/ClothEditorToolBuilders.h
--- a/Engine/Plugins/ChaosClothAssetEditor/Source/ChaosClothAssetEditorTools/Public/ChaosClothAsset/ClothEditorToolBuilders.h
+++ b/Engine/Plugins/ChaosClothAssetEditor/Source/ChaosClothAssetEditorTools/Public/ChaosClothAsset/ClothEditorToolBuilders.h
@@ -8,6 +8,8 @@
 #include "BaseTools/MeshSurfacePointMeshEditingTool.h"
 #include "ClothEditorToolBuilders.generated.h"
 
+class UClothEditorContextObject;
+
 namespace UE::Chaos::ClothAsset
 {
 	enum class EClothPatternVertexType : uint8;
@@ -28,7 +30,7 @@ class IChaosClothAssetEditorToolBuilder
 public:
 
 	/** Returns all Construction View modes that this tool can operate in. The first element should be the preferred mode to switch to if necessary. */
-	virtual void GetSupportedViewModes(TArray<UE::Chaos::ClothAsset::EClothPatternVertexType>& Modes) const = 0;
+	virtual void GetSupportedViewModes(const UClothEditorContextObject& ContextObject, TArray<UE::Chaos::ClothAsset::EClothPatternVertexType>& Modes) const = 0;
 
 	/** Returns whether or not view can be set to wireframe when this tool is active.. */
 	virtual bool CanSetConstructionViewWireframeActive() const { return true; }
@@ -41,7 +43,7 @@ class CHAOSCLOTHASSETEDITORTOOLS_API UClothEditorWeightMapPaintToolBuilder : pub
 	GENERATED_BODY()
 
 private:
-	virtual void GetSupportedViewModes(TArray<UE::Chaos::ClothAsset::EClothPatternVertexType>& Modes) const override;
+	virtual void GetSupportedViewModes(const UClothEditorContextObject& ContextObject, TArray<UE::Chaos::ClothAsset::EClothPatternVertexType>& Modes) const override;
 	virtual UMeshSurfacePointTool* CreateNewTool(const FToolBuilderState& SceneState) const override;
 	virtual bool CanSetConstructionViewWireframeActive() const { return false; }
 };
@@ -53,8 +55,7 @@ class CHAOSCLOTHASSETEDITORTOOLS_API UClothMeshSelectionToolBuilder : public UIn
 	GENERATED_BODY()
 
 private:
-
-	virtual void GetSupportedViewModes(TArray<UE::Chaos::ClothAsset::EClothPatternVertexType>& Modes) const override;
+	virtual void GetSupportedViewModes(const UClothEditorContextObject& ContextObject, TArray<UE::Chaos::ClothAsset::EClothPatternVertexType>& Modes) const override;
 	virtual bool CanBuildTool(const FToolBuilderState& SceneState) const override;
 	virtual UInteractiveTool* BuildTool(const FToolBuilderState& SceneState) const override;
 	virtual const FToolTargetTypeRequirements& GetTargetRequirements() const override;
@@ -67,8 +68,7 @@ class CHAOSCLOTHASSETEDITORTOOLS_API UClothTransferSkinWeightsToolBuilder : publ
 	GENERATED_BODY()
 
 private:
-
-	virtual void GetSupportedViewModes(TArray<UE::Chaos::ClothAsset::EClothPatternVertexType>& Modes) const override;
+	virtual void GetSupportedViewModes(const UClothEditorContextObject& ContextObject, TArray<UE::Chaos::ClothAsset::EClothPatternVertexType>& Modes) const override;
 	virtual USingleSelectionMeshEditingTool* CreateNewTool(const FToolBuilderState& SceneState) const override;
 
 };
