Cloth Remesh node: do the 2D -> 3D projection for each pattern individually to handle patterns that might overlap in 2D space

diff --git a/Engine/Plugins/ChaosClothAssetEditor/Source/ChaosClothAssetDataflowNodes/Private/ChaosClothAsset/RemeshNode.cpp b/Engine/Plugins/ChaosClothAssetEditor/Source/ChaosClothAssetDataflowNodes/Private/ChaosClothAsset/RemeshNode.cpp
--- a/Engine/Plugins/ChaosClothAssetEditor/Source/ChaosClothAssetDataflowNodes/Private/ChaosClothAsset/RemeshNode.cpp
+++ b/Engine/Plugins/ChaosClothAssetEditor/Source/ChaosClothAssetDataflowNodes/Private/ChaosClothAsset/RemeshNode.cpp
@@ -834,34 +834,6 @@ namespace UE::Chaos::ClothAsset::Private
 	}
 
 
-	void ProjectTo3D(const UE::Geometry::FDynamicMesh3& Mesh2D,
-		const UE::Geometry::FDynamicMesh3& ProjectionTarget2D,
-		const UE::Geometry::FDynamicMesh3& ProjectionTarget3D,
-		UE::Geometry::FDynamicMesh3& OutMesh3D)
-	{
-		using namespace UE::Geometry;
-
-		OutMesh3D.Copy(Mesh2D);
-
-		TSharedPtr<UE::Geometry::FDynamicMeshAABBTree3> ProjectionTargetSpatial = MakeShared<UE::Geometry::FDynamicMeshAABBTree3>(&ProjectionTarget2D, true);
-
-		for (const int VertexIndex : Mesh2D.VertexIndicesItr())
-		{
-			const FVector3d SrcVert = Mesh2D.GetVertex(VertexIndex);
-
-			double Distance;
-			const int NearestTriangle = ProjectionTargetSpatial->FindNearestTriangle(SrcVert, Distance);
-
-			const FDistPoint3Triangle3d Dist = TMeshQueries<FDynamicMesh3>::TriangleDistance(ProjectionTarget2D, NearestTriangle, SrcVert);
-			const FVector3d Bary = Dist.TriangleBaryCoords;
-			const FVector3d InterpolatedPoint = ProjectionTarget3D.GetTriBaryPoint(NearestTriangle, Bary[0], Bary[1], Bary[2]);
-
-			OutMesh3D.SetVertex(VertexIndex, InterpolatedPoint);
-		}
-	}
-
-
-
 	bool Simplify(UE::Geometry::FDynamicMesh3& Mesh, int TargetVertexCount, bool bCoarsenBoundaries, UE::Geometry::FCompactMaps* CompactMaps = nullptr)
 	{
 		using namespace UE::Geometry;
@@ -1136,14 +1108,57 @@ void FChaosClothAssetRemeshNode::RemeshSimMesh(const TSharedRef<const FManagedAr
 
 	// Project the 3D vertices onto the input 3D mesh
 
-	FDynamicMesh3 SourceMesh2D;
-	Converter.Convert(ClothCollection, INDEX_NONE, EClothPatternVertexType::Sim2D, SourceMesh2D);
-	FDynamicMesh3 SourceMesh3D;
-	Converter.Convert(ClothCollection, INDEX_NONE, EClothPatternVertexType::Sim3D, SourceMesh3D);
+	// For each 2D vertex, we will find the closest triangle on the input 2D mesh, then look up that triangle on the input 3D mesh to get the final 3D location.
+	// We will do this pattern-by-pattern to handle issues where the patterns overlap in 2D space.
+
+	const FDynamicMeshPolygroupAttribute* const NewPatternIndexLayer = Mesh2D.Attributes()->GetPolygroupLayer((int)PatternIndexLayerID);
+	check(NewPatternIndexLayer);
 
 	FDynamicMesh3 Mesh3D;
-	UE::Chaos::ClothAsset::Private::ProjectTo3D(Mesh2D, SourceMesh2D, SourceMesh3D, Mesh3D);
+	Mesh3D.Copy(Mesh2D);
 
+	TMap<int32, TSet<int32>> PatternVertexIDs;
+	for (int32 TID : Mesh2D.TriangleIndicesItr())
+	{
+		const int32 PatternID = NewPatternIndexLayer->GetValue(TID);
+		if (!PatternVertexIDs.Contains(PatternID))
+		{
+			PatternVertexIDs.Add(PatternID, TSet<int32>());
+		}
+
+		const FIndex3i Tri = Mesh2D.GetTriangle(TID);
+		PatternVertexIDs[PatternID].Add(Tri[0]);
+		PatternVertexIDs[PatternID].Add(Tri[1]);
+		PatternVertexIDs[PatternID].Add(Tri[2]);
+	}
+
+	for (int32 PatternID = 0; PatternID < InClothFacade.GetNumSimPatterns(); ++PatternID)
+	{
+		if (!PatternVertexIDs.Contains(PatternID))
+		{
+			continue;
+		}
+
+		FDynamicMesh3 ProjectionTarget2D;
+		Converter.Convert(ClothCollection, PatternID, EClothPatternVertexType::Sim2D, ProjectionTarget2D);
+		FDynamicMesh3 ProjectionTarget3D;
+		Converter.Convert(ClothCollection, PatternID, EClothPatternVertexType::Sim3D, ProjectionTarget3D);
+		TSharedPtr<UE::Geometry::FDynamicMeshAABBTree3> ProjectionTargetSpatial = MakeShared<UE::Geometry::FDynamicMeshAABBTree3>(&ProjectionTarget2D, true);
+
+		for (const int32 VID : PatternVertexIDs[PatternID])
+		{
+			const FVector3d SrcVert = Mesh2D.GetVertex(VID);
+
+			double Distance;
+			const int NearestTriangle = ProjectionTargetSpatial->FindNearestTriangle(SrcVert, Distance);
+
+			const FDistPoint3Triangle3d Dist = TMeshQueries<FDynamicMesh3>::TriangleDistance(ProjectionTarget2D, NearestTriangle, SrcVert);
+			const FVector3d Bary = Dist.TriangleBaryCoords;
+			const FVector3d InterpolatedPoint = ProjectionTarget3D.GetTriBaryPoint(NearestTriangle, Bary[0], Bary[1], Bary[2]);
+
+			Mesh3D.SetVertex(VID, InterpolatedPoint);
+		}
+	}
 
 	// Build the output cloth sim mesh
 
