Chaos Cloth: fix bug in new merge collections node because the Serialize code is expecting 1 rather than 2 initial inputs.

diff --git a/Engine/Plugins/ChaosClothAssetEditor/Source/ChaosClothAssetDataflowNodes/Private/ChaosClothAsset/MergeClothCollectionsNode.cpp b/Engine/Plugins/ChaosClothAssetEditor/Source/ChaosClothAssetDataflowNodes/Private/ChaosClothAsset/MergeClothCollectionsNode.cpp
--- a/Engine/Plugins/ChaosClothAssetEditor/Source/ChaosClothAssetDataflowNodes/Private/ChaosClothAsset/MergeClothCollectionsNode.cpp
+++ b/Engine/Plugins/ChaosClothAssetEditor/Source/ChaosClothAssetDataflowNodes/Private/ChaosClothAsset/MergeClothCollectionsNode.cpp
@@ -474,10 +474,11 @@ void FChaosClothAssetMergeClothCollectionsNode_v2::Serialize(FArchive& Ar)
 {
 	if (Ar.IsLoading())
 	{
-		check(Collections.Num() > 0);
+		check(Collections.Num() > 1);
 		check(FindInput(GetConnectionReference(0)));
+		check(FindInput(GetConnectionReference(1)));
 
-		for (int32 Index = 1; Index < Collections.Num(); ++Index)
+		for (int32 Index = 2; Index < Collections.Num(); ++Index)
 		{
 			RegisterInputArrayConnection(GetConnectionReference(Index));
 		}
