Adding more comments.

diff --git a/Engine/Plugins/ChaosClothAssetEditor/Source/ChaosClothAssetDataflowNodes/Public/ChaosClothAsset/TransferSkinWeightsNode.h b/Engine/Plugins/ChaosClothAssetEditor/Source/ChaosClothAssetDataflowNodes/Public/ChaosClothAsset/TransferSkinWeightsNode.h
--- a/Engine/Plugins/ChaosClothAssetEditor/Source/ChaosClothAssetDataflowNodes/Public/ChaosClothAsset/TransferSkinWeightsNode.h
+++ b/Engine/Plugins/ChaosClothAssetEditor/Source/ChaosClothAssetDataflowNodes/Public/ChaosClothAsset/TransferSkinWeightsNode.h
@@ -134,5 +134,28 @@ public:
 	FChaosClothAssetTransferSkinWeightsNode(const Dataflow::FNodeParameters& InParam, FGuid InGuid = FGuid::NewGuid());
 
 private:
+
+	/** 
+	 * The node takes a user-selected skinned SkeletalMesh asset and transfers the skin weights to a simulation/render
+	 * meshes using a selected algorithm. You can either transfer to simulation and render meshes separately 
+	 * or transfer to both in one go using the TargetMeshType property. When transferring to render meshes, you can 
+	 * control the source of the transfer which can either be the body or the simulation mesh (recommended). In the latter 
+	 * case we always use the ClosestPointOnSurface algorithm.
+	 * 
+	 * 
+	 * InpaintWeights algorithm:
+	 * Main algorithm for transferring weights, which is based on the "Robust Skin Weights Transfer via Weight Inpainting 
+	 * Siggraph Asia 2023". The implementation and explanation of the algorithm can be found in 
+	 * "Engine\Plugins\Runtime\GeometryProcessing\Source\DynamicMesh\Private\Operations\TransferBoneWeights.h(cpp)"
+	 * 
+	 * 
+	 * Handling of disconnected render meshes:
+	 * It is usually the case that sim mesh is welded and manifold meaning that the inpaint method should always 
+	 * succeed and give the best results. However, the render mesh is often not welded and consists of multiple 
+	 * disconnected parts. This is usually fine, and inpaint should work well except in places where there 
+	 * is a big crease along the stitch (like armpit areas), so vertices that are close to each other can have very 
+	 * different normals which could potentially lead to different weights being computed. You can either try to 
+	 * increase the normal threshold or switch to the closest point method.
+	 */
 	virtual void Evaluate(Dataflow::FContext& Context, const FDataflowOutput* Out) const override;
 };
