Added conversion of skeletal mesh render data to dynamic mesh when mesh description is not available.

#rnx

diff --git a/Engine/Plugins/ChaosClothAssetEditor/Source/ChaosClothAssetDataflowNodes/ChaosClothAssetDataflowNodes.Build.cs b/Engine/Plugins/ChaosClothAssetEditor/Source/ChaosClothAssetDataflowNodes/ChaosClothAssetDataflowNodes.Build.cs
--- a/Engine/Plugins/ChaosClothAssetEditor/Source/ChaosClothAssetDataflowNodes/ChaosClothAssetDataflowNodes.Build.cs
+++ b/Engine/Plugins/ChaosClothAssetEditor/Source/ChaosClothAssetDataflowNodes/ChaosClothAssetDataflowNodes.Build.cs
@@ -38,6 +38,7 @@ public class ChaosClothAssetDataflowNodes : ModuleRules
 				"MeshUtilitiesCommon",
 				"ModelingOperatorsEditorOnly",	// TODO: Someday remove editor dependencies, see UE-206172
 				"ModelingOperators",
+				"MeshConversionEngineTypes",
 				"RenderCore",
 				"SkeletalMeshDescription",
 				"Slate",
diff --git a/Engine/Plugins/ChaosClothAssetEditor/Source/ChaosClothAssetDataflowNodes/Private/ChaosClothAsset/TransferSkinWeightsNode.cpp b/Engine/Plugins/ChaosClothAssetEditor/Source/ChaosClothAssetDataflowNodes/Private/ChaosClothAsset/TransferSkinWeightsNode.cpp
--- a/Engine/Plugins/ChaosClothAssetEditor/Source/ChaosClothAssetDataflowNodes/Private/ChaosClothAsset/TransferSkinWeightsNode.cpp
+++ b/Engine/Plugins/ChaosClothAssetEditor/Source/ChaosClothAssetDataflowNodes/Private/ChaosClothAsset/TransferSkinWeightsNode.cpp
@@ -22,6 +22,8 @@
 #include "Selections/MeshConnectedComponents.h"
 #include "BoneWeights.h"
 #include "ChaosClothAsset/ClothCollectionGroup.h"
+#include "Rendering/SkeletalMeshRenderData.h"
+#include "SkeletalMeshLODRenderDataToDynamicMesh.h"
 
 #include UE_INLINE_GENERATED_CPP_BY_NAME(TransferSkinWeightsNode)
 
@@ -71,12 +73,47 @@ namespace UE::Chaos::ClothAsset::Private
 		return Attribute;
 	}
 
-	/** Convert the USkeletalMesh to DynamicMesh. */
-	static void SkeletalMeshToDynamicMesh(USkeletalMesh* FromSkeletalMeshAsset, int32 LodIndex, FDynamicMesh3& ToDynamicMesh)
+	/** 
+	 * Convert the USkeletalMesh to DynamicMesh. If requested LOD was auto-generated, will convert render data instead of the mesh description.
+	 */
+	static bool SkeletalMeshToDynamicMesh(const USkeletalMesh* SkeletalMesh, const int32 LodIndex, FDynamicMesh3& ToDynamicMesh)
 	{
-		const FMeshDescription* SourceMesh = FromSkeletalMeshAsset->GetMeshDescription(LodIndex);
-		FMeshDescriptionToDynamicMesh Converter;
-		Converter.Convert(SourceMesh, ToDynamicMesh);
+		if (SkeletalMesh->HasMeshDescription(LodIndex))
+		{
+			const FMeshDescription* SourceMesh = SkeletalMesh->GetMeshDescription(LodIndex);
+			if (!SourceMesh)
+			{
+				return false;
+			}
+
+			FMeshDescriptionToDynamicMesh Converter;
+			Converter.Convert(SourceMesh, ToDynamicMesh);
+		}
+		else
+		{
+			const FSkeletalMeshRenderData* RenderData = SkeletalMesh->GetResourceForRendering();
+			if (!RenderData)
+			{
+				return false;
+			}
+
+			if (!RenderData->LODRenderData.IsValidIndex(LodIndex))
+			{
+				return false;
+			}
+
+			const FSkeletalMeshLODRenderData* SkeletalMeshLODRenderData = &(RenderData->LODRenderData[LodIndex]);
+
+			UE::Geometry::FSkeletalMeshLODRenderDataToDynamicMesh::ConversionOptions ConversionOptions;
+			ConversionOptions.bWantUVs = false;
+			ConversionOptions.bWantVertexColors = false;
+			ConversionOptions.bWantMaterialIDs = false;
+			ConversionOptions.bWantSkinWeights = true;
+
+			UE::Geometry::FSkeletalMeshLODRenderDataToDynamicMesh::Convert(SkeletalMeshLODRenderData, SkeletalMesh->GetRefSkeleton(), ConversionOptions, ToDynamicMesh);
+		}
+
+		return true;
 	}
 
 	/** Convert the ClothCollection to DynamicMesh. */
@@ -750,24 +787,19 @@ void FChaosClothAssetTransferSkinWeightsNode::Evaluate(Dataflow::FContext& Conte
 				return;
 			}
 
-			if (!SkeletalMesh->HasMeshDescription(LodIndex))
+			FDynamicMesh3 SourceDynamicMesh;
+			if (!SkeletalMeshToDynamicMesh(SkeletalMesh, LodIndex, SourceDynamicMesh))
 			{
-				//TODO: Once the support for converting render data to dynamic meshes is added we can support auto LODs
 				FClothDataflowTools::LogAndToastWarning(*this,
-					LOCTEXT("AutogeneratedLodIndexHeadline", "Autogenerated LOD not supported."),
+					LOCTEXT("InvalidLodHeadline", "Could not convert LOD to Dynamic Mesh."),
 					FText::Format(
-						LOCTEXT("AutogeneratedLodIndexDetails", "Auto generated LOD index {0} for skeletal mesh {1} is not supported."),
+						LOCTEXT("InvalidLodDetails", "Could not convert LOD index {0} of the skeletal mesh {1} to dyanamic mesh."),
 						LodIndex,
 						FText::FromString(SkeletalMesh.GetName())));
 				SetValue(Context, MoveTemp(*ClothCollection), &Collection);
 				return;
 			}
 
-			//
-			// Convert source Skeletal Mesh to Dynamic Mesh.
-			//
-			FDynamicMesh3 SourceDynamicMesh;
-			SkeletalMeshToDynamicMesh(SkeletalMesh, LodIndex, SourceDynamicMesh);
 			MeshTransforms::ApplyTransform(SourceDynamicMesh, Transform, true);
 			const FReferenceSkeleton& TargetRefSkeleton = SkeletalMesh->GetRefSkeleton();
 
