Chaos Cloth Asset:
- Added imported section indices to the StaticMeshImport node.
- Added StaticMesh input to the StaticMeshImport node.
- Added SkeletalMesh input to the SkeletalMeshImport node.

diff --git a/Engine/Plugins/ChaosClothAssetEditor/Source/ChaosClothAssetDataflowNodes/Private/ChaosClothAsset/SkeletalMeshImportNode.cpp b/Engine/Plugins/ChaosClothAssetEditor/Source/ChaosClothAssetDataflowNodes/Private/ChaosClothAsset/SkeletalMeshImportNode.cpp
--- a/Engine/Plugins/ChaosClothAssetEditor/Source/ChaosClothAssetDataflowNodes/Private/ChaosClothAsset/SkeletalMeshImportNode.cpp
+++ b/Engine/Plugins/ChaosClothAssetEditor/Source/ChaosClothAssetDataflowNodes/Private/ChaosClothAsset/SkeletalMeshImportNode.cpp
@@ -20,6 +20,7 @@
 FChaosClothAssetSkeletalMeshImportNode::FChaosClothAssetSkeletalMeshImportNode(const Dataflow::FNodeParameters& InParam, FGuid InGuid)
 	: FDataflowNode(InParam, InGuid)
 {
+	RegisterInputConnection(&SkeletalMesh);
 	RegisterOutputConnection(&Collection);
 }
 
@@ -33,9 +34,9 @@ void FChaosClothAssetSkeletalMeshImportNode::Evaluate(Dataflow::FContext& Contex
 		FCollectionClothFacade ClothFacade(ClothCollection);
 		ClothFacade.DefineSchema();
 
-		if (SkeletalMesh)
+		if (const USkeletalMesh* const InSkeletalMesh = GetValue(Context, &SkeletalMesh))
 		{
-			const FSkeletalMeshModel* ImportedModel = SkeletalMesh->GetImportedModel();
+			const FSkeletalMeshModel* ImportedModel = InSkeletalMesh->GetImportedModel();
 			const bool bIsValidLOD = ImportedModel && ImportedModel->LODModels.IsValidIndex(LODIndex);
 			if (!bIsValidLOD)
 			{
@@ -43,7 +44,7 @@ void FChaosClothAssetSkeletalMeshImportNode::Evaluate(Dataflow::FContext& Contex
 					FText::Format(
 						LOCTEXT("InvalidLODDetails", "No valid LOD {0} found for skeletal mesh {1}."),
 						LODIndex,
-						FText::FromString(SkeletalMesh->GetName())));
+						FText::FromString(InSkeletalMesh->GetName())));
 
 				SetValue(Context, MoveTemp(*ClothCollection), &Collection);
 				return;
@@ -62,7 +63,7 @@ void FChaosClothAssetSkeletalMeshImportNode::Evaluate(Dataflow::FContext& Contex
 						FText::Format(
 							LOCTEXT("InvalidSectionDetails", "No valid section {0} found for skeletal mesh {1}."),
 							Section,
-							FText::FromString(SkeletalMesh->GetName())));
+							FText::FromString(InSkeletalMesh->GetName())));
 
 					continue;
 				}
@@ -74,19 +75,19 @@ void FChaosClothAssetSkeletalMeshImportNode::Evaluate(Dataflow::FContext& Contex
 
 				if (bImportRenderMesh)
 				{
-					const TArray<FSkeletalMaterial>& Materials = SkeletalMesh->GetMaterials();
+					const TArray<FSkeletalMaterial>& Materials = InSkeletalMesh->GetMaterials();
 					check(Section < Materials.Num());
 					const FString RenderMaterialPathName = Materials[Section].MaterialInterface ? Materials[Section].MaterialInterface->GetPathName() : "";
 					FClothDataflowTools::AddRenderPatternFromSkeletalMeshSection(ClothCollection, LODModel, Section, RenderMaterialPathName);
 				}
 			}
 
-			if (const UPhysicsAsset* PhysicsAsset = bSetPhysicsAsset ? SkeletalMesh->GetPhysicsAsset() : nullptr)
+			if (const UPhysicsAsset* PhysicsAsset = bSetPhysicsAsset ? InSkeletalMesh->GetPhysicsAsset() : nullptr)
 			{
 				ClothFacade.SetPhysicsAssetPathName(PhysicsAsset->GetPathName());
 			}
 
-			ClothFacade.SetSkeletalMeshPathName(SkeletalMesh->GetPathName());
+			ClothFacade.SetSkeletalMeshPathName(InSkeletalMesh->GetPathName());
 		}
 		SetValue(Context, MoveTemp(*ClothCollection), &Collection);
 	}
diff --git a/Engine/Plugins/ChaosClothAssetEditor/Source/ChaosClothAssetDataflowNodes/Private/ChaosClothAsset/StaticMeshImportNode.cpp b/Engine/Plugins/ChaosClothAssetEditor/Source/ChaosClothAssetDataflowNodes/Private/ChaosClothAsset/StaticMeshImportNode.cpp
--- a/Engine/Plugins/ChaosClothAssetEditor/Source/ChaosClothAssetDataflowNodes/Private/ChaosClothAsset/StaticMeshImportNode.cpp
+++ b/Engine/Plugins/ChaosClothAssetEditor/Source/ChaosClothAssetDataflowNodes/Private/ChaosClothAsset/StaticMeshImportNode.cpp
@@ -17,25 +17,40 @@
 
 namespace UE::Chaos::ClothAsset::Private
 {
-	static bool InitializeDataFromMeshDescription(
-		const UStaticMesh* const StaticMesh,
-		const FMeshDescription* const InMeshDescription,
-		const FMeshBuildSettings& InBuildSettings,
-		const TArray<FStaticMaterial>& StaticMaterials,
-		const TSharedRef<FManagedArrayCollection>& ClothCollection)
+	static bool InitializeSimMeshFromSkeletalMeshModel(
+		const FSkeletalMeshLODModel& SkeletalMeshLODModel,
+		const TSharedRef<FManagedArrayCollection>& ClothCollection,
+		int32 InSectionIndex,
+		int32 UVChannelIndex,
+		const FVector2f& UVScale)
 	{
-		FSkeletalMeshLODModel SkeletalMeshModel;
-		if (FClothDataflowTools::BuildSkeletalMeshModelFromMeshDescription(InMeshDescription, InBuildSettings, SkeletalMeshModel))
+		if (InSectionIndex == INDEX_NONE || SkeletalMeshLODModel.Sections.IsValidIndex(InSectionIndex))
 		{
-			FStaticMeshConstAttributes MeshAttributes(*InMeshDescription);
-			TPolygonGroupAttributesConstRef<FName> MaterialSlotNames = MeshAttributes.GetPolygonGroupMaterialSlotNames();
-			for (int32 SectionIndex = 0; SectionIndex < SkeletalMeshModel.Sections.Num(); ++SectionIndex)
+			const int32 StartSection = (InSectionIndex != INDEX_NONE) ? InSectionIndex : 0;
+			const int32 EndSection = (InSectionIndex != INDEX_NONE) ? InSectionIndex + 1 : SkeletalMeshLODModel.Sections.Num();
+			for (int32 SectionIndex = StartSection; SectionIndex < EndSection; ++SectionIndex)
 			{
-				// Section MaterialIndex refers to the polygon group index. Look up which material this corresponds with.
-				const FName& MaterialSlotName = MaterialSlotNames[SkeletalMeshModel.Sections[SectionIndex].MaterialIndex];
-				const int32 MaterialIndex = StaticMesh->GetMaterialIndexFromImportedMaterialSlotName(MaterialSlotName);
-				const FString RenderMaterialPathName = StaticMaterials.IsValidIndex(MaterialIndex)&& StaticMaterials[MaterialIndex].MaterialInterface ? StaticMaterials[MaterialIndex].MaterialInterface->GetPathName() : "";
-				FClothDataflowTools::AddRenderPatternFromSkeletalMeshSection(ClothCollection, SkeletalMeshModel, SectionIndex, RenderMaterialPathName);
+				FClothDataflowTools::AddSimPatternsFromSkeletalMeshSection(ClothCollection, SkeletalMeshLODModel, SectionIndex, UVChannelIndex, UVScale);
+			}
+			return true;
+		}
+		return false;
+	}
+
+	static bool InitializeRenderMeshFromSkeletalMeshModel(
+		const FSkeletalMeshLODModel& SkeletalMeshLODModel,
+		const TSharedRef<FManagedArrayCollection>& ClothCollection,
+		int32 InSectionIndex,
+		const TFunctionRef<FString(int32 /*SectionIndex*/)>& GetMaterialPathNameFunction)
+	{
+		if (InSectionIndex == INDEX_NONE || SkeletalMeshLODModel.Sections.IsValidIndex(InSectionIndex))
+		{
+			const int32 StartSection = (InSectionIndex != INDEX_NONE) ? InSectionIndex : 0;
+			const int32 EndSection = (InSectionIndex != INDEX_NONE) ? InSectionIndex + 1 : SkeletalMeshLODModel.Sections.Num();
+			for (int32 SectionIndex = StartSection; SectionIndex < EndSection; ++SectionIndex)
+			{
+				const FString MaterialPathName = GetMaterialPathNameFunction(SectionIndex);
+				FClothDataflowTools::AddRenderPatternFromSkeletalMeshSection(ClothCollection, SkeletalMeshLODModel, SectionIndex, MaterialPathName);
 			}
 			return true;
 		}
@@ -46,6 +61,7 @@ namespace UE::Chaos::ClothAsset::Private
 FChaosClothAssetStaticMeshImportNode::FChaosClothAssetStaticMeshImportNode(const Dataflow::FNodeParameters& InParam, FGuid InGuid)
 	: FDataflowNode(InParam, InGuid)
 {
+	RegisterInputConnection(&StaticMesh);
 	RegisterOutputConnection(&Collection);
 }
 
@@ -61,39 +77,73 @@ void FChaosClothAssetStaticMeshImportNode::Evaluate(Dataflow::FContext& Context,
 		FCollectionClothFacade ClothFacade(ClothCollection);
 		ClothFacade.DefineSchema();
 
-		if (StaticMesh && (bImportSimMesh || bImportRenderMesh))
+		const UStaticMesh* const InStaticMesh = GetValue(Context, &StaticMesh);
+
+		if (InStaticMesh && (bImportSimMesh || bImportRenderMesh))
 		{
-			const int32 NumLods = StaticMesh->GetNumSourceModels();
+			const int32 NumLods = InStaticMesh->GetNumSourceModels();
 			if (LODIndex < NumLods)
 			{
-				const FMeshDescription* const MeshDescription = StaticMesh->GetMeshDescription(LODIndex);
+				const FMeshDescription* const MeshDescription = InStaticMesh->GetMeshDescription(LODIndex);
 				check(MeshDescription);
 
-				if (bImportSimMesh)
+				if (bImportSimMesh && SimMeshSection == INDEX_NONE)
 				{
+					// Keep legacy welding behavior when the entire static mesh is imported
 					FMeshDescriptionToDynamicMesh Converter;
 					Converter.bPrintDebugMessages = false;
 					Converter.bEnableOutputGroups = false;
 					Converter.bVIDsFromNonManifoldMeshDescriptionAttr = true;
 					UE::Geometry::FDynamicMesh3 DynamicMesh;
-
 					Converter.Convert(MeshDescription, DynamicMesh);
 					constexpr bool bAppend = false;
 					FClothGeometryTools::BuildSimMeshFromDynamicMesh(ClothCollection, DynamicMesh, UVChannel, UVScale, bAppend);
 				}
-				if (bImportRenderMesh)
+				const bool bNeedsSkeletalMeshModel = bImportRenderMesh || (bImportSimMesh && SimMeshSection != INDEX_NONE);
+
+				FSkeletalMeshLODModel SkeletalMeshLODModel;
+				const FMeshBuildSettings& BuildSettings = InStaticMesh->GetSourceModel(LODIndex).BuildSettings;
+				if (bNeedsSkeletalMeshModel && 
+					FClothDataflowTools::BuildSkeletalMeshModelFromMeshDescription(MeshDescription, InStaticMesh->GetSourceModel(LODIndex).BuildSettings, SkeletalMeshLODModel))
 				{
-					// Add render data (section = pattern)
-					if (!InitializeDataFromMeshDescription(StaticMesh, MeshDescription, StaticMesh->GetSourceModel(LODIndex).BuildSettings, StaticMesh->GetStaticMaterials(), ClothCollection))
+					if (bImportSimMesh && SimMeshSection != INDEX_NONE)
 					{
-						FClothDataflowTools::LogAndToastWarning(*this,
-							LOCTEXT("InvalidRenderMeshHeadline", "Invalid render mesh."),
-							FText::Format(
-								LOCTEXT("InvalidRenderMeshDetails", "The input static mesh {0} failed to convert to a valid render mesh."),
-								FText::FromString(StaticMesh->GetName())));
+						if (!InitializeSimMeshFromSkeletalMeshModel(SkeletalMeshLODModel, ClothCollection, SimMeshSection, UVChannel, UVScale))
+						{
+							FClothDataflowTools::LogAndToastWarning(*this,
+								LOCTEXT("InvalidRenderMeshHeadline", "Invalid render mesh."),
+								FText::Format(
+									LOCTEXT("InvalidRenderMeshDetails", "The input static mesh {0} failed to convert to a valid render mesh."),
+									FText::FromString(InStaticMesh->GetName())));
+						}
+					}
+					if (bImportRenderMesh)
+					{
+						const FStaticMeshConstAttributes MeshAttributes(*MeshDescription);
+						const TPolygonGroupAttributesConstRef<FName> MaterialSlotNames = MeshAttributes.GetPolygonGroupMaterialSlotNames();
+						auto GetMaterialPathName = [&InStaticMesh, &SkeletalMeshLODModel, &MaterialSlotNames](int32 SectionIndex) -> FString
+							{
+								// Section MaterialIndex refers to the polygon group index. Look up which material this corresponds with.
+								const FName& MaterialSlotName = MaterialSlotNames[SkeletalMeshLODModel.Sections[SectionIndex].MaterialIndex];
+								const int32 MaterialIndex = InStaticMesh->GetMaterialIndexFromImportedMaterialSlotName(MaterialSlotName);
+								const TArray<FStaticMaterial>& StaticMaterials = InStaticMesh->GetStaticMaterials();
+								if (StaticMaterials.IsValidIndex(MaterialIndex) && StaticMaterials[MaterialIndex].MaterialInterface)
+								{
+									return StaticMaterials[MaterialIndex].MaterialInterface->GetPathName();
+								}
+								return FString();
+							};
+
+						if (!InitializeRenderMeshFromSkeletalMeshModel(SkeletalMeshLODModel, ClothCollection, RenderMeshSection, GetMaterialPathName))
+						{
+							FClothDataflowTools::LogAndToastWarning(*this,
+								LOCTEXT("InvalidRenderMeshHeadline", "Invalid render mesh."),
+								FText::Format(
+									LOCTEXT("InvalidRenderMeshDetails", "The input static mesh {0} failed to convert to a valid render mesh."),
+									FText::FromString(InStaticMesh->GetName())));
+						}
 					}
 				}
-				
 				// Bind to root bone by default
 				FClothGeometryTools::BindMeshToRootBone(ClothCollection, bImportSimMesh, bImportRenderMesh);
 			}
@@ -104,7 +154,7 @@ void FChaosClothAssetStaticMeshImportNode::Evaluate(Dataflow::FContext& Context,
 					FText::Format(LOCTEXT("InvalidLODIndexDetails",
 						"{0} is not a valid LOD index for static mesh {1}.\n"
 						"This static mesh has {2} LOD(s)."),
-						FText::FromString(StaticMesh->GetName()),
+						FText::FromString(InStaticMesh->GetName()),
 						LODIndex,
 						NumLods));
 			}
diff --git a/Engine/Plugins/ChaosClothAssetEditor/Source/ChaosClothAssetDataflowNodes/Public/ChaosClothAsset/SkeletalMeshImportNode.h b/Engine/Plugins/ChaosClothAssetEditor/Source/ChaosClothAssetDataflowNodes/Public/ChaosClothAsset/SkeletalMeshImportNode.h
--- a/Engine/Plugins/ChaosClothAssetEditor/Source/ChaosClothAssetDataflowNodes/Public/ChaosClothAsset/SkeletalMeshImportNode.h
+++ b/Engine/Plugins/ChaosClothAssetEditor/Source/ChaosClothAssetDataflowNodes/Public/ChaosClothAsset/SkeletalMeshImportNode.h
@@ -20,7 +20,7 @@ public:
 	FManagedArrayCollection Collection;
 
 	/** The skeletal mesh to import. */
-	UPROPERTY(EditAnywhere, Category = "Skeletal Mesh Import")
+	UPROPERTY(EditAnywhere, Category = "Skeletal Mesh Import", Meta = (DataflowInput))
 	TObjectPtr<const USkeletalMesh> SkeletalMesh;
 
 	/** The skeletal mesh LOD to import. */
diff --git a/Engine/Plugins/ChaosClothAssetEditor/Source/ChaosClothAssetDataflowNodes/Public/ChaosClothAsset/StaticMeshImportNode.h b/Engine/Plugins/ChaosClothAssetEditor/Source/ChaosClothAssetDataflowNodes/Public/ChaosClothAsset/StaticMeshImportNode.h
--- a/Engine/Plugins/ChaosClothAssetEditor/Source/ChaosClothAssetDataflowNodes/Public/ChaosClothAsset/StaticMeshImportNode.h
+++ b/Engine/Plugins/ChaosClothAssetEditor/Source/ChaosClothAssetDataflowNodes/Public/ChaosClothAsset/StaticMeshImportNode.h
@@ -20,7 +20,7 @@ public:
 	FManagedArrayCollection Collection;
 
 	/* The Static Mesh to import from. */
-	UPROPERTY(EditAnywhere, Category = "Static Mesh Import")
+	UPROPERTY(EditAnywhere, Category = "Static Mesh Import", Meta = (DataflowInput))
 	TObjectPtr<const UStaticMesh> StaticMesh;
 
 	/* Which static mesh Lod to import. */
@@ -31,21 +31,29 @@ public:
 	UPROPERTY(EditAnywhere, Category = "Static Mesh Import")
 	bool bImportSimMesh = true;
 
-	/* Import static mesh data as render mesh data. */
-	UPROPERTY(EditAnywhere, Category = "Static Mesh Import")
-	bool bImportRenderMesh = true;
+	/* Material section to import as sim mesh data. Use -1 to import all sections. */
+	UPROPERTY(EditAnywhere, Category = "Static Mesh Import", Meta = (ClampMin = "-1", EditCondition = "bImportSimMesh"))
+	int32 SimMeshSection = INDEX_NONE;
 
 	/**
 	 * UV channel of the static mesh to import the 2D simulation mesh patterns from.
 	 * If set to -1, or the specified UVChannel doesn't exist then the import will unwrap the 3D simulation mesh into 2D simulation mesh patterns.
 	 */
-	UPROPERTY(EditAnywhere, Category = "Static Mesh Import", meta = (ClampMin = "-1", EditCondition = "bImportSimMesh"))
+	UPROPERTY(EditAnywhere, Category = "Static Mesh Import", Meta = (ClampMin = "-1", EditCondition = "bImportSimMesh"))
 	int32 UVChannel = 0;
 
 	/* Apply this scale to the UVs when populating Sim Mesh positions. */
 	UPROPERTY(EditAnywhere, Category = "Static Mesh Import", Meta = (AllowPreserveRatio, EditCondition = "bImportSimMesh && (UVChannel >= 0)"))
 	FVector2f UVScale = { 1.f, 1.f };
 
+	/* Import static mesh data as render mesh data. */
+	UPROPERTY(EditAnywhere, Category = "Static Mesh Import")
+	bool bImportRenderMesh = true;
+
+	/* Material section to import as render mesh data. Use -1 to import all sections. */
+	UPROPERTY(EditAnywhere, Category = "Static Mesh Import", Meta = (ClampMin = "-1", EditCondition = "bImportRenderMesh"))
+	int32 RenderMeshSection = INDEX_NONE;
+
 	FChaosClothAssetStaticMeshImportNode(const Dataflow::FNodeParameters& InParam, FGuid InGuid = FGuid::NewGuid());
 
 private:
