Chaos Cloth Asset - Added the Name property to the the inputs of the Attribute node.

#rnx

diff --git a/Engine/Plugins/ChaosClothAssetEditor/Source/ChaosClothAssetDataflowNodes/Private/ChaosClothAsset/AttributeNode.cpp b/Engine/Plugins/ChaosClothAssetEditor/Source/ChaosClothAssetDataflowNodes/Private/ChaosClothAsset/AttributeNode.cpp
--- a/Engine/Plugins/ChaosClothAssetEditor/Source/ChaosClothAssetDataflowNodes/Private/ChaosClothAsset/AttributeNode.cpp
+++ b/Engine/Plugins/ChaosClothAssetEditor/Source/ChaosClothAssetDataflowNodes/Private/ChaosClothAsset/AttributeNode.cpp
@@ -4,12 +4,121 @@
 #include "ChaosClothAsset/ClothCollectionGroup.h"
 #include "ChaosClothAsset/CollectionClothFacade.h"
 #include "ChaosClothAsset/ClothDataflowTools.h"
+#include "ChaosClothAsset/WeightedValue.h"
 #include "Dataflow/DataflowInputOutput.h"
 
 #include UE_INLINE_GENERATED_CPP_BY_NAME(AttributeNode)
 
 #define LOCTEXT_NAMESPACE "ChaosClothAssetAttributeNode"
 
+FChaosClothAssetAttributeNode_V2::FChaosClothAssetAttributeNode_V2(const Dataflow::FNodeParameters& InParam, FGuid InGuid)
+	: FDataflowNode(InParam, InGuid)
+{
+	RegisterInputConnection(&Collection);
+	RegisterInputConnection(&Name.StringValue, GET_MEMBER_NAME_CHECKED(FChaosClothAssetConnectableIOStringValue, StringValue));
+	RegisterOutputConnection(&Collection, &Collection);
+	RegisterOutputConnection(&Name.StringValue, &Name.StringValue, GET_MEMBER_NAME_CHECKED(FChaosClothAssetConnectableIOStringValue, StringValue));
+}
+
+void FChaosClothAssetAttributeNode_V2::Evaluate(Dataflow::FContext& Context, const FDataflowOutput* Out) const
+{
+	Name.StringValue_Override = GetValue<FString>(Context, &Name.StringValue, UE::Chaos::ClothAsset::FWeightMapTools::NotOverridden);
+
+	if (Out->IsA<FManagedArrayCollection>(&Collection))
+	{
+		using namespace UE::Chaos::ClothAsset;
+
+		// Evaluate in collection
+		FManagedArrayCollection InCollection = GetValue<FManagedArrayCollection>(Context, &Collection);
+		const TSharedRef<FManagedArrayCollection> ClothCollection = MakeShared<FManagedArrayCollection>(MoveTemp(InCollection));
+
+		FCollectionClothFacade Cloth(ClothCollection);
+		const FName GroupName = *Group.Name;
+
+		const FString& InName = GetValue(Context, &Name.StringValue);
+
+		if (Cloth.IsValid() && !InName.IsEmpty())
+		{
+			if (ClothCollection->HasGroup(GroupName))
+			{
+				switch (Type)
+				{
+				case EChaosClothAssetNodeAttributeType::Integer:
+					Cloth.AddUserDefinedAttribute<int32>(*InName, GroupName);
+					{
+						TArrayView<int32> Values = Cloth.GetUserDefinedAttribute<int32>(*InName, GroupName);
+						for (int32& Value : Values)
+						{
+							Value = IntValue;
+						}
+					}
+					break;
+				case EChaosClothAssetNodeAttributeType::Float:
+					Cloth.AddUserDefinedAttribute<float>(*InName, GroupName);
+					{
+						TArrayView<float> Values = Cloth.GetUserDefinedAttribute<float>(*InName, GroupName);
+						for (float& Value : Values)
+						{
+							Value = FloatValue;
+						}
+					}
+					break;
+				case EChaosClothAssetNodeAttributeType::Vector:
+					Cloth.AddUserDefinedAttribute<FVector3f>(*InName, GroupName);
+					{
+						TArrayView<FVector3f> Values = Cloth.GetUserDefinedAttribute<FVector3f>(*InName, GroupName);
+						for (FVector3f& Value : Values)
+						{
+							Value = VectorValue;
+						}
+					}
+					break;
+				}
+			}
+			else if (!GroupName.IsNone())
+			{
+				FClothDataflowTools::LogAndToastWarning(
+					*this,
+					LOCTEXT("CreateAttributeHeadline", "Invalid Group"),
+					FText::Format(LOCTEXT("CreateAttributeDetail", "No group \"{0}\" currently exists on the input collection"), FText::FromName(GroupName)));
+			}
+		}
+		SetValue(Context, MoveTemp(*ClothCollection), &Collection);
+	}
+	else if (Out->IsA<FString>(&Name.StringValue))
+	{
+		const FString& InName = GetValue(Context, &Name.StringValue);
+		SetValue(Context, InName, &Name.StringValue);
+	}
+}
+
+void FChaosClothAssetAttributeNode_V2::OnSelected(Dataflow::FContext& Context)
+{
+	using namespace UE::Chaos::ClothAsset;
+
+	// Re-evaluate the input collection
+	FManagedArrayCollection InCollection = GetValue<FManagedArrayCollection>(Context, &Collection);
+	const TSharedRef<FManagedArrayCollection> ClothCollection = MakeShared<FManagedArrayCollection>(MoveTemp(InCollection));
+	FCollectionClothFacade Cloth(ClothCollection);
+
+	// Update the list of used group for the UI customization
+	const TArray<FName> GroupNames = ClothCollection->GroupNames();
+	CachedCollectionGroupNames.Reset(GroupNames.Num());
+	for (const FName& GroupName : GroupNames)
+	{
+		if (Cloth.IsValidClothCollectionGroupName(GroupName))  // Restrict to the cloth facade groups
+		{
+			CachedCollectionGroupNames.Emplace(GroupName);
+		}
+	}
+}
+
+void FChaosClothAssetAttributeNode_V2::OnDeselected()
+{
+	// Clean up, to avoid another toolkit picking up the wrong context evaluation
+	CachedCollectionGroupNames.Reset();
+}
+
 FChaosClothAssetAttributeNode::FChaosClothAssetAttributeNode(const Dataflow::FNodeParameters& InParam, FGuid InGuid)
 	: FDataflowNode(InParam, InGuid)
 {
diff --git a/Engine/Plugins/ChaosClothAssetEditor/Source/ChaosClothAssetDataflowNodes/Private/ChaosClothAsset/ChaosClothAssetDataflowNodesModule.cpp b/Engine/Plugins/ChaosClothAssetEditor/Source/ChaosClothAssetDataflowNodes/Private/ChaosClothAsset/ChaosClothAssetDataflowNodesModule.cpp
--- a/Engine/Plugins/ChaosClothAssetEditor/Source/ChaosClothAssetDataflowNodes/Private/ChaosClothAsset/ChaosClothAssetDataflowNodesModule.cpp
+++ b/Engine/Plugins/ChaosClothAssetEditor/Source/ChaosClothAssetDataflowNodes/Private/ChaosClothAsset/ChaosClothAssetDataflowNodesModule.cpp
@@ -79,6 +79,7 @@ namespace UE::Chaos::ClothAsset
 			DATAFLOW_NODE_REGISTER_CREATION_FACTORY(FChaosClothAssetAddWeightMapNode);
 			DATAFLOW_NODE_REGISTER_CREATION_FACTORY(FChaosClothAssetWeightMapNode);
 			DATAFLOW_NODE_REGISTER_CREATION_FACTORY(FChaosClothAssetAttributeNode);
+			DATAFLOW_NODE_REGISTER_CREATION_FACTORY(FChaosClothAssetAttributeNode_V2);
 			DATAFLOW_NODE_REGISTER_CREATION_FACTORY(FChaosClothAssetBindToRootBoneNode);
 			DATAFLOW_NODE_REGISTER_CREATION_FACTORY(FChaosClothAssetCopySimulationToRenderMeshNode);
 			DATAFLOW_NODE_REGISTER_CREATION_FACTORY(FChaosClothAssetDatasmithImportNode);
diff --git a/Engine/Plugins/ChaosClothAssetEditor/Source/ChaosClothAssetDataflowNodes/Public/ChaosClothAsset/AttributeNode.h b/Engine/Plugins/ChaosClothAssetEditor/Source/ChaosClothAssetDataflowNodes/Public/ChaosClothAsset/AttributeNode.h
--- a/Engine/Plugins/ChaosClothAssetEditor/Source/ChaosClothAssetDataflowNodes/Public/ChaosClothAsset/AttributeNode.h
+++ b/Engine/Plugins/ChaosClothAssetEditor/Source/ChaosClothAssetDataflowNodes/Public/ChaosClothAsset/AttributeNode.h
@@ -30,6 +30,57 @@ enum class EChaosClothAssetNodeAttributeType : uint8
 
 /** Create a new attribute for the specified group. */
 USTRUCT(Meta = (DataflowCloth, Experimental))
+struct FChaosClothAssetAttributeNode_V2 : public FDataflowNode
+{
+	GENERATED_USTRUCT_BODY()
+	DATAFLOW_NODE_DEFINE_INTERNAL(FChaosClothAssetAttributeNode_V2, "Attribute", "Cloth", "Cloth Attribute")
+
+public:
+
+	UPROPERTY(Meta = (Dataflowinput, DataflowOutput, DataflowPassthrough = "Collection"))
+	FManagedArrayCollection Collection;
+
+	/** The name of the attribute to create. */
+	UPROPERTY(EditAnywhere, Category = "Attribute")
+	FChaosClothAssetConnectableIOStringValue Name;
+
+	/** The attribute group. */
+	UPROPERTY(EditAnywhere, Category = "Attribute")
+	FChaosClothAssetNodeAttributeGroup Group;
+
+	/** The attribute type. */
+	UPROPERTY(EditAnywhere, Category = "Attribute")
+	EChaosClothAssetNodeAttributeType Type = EChaosClothAssetNodeAttributeType::Integer;
+
+	/** Default integer value. */
+	UPROPERTY(EditAnywhere, Category = "Attribute", Meta = (EditCondition = "Type == EChaosClothAssetNodeAttributeType::Integer", EditConditionHides))
+	int32 IntValue = 0;
+
+	/** Default float value. */
+	UPROPERTY(EditAnywhere, Category = "Attribute", Meta = (EditCondition = "Type == EChaosClothAssetNodeAttributeType::Float", EditConditionHides))
+	float FloatValue = 0.f;
+
+	/** Default vector value. */
+	UPROPERTY(EditAnywhere, Category = "Attribute", Meta = (EditCondition = "Type == EChaosClothAssetNodeAttributeType::Vector", EditConditionHides))
+	FVector3f VectorValue = FVector3f::ZeroVector;
+
+	FChaosClothAssetAttributeNode_V2(const Dataflow::FNodeParameters& InParam, FGuid InGuid = FGuid::NewGuid());
+
+	/** Return a cached array of all the groups used by the input collection during at the time of the latest evaluation. */
+	const TArray<FName>& GetCachedCollectionGroupNames() const { return CachedCollectionGroupNames; }
+
+private:
+	//~ Begin FDataflowNode interface
+	virtual void Evaluate(Dataflow::FContext& Context, const FDataflowOutput* Out) const override;
+	virtual void OnSelected(Dataflow::FContext& Context) override;
+	virtual void OnDeselected() override;
+	//~ End FDataflowNode interface
+
+	TArray<FName> CachedCollectionGroupNames;
+};
+
+/** Create a new attribute for the specified group. */
+USTRUCT(Meta = (DataflowCloth, Experimental, Deprecated = "5.5"))
 struct FChaosClothAssetAttributeNode : public FDataflowNode
 {
 	GENERATED_USTRUCT_BODY()
