Chaos Cloth Asset - Added Sim Collection input to the TransferSkinWeights node and improved the property UI.

#tests Tested with various transfer configurations and traced through the code to make sure it uses the right code paths.

diff --git a/Engine/Plugins/ChaosClothAssetEditor/Source/ChaosClothAssetDataflowNodes/Private/ChaosClothAsset/TransferSkinWeightsNode.cpp b/Engine/Plugins/ChaosClothAssetEditor/Source/ChaosClothAssetDataflowNodes/Private/ChaosClothAsset/TransferSkinWeightsNode.cpp
--- a/Engine/Plugins/ChaosClothAssetEditor/Source/ChaosClothAssetDataflowNodes/Private/ChaosClothAsset/TransferSkinWeightsNode.cpp
+++ b/Engine/Plugins/ChaosClothAssetEditor/Source/ChaosClothAssetDataflowNodes/Private/ChaosClothAsset/TransferSkinWeightsNode.cpp
@@ -118,7 +118,7 @@ namespace UE::Chaos::ClothAsset::Private
 
 	/** Convert the ClothCollection to DynamicMesh. */
 	static bool ClothToDynamicMesh(
-		const TSharedRef<FManagedArrayCollection>& ClothCollection,
+		const TSharedRef<const FManagedArrayCollection>& ClothCollection,
 		const FReferenceSkeleton& TargetRefSkeleton, // the reference skeleton to add to the dynamic mesh
 		const bool bIsSim, // if true, Mesh will contain the sim mesh, otherwise the render mesh
 		UE::Geometry::FDynamicMesh3& Mesh) // the resulting sim or render mesh
@@ -127,7 +127,7 @@ namespace UE::Chaos::ClothAsset::Private
 		using namespace UE::Chaos::ClothAsset;
 
 		// Check if ClothCollection is empty
-		FCollectionClothFacade ClothFacade(ClothCollection);
+		FCollectionClothConstFacade ClothFacade(ClothCollection);
 		const int32 NumVertices = bIsSim ? ClothFacade.GetNumSimVertices3D() : ClothFacade.GetNumRenderVertices();
 		const int32 NumFaces = bIsSim ? ClothFacade.GetNumSimFaces() :  ClothFacade.GetNumRenderFaces();
 
@@ -590,12 +590,20 @@ namespace UE::Chaos::ClothAsset::Private
 		return true;
 	}
 
+	static bool TransferClosestPointOnSurface(
+		const FReferenceSkeleton& TargetRefSkeleton,
+		const FDynamicMesh3& SkeletalDynamicMesh,
+		const TSharedRef<FManagedArrayCollection>& ClothCollection,
+		const TransferBoneWeightsSettings& TransferSettings,
+		const TSharedPtr<const FManagedArrayCollection>& SimClothCollection);
+
 	/** Transfer skin weights to sim and render cloth. */
 	static bool TransferInpaintWeights(
 		const FReferenceSkeleton& TargetRefSkeleton,
 		const UE::Geometry::FDynamicMesh3& SourceDynamicMesh,
 		const TSharedRef<FManagedArrayCollection>& ClothCollection,
-		const TransferBoneWeightsSettings& TransferSettings)
+		const TransferBoneWeightsSettings& TransferSettings,
+		const TSharedPtr<const FManagedArrayCollection>& SimClothCollection)
 	{
 		using namespace UE::Geometry;
 
@@ -618,30 +626,16 @@ namespace UE::Chaos::ClothAsset::Private
 			// Compute the bone weights for the render mesh by transferring weights from the sim mesh
 			//
 			if (TransferSettings.bTransferToRenderFromSim)
-			{	
-				// If we previously transferred to sim mesh we can just reuse that dynamic mesh
-				if (!TransferSettings.bTransferToSim)
-				{
-					// Otherwise get it from the collection
-					if (!ClothToDynamicMesh(ClothCollection, TargetRefSkeleton, true, WeldedSimMesh))
-					{
-						return false;
-					}
-				}
+			{
+				ensureMsgf(TransferSettings.bTransferToSim, TEXT("The UI shouldn't allow a sim mesh InpaintWeights transfer without the sim mesh set to transfer, check the EditConditions."));
 
-				if (WeldedSimMesh.VertexCount() > 0 && WeldedSimMesh.TriangleCount() > 0)
-				{
-					// Transfer skin weights to render cloth from the sim mesh
-					if (!TransferInpaintWeightsToRender(TargetRefSkeleton, WeldedSimMesh, ClothCollection, TransferSettings))
-					{
-						return false;
-					}
-				}
-				else
-				{
-					UE_LOG(LogChaosClothAssetDataflowNodes, Warning, TEXT("TransferSkinWeightsNode: Could not transfer from simulation mesh because its empty."));
-					return false;
-				}
+				// Transfers from the sim mesh most likely mean that a SkeletalMesh InpaintWeights transfer is not working for this render mesh,
+				// therefore it is best to assume that a closest point transfer will always provide a better result in this particular case.
+				TransferBoneWeightsSettings NewTransferSettings = TransferSettings;
+				NewTransferSettings.bTransferToSim = false;
+				NewTransferSettings.TransferMethod = FTransferBoneWeights::ETransferBoneWeightsMethod::ClosestPointOnSurface;
+				
+				return TransferClosestPointOnSurface(TargetRefSkeleton, SourceDynamicMesh, ClothCollection, NewTransferSettings, SimClothCollection);
 			}
 			else 
 			{
@@ -660,10 +654,11 @@ namespace UE::Chaos::ClothAsset::Private
 		const FReferenceSkeleton& TargetRefSkeleton,
 		const FDynamicMesh3& SkeletalDynamicMesh,
 		const TSharedRef<FManagedArrayCollection>& ClothCollection,
-		const UE::Chaos::ClothAsset::Private::TransferBoneWeightsSettings& TransferSettings)
+		const TransferBoneWeightsSettings& TransferSettings,
+		const TSharedPtr<const FManagedArrayCollection>& SimClothCollection)
 	{
 		using namespace UE::Geometry;
-		
+
 		FCollectionClothFacade ClothFacade(ClothCollection);
 
 		//
@@ -713,8 +708,9 @@ namespace UE::Chaos::ClothAsset::Private
 			FDynamicMesh3 WeldedSimMesh;
 			if (TransferSettings.bTransferToRenderFromSim)
 			{
-				// Convert sim cloth to dynamic mesh
-				if (!ClothToDynamicMesh(ClothCollection, TargetRefSkeleton, true, WeldedSimMesh))
+				// Convert sim cloth to dynamic mesh, use a different sim cloth collection if needed
+				constexpr bool bIsSimMesh = true;
+				if (!ClothToDynamicMesh(SimClothCollection.IsValid() ? SimClothCollection.ToSharedRef() : ClothCollection, TargetRefSkeleton, bIsSimMesh, WeldedSimMesh))
 				{
 					return false;
 				}
@@ -752,10 +748,14 @@ FChaosClothAssetTransferSkinWeightsNode::FChaosClothAssetTransferSkinWeightsNode
 	: FDataflowNode(InParam, InGuid)
 {
 	RegisterInputConnection(&Collection);
-	RegisterOutputConnection(&Collection, &Collection);
+	RegisterInputConnection(&SkeletalMesh);
+	RegisterInputConnection(&SimCollection)
+		.SetCanHidePin(true)
+		.SetPinIsHidden(true);
 	RegisterInputConnection(&InpaintMask.WeightMap, GET_MEMBER_NAME_CHECKED(FChaosClothAssetWeightedValueNonAnimatableNoLowHighRange, WeightMap))
 		.SetCanHidePin(true)
 		.SetPinIsHidden(true);
+	RegisterOutputConnection(&Collection, &Collection);
 }
 
 void FChaosClothAssetTransferSkinWeightsNode::Evaluate(Dataflow::FContext& Context, const FDataflowOutput* Out) const
@@ -771,37 +771,115 @@ void FChaosClothAssetTransferSkinWeightsNode::Evaluate(Dataflow::FContext& Conte
 
 		// Evaluate inputs
 		FManagedArrayCollection InputCollection = GetValue<FManagedArrayCollection>(Context, &Collection);
+
 		const TSharedRef<FManagedArrayCollection> ClothCollection = MakeShared<FManagedArrayCollection>(MoveTemp(InputCollection));
+		
+		const TSharedPtr<FManagedArrayCollection> SimClothCollection = IsConnected(&SimCollection) ?
+			MakeShared<FManagedArrayCollection>(GetValue(Context, &SimCollection)) :
+			TSharedPtr<FManagedArrayCollection>();
 
-		if (SkeletalMesh && FCollectionClothFacade(ClothCollection).IsValid())  // Can only act on the collection if it is a valid cloth collection
+		FCollectionClothFacade ClothFacade(ClothCollection);
+		if (ClothFacade.IsValid())  // Can only act on the collection if it is a valid cloth collection
 		{
-			if (!SkeletalMesh->IsValidLODIndex(LodIndex))
+			TStrongObjectPtr<const USkeletalMesh> InSkeletalMesh;
+			const FReferenceSkeleton* TargetRefSkeleton = nullptr;
+			FDynamicMesh3 SourceDynamicMesh;
+
+			const bool bNeedsSkeletalMesh =
+				TargetMeshType != EChaosClothAssetTransferTargetMeshType::Render ||
+				RenderMeshSourceType == EChaosClothAssetTransferRenderMeshSource::SkeletalMesh;
+
+			if (bNeedsSkeletalMesh)
 			{
-				FClothDataflowTools::LogAndToastWarning(*this,
-					LOCTEXT("InvalidLodIndexHeadline", "Invalid LOD Index."),
-					FText::Format(
-						LOCTEXT("InvalidLodIndexDetails", "LOD index {0} is not a valid LOD for skeletal mesh {1}."),
-						LodIndex,
-						FText::FromString(SkeletalMesh.GetName())));
-				SetValue(Context, MoveTemp(*ClothCollection), &Collection);
-				return;
-			}
+				InSkeletalMesh = TStrongObjectPtr<const USkeletalMesh>(GetValue(Context, &SkeletalMesh));
 
-			FDynamicMesh3 SourceDynamicMesh;
-			if (!SkeletalMeshToDynamicMesh(SkeletalMesh, LodIndex, SourceDynamicMesh))
+				if (!InSkeletalMesh)
+				{
+					FClothDataflowTools::LogAndToastWarning(*this,
+						LOCTEXT("InvalidSkeletalMeshHeadline", "Invalid Skeletal Mesh."),
+						FText::Format(
+							LOCTEXT("InvalidSkeletalMeshDetails", "No skeletal mesh has been specified and one is required for this type of transfer."),
+							LodIndex));
+					SetValue(Context, MoveTemp(*ClothCollection), &Collection);
+					return;
+				}
+
+				if (!InSkeletalMesh->IsValidLODIndex(LodIndex))
+				{
+					FClothDataflowTools::LogAndToastWarning(*this,
+						LOCTEXT("InvalidLodIndexHeadline", "Invalid LOD Index."),
+						FText::Format(
+							LOCTEXT("InvalidLodIndexDetails", "LOD index {0} is not a valid LOD for skeletal mesh {1}."),
+							LodIndex,
+							FText::FromString(InSkeletalMesh->GetName())));
+					SetValue(Context, MoveTemp(*ClothCollection), &Collection);
+					return;
+				}
+
+				if (!SkeletalMeshToDynamicMesh(InSkeletalMesh.Get(), LodIndex, SourceDynamicMesh))
+				{
+					FClothDataflowTools::LogAndToastWarning(*this,
+						LOCTEXT("InvalidLodHeadline", "Could not convert LOD to Dynamic Mesh."),
+						FText::Format(
+							LOCTEXT("InvalidLodDetails", "Could not convert LOD index {0} of the skeletal mesh {1} to dyanamic mesh."),
+							LodIndex,
+							FText::FromString(InSkeletalMesh->GetName())));
+					SetValue(Context, MoveTemp(*ClothCollection), &Collection);
+					return;
+				}
+
+				MeshTransforms::ApplyTransform(SourceDynamicMesh, Transform, true);
+				TargetRefSkeleton = &InSkeletalMesh->GetRefSkeleton();
+
+				ClothFacade.SetSkeletalMeshPathName(InSkeletalMesh->GetPathName());
+			}
+			else
 			{
-				FClothDataflowTools::LogAndToastWarning(*this,
-					LOCTEXT("InvalidLodHeadline", "Could not convert LOD to Dynamic Mesh."),
-					FText::Format(
-						LOCTEXT("InvalidLodDetails", "Could not convert LOD index {0} of the skeletal mesh {1} to dyanamic mesh."),
-						LodIndex,
-						FText::FromString(SkeletalMesh.GetName())));
-				SetValue(Context, MoveTemp(*ClothCollection), &Collection);
-				return;
+				// Reuse the input sim mesh skeleton
+				auto GetCollectionRefSkeleton = [&InSkeletalMesh](const FCollectionClothConstFacade& ClothFacade) -> const FReferenceSkeleton*
+					{
+						const FString& SkeletalMeshPathName = ClothFacade.GetSkeletalMeshPathName();
+						InSkeletalMesh = TStrongObjectPtr<const USkeletalMesh>(LoadObject<USkeletalMesh>(nullptr, *SkeletalMeshPathName));
+						return InSkeletalMesh ? &InSkeletalMesh->GetRefSkeleton() : nullptr;
+					};
+
+				if (SimClothCollection.IsValid())
+				{
+					FCollectionClothConstFacade SimClothFacade(SimClothCollection.ToSharedRef());
+					if (SimClothFacade.IsValid())
+					{
+						TargetRefSkeleton = GetCollectionRefSkeleton(SimClothFacade);
+					}
+					if (!TargetRefSkeleton)
+					{
+						FClothDataflowTools::LogAndToastWarning(*this,
+							LOCTEXT("InvalidSimRefSkeletonHeadline", "Invalid Reference Skeleton."),
+							LOCTEXT("InvalidSimRefSkeletonDetails", "Couldn't find a valid reference skeleton from the input sim collection."));
+						SetValue(Context, MoveTemp(*ClothCollection), &Collection);
+						return;
+					}
+				}
+				else
+				{
+					TargetRefSkeleton = GetCollectionRefSkeleton(ClothFacade);
+
+					if (!TargetRefSkeleton)
+					{
+						FClothDataflowTools::LogAndToastWarning(*this,
+							LOCTEXT("InvalidRefSkeletonHeadline", "Invalid Reference Skeleton."),
+							LOCTEXT("InvalidRefSkeletonDetails", "Couldn't find a valid reference skeleton from the input collection."));
+						SetValue(Context, MoveTemp(*ClothCollection), &Collection);
+						return;
+					}
+				}
 			}
 
-			MeshTransforms::ApplyTransform(SourceDynamicMesh, Transform, true);
-			const FReferenceSkeleton& TargetRefSkeleton = SkeletalMesh->GetRefSkeleton();
+			// Clean up orphaned vertices
+			UE::Chaos::ClothAsset::FClothGeometryTools::CleanupAndCompactMesh(ClothCollection);
+			if (SimClothCollection.IsValid())
+			{
+				UE::Chaos::ClothAsset::FClothGeometryTools::CleanupAndCompactMesh(SimClothCollection.ToSharedRef());
+			}
 
 			//
 			// Setup the bone weight transfer settings.
@@ -823,52 +901,20 @@ void FChaosClothAssetTransferSkinWeightsNode::Evaluate(Dataflow::FContext& Conte
 			TransferSettings.SmoothingStrength = SmoothingStrength;
 			TransferSettings.InpaintMaskWeightMapName = GetValue<FString>(Context, &InpaintMask.WeightMap);
 			
-			//
-			// Transfer the bone weights from the source Skeletal mesh to the Cloth asset.
-			//
-			FCollectionClothFacade ClothFacade(ClothCollection);
-			ClothFacade.SetSkeletalMeshPathName(SkeletalMesh->GetPathName());
-
-			// Clean up orphaned vertices
-			UE::Chaos::ClothAsset::FClothGeometryTools::CleanupAndCompactMesh(ClothCollection);
-
 			bool bTransferResult = false;
-			if (TransferMethod == EChaosClothAssetTransferSkinWeightsMethod::InpaintWeights && 
-				TransferSettings.bTransferToRenderFromSim && 
-				TransferSettings.bTransferToSim && 
-				TransferSettings.bTransferToRender)
-			{
-				// Custom setup for the default behavior of the node that gives the best results in the common case where:
-				//  - the sim mesh is welded and manifold
-				//  - the render mesh is un-welded and highly non-manifold with many disconnected regions in areas like 
-				//    the armpits
-				//  - the render and the sim mesh have similar shapes
-				//  - we are transferring weight from the sim mesh and not the body
-
-				// First transfer to sim only using inpaint weights algorithm
-				TransferSettings.bTransferToRender = false;
-				bTransferResult = TransferInpaintWeights(TargetRefSkeleton, SourceDynamicMesh, ClothCollection, TransferSettings);
-
-				if (bTransferResult)
-				{
-					// Now transfer to render only using closest point 
-					TransferSettings.bTransferToSim = false;
-					TransferSettings.bTransferToRender = true;
-					TransferSettings.TransferMethod = FTransferBoneWeights::ETransferBoneWeightsMethod::ClosestPointOnSurface;
-					bTransferResult = TransferClosestPointOnSurface(TargetRefSkeleton, SourceDynamicMesh, ClothCollection, TransferSettings);
-				}
-			}
-			else if (TransferMethod == EChaosClothAssetTransferSkinWeightsMethod::InpaintWeights)
-			{
-				bTransferResult = TransferInpaintWeights(TargetRefSkeleton, SourceDynamicMesh, ClothCollection, TransferSettings);
-			}
-			else if (TransferMethod == EChaosClothAssetTransferSkinWeightsMethod::ClosestPointOnSurface)
-			{
-				bTransferResult = TransferClosestPointOnSurface(TargetRefSkeleton, SourceDynamicMesh, ClothCollection, TransferSettings);
-			}
-			else
+			switch (TransferMethod)
 			{
+			case EChaosClothAssetTransferSkinWeightsMethod::InpaintWeights:
+				bTransferResult = TransferInpaintWeights(*TargetRefSkeleton, SourceDynamicMesh, ClothCollection, TransferSettings, SimClothCollection);
+				break;
+			
+			case EChaosClothAssetTransferSkinWeightsMethod::ClosestPointOnSurface:
+				bTransferResult = TransferClosestPointOnSurface(*TargetRefSkeleton, SourceDynamicMesh, ClothCollection, TransferSettings, SimClothCollection);
+				break;
+
+			default:
 				checkNoEntry();
+				break;
 			}
 
 			if (!bTransferResult)
diff --git a/Engine/Plugins/ChaosClothAssetEditor/Source/ChaosClothAssetDataflowNodes/Public/ChaosClothAsset/TransferSkinWeightsNode.h b/Engine/Plugins/ChaosClothAssetEditor/Source/ChaosClothAssetDataflowNodes/Public/ChaosClothAsset/TransferSkinWeightsNode.h
--- a/Engine/Plugins/ChaosClothAssetEditor/Source/ChaosClothAssetDataflowNodes/Public/ChaosClothAsset/TransferSkinWeightsNode.h
+++ b/Engine/Plugins/ChaosClothAssetEditor/Source/ChaosClothAssetDataflowNodes/Public/ChaosClothAsset/TransferSkinWeightsNode.h
@@ -38,23 +38,23 @@ UENUM(BlueprintType)
 enum class EChaosClothAssetTransferTargetMeshType : uint8
 {
 	/** Perform the skin weights transfer for both the simulation and render meshes. */
-	All,
+	All UMETA(DisplayName = "Sim & Render Meshes"),
 
 	/** Perform the skin weights transfer for the simulation mesh only. */
-	Simulation,
+	Simulation UMETA(DisplayName = "Sim Mesh"),
 	
 	/** Perform the skin weights transfer for the render mesh only. */
-	Render
+	Render UMETA(DisplayName = "Render Mesh")
 };
 
 UENUM(BlueprintType)
 enum class EChaosClothAssetTransferRenderMeshSource : uint8
 {
-	/** For render mesh, transfer weights from the source skeletal mesh. */
-	SkeletalMesh,
+	/** For render mesh, transfer weights from the source Skeletal Mesh. */
+	SkeletalMesh UMETA(DisplayName = "Skeletal Mesh"),
 
-	/** For render mesh, transfer weights from the simulation mesh. */
-	SimulationMesh
+	/** For render mesh, transfer weights from the Collection input sim mesh, or Sim Collection input if connected. */
+	SimulationMesh UMETA(DisplayName = "Collection/Sim Collection")
 };
 
 /** Transfer the skinning weights set on a skeletal mesh to the simulation and/or render mesh stored in the cloth collection. */
@@ -68,28 +68,48 @@ public:
 	UPROPERTY(Meta = (Dataflowinput, DataflowOutput, DataflowPassthrough = "Collection"))
 	FManagedArrayCollection Collection;
 
+	/** The type of cloth mesh the skeletal mesh transfer will be applied to, simulation, render mesh, or both. */
+	UPROPERTY(EditAnywhere, Category = "Transfer Skin Weights", Meta = (DisplayName = "Target Mesh(es)"))
+	EChaosClothAssetTransferTargetMeshType TargetMeshType = EChaosClothAssetTransferTargetMeshType::All;
+
+private:
+	/** For the sim mesh, simulation mesh transfers always use the specified skeletal mesh. */
+	UPROPERTY(VisibleAnywhere, Category = "Transfer Skin Weights", Meta = (DisplayName = "Sim Mesh Transfer Source", EditCondition = "TargetMeshType != EChaosClothAssetTransferTargetMeshType::Render", EditConditionHides))
+	FString SimMeshSourceTypeHint = TEXT("Skeletal Mesh");
+
+public:
+	/** For the render mesh, choose which source to use, either the default or specified simulation mesh or the specified skeletal mesh. */
+	UPROPERTY(EditAnywhere, Category = "Transfer Skin Weights", Meta = (DisplayName = "Render Mesh Transfer Source", EditCondition = "TargetMeshType != EChaosClothAssetTransferTargetMeshType::Simulation", EditConditionHides))
+	EChaosClothAssetTransferRenderMeshSource RenderMeshSourceType = EChaosClothAssetTransferRenderMeshSource::SimulationMesh;
+
 	/** The skeletal mesh to transfer the skin weights from. */
-	UPROPERTY(EditAnywhere, Category = "Transfer Skin Weights|Source Mesh")
+	UPROPERTY(EditAnywhere, Category = "Transfer Skin Weights", Meta = (DataflowInput, EditCondition = "TargetMeshType != EChaosClothAssetTransferTargetMeshType::Render || RenderMeshSourceType == EChaosClothAssetTransferRenderMeshSource::SkeletalMesh", EditConditionHides))
 	TObjectPtr<USkeletalMesh> SkeletalMesh;
 
+	/** The collection containing the sim mesh to use when the Render Mesh Transfer Source is set to Collection/Sim Collection. When this input isn't connected, the Collection input is used instead. */
+	UPROPERTY(Meta = (DataflowPassthrough = "Collection", Dataflowinput))
+	FManagedArrayCollection SimCollection;
+
 	/** The skeletal mesh LOD to transfer the skin weights from. */
-	UPROPERTY(EditAnywhere, Category = "Transfer Skin Weights|Source Mesh", Meta = (DisplayName = "LOD Index"))
+	UPROPERTY(EditAnywhere, Category = "Transfer Skin Weights", Meta = (DisplayName = "LOD Index", EditCondition = "TargetMeshType != EChaosClothAssetTransferTargetMeshType::Render || RenderMeshSourceType == EChaosClothAssetTransferRenderMeshSource::SkeletalMesh", EditConditionHides))
 	int32 LodIndex = 0;
 
 	/** The relative transform between the skeletal mesh and the cloth asset. */
-	UPROPERTY(EditAnywhere, Category = "Transfer Skin Weights|Source Mesh")
+	UPROPERTY(EditAnywhere, Category = "Transfer Skin Weights", Meta = (EditCondition = "TargetMeshType != EChaosClothAssetTransferTargetMeshType::Render || RenderMeshSourceType == EChaosClothAssetTransferRenderMeshSource::SkeletalMesh", EditConditionHides))
 	FTransform Transform;
 
-	/** The type of the mesh the transfer will be applied to. */
-	UPROPERTY(EditAnywhere, Category = "Transfer Skin Weights|Transfer Method", Meta = (DisplayName = "Target Mesh Type"))
-	EChaosClothAssetTransferTargetMeshType TargetMeshType = EChaosClothAssetTransferTargetMeshType::All;
-
-	/** For the render mesh, choose which source to use. */
-	UPROPERTY(EditAnywhere, Category = "Transfer Skin Weights|Transfer Method", Meta = (DisplayName = "Render Mesh Source Type", EditCondition="TargetMeshType!=EChaosClothAssetTransferTargetMeshType::Simulation"))
-	EChaosClothAssetTransferRenderMeshSource RenderMeshSourceType = EChaosClothAssetTransferRenderMeshSource::SimulationMesh;
+private:
+	/** Algorithm used for the transfer method. When the Render Mesh Transfer Source is set to use the sim mesh from the Collection/Sim Collection input, only the ClosestPointOnSurface method is available. */
+	UPROPERTY(VisibleAnywhere, Category = "Transfer Skin Weights", Meta = (DisplayName = "Transfer Method", EditCondition = "TargetMeshType == EChaosClothAssetTransferTargetMeshType::Render && RenderMeshSourceType != EChaosClothAssetTransferRenderMeshSource::SkeletalMesh", EditConditionHides))
+	FString TransferMethodHint = TEXT("Closest Point On Surface");
 
-	/** Algorithm used for the transfer method. Use the simple ClosestPointOnSurface method or the more complex InpaintWeights method for better results. */
-	UPROPERTY(EditAnywhere, Category = "Transfer Skin Weights|Transfer Method", Meta = (DisplayName = "Algorithm"))
+public:
+	/**
+	 * Algorithm used for the transfer method.
+	 * Use the simple ClosestPointOnSurface method or the more complex InpaintWeights method for better results.
+	 * Note: When using the simulation mesh as source for the render mesh transfer, the algorithm will always be the ClosestPointOnSurface method, whatever this setting is.
+	 */
+	UPROPERTY(EditAnywhere, Category = "Transfer Skin Weights", Meta = (EditCondition = "TargetMeshType != EChaosClothAssetTransferTargetMeshType::Render || RenderMeshSourceType == EChaosClothAssetTransferRenderMeshSource::SkeletalMesh", EditConditionHides))
 	EChaosClothAssetTransferSkinWeightsMethod TransferMethod = EChaosClothAssetTransferSkinWeightsMethod::InpaintWeights;
 
 	/**
@@ -97,14 +117,14 @@ public:
 	 * All points outside of the search radius will be ignored. 
 	 * When set to a negative value (e.g. -1), all points will be considered.
 	 */
-	UPROPERTY(EditAnywhere, Category = "Transfer Skin Weights|Transfer Method", Meta = (UIMin = -1, UIMax = 2, ClampMin = -1, ClampMax = 2, EditCondition="TransferMethod==EChaosClothAssetTransferSkinWeightsMethod::InpaintWeights"))
+	UPROPERTY(EditAnywhere, Category = "Transfer Skin Weights", Meta = (UIMin = -1, UIMax = 2, ClampMin = -1, ClampMax = 2, EditCondition = "(TargetMeshType != EChaosClothAssetTransferTargetMeshType::Render || RenderMeshSourceType == EChaosClothAssetTransferRenderMeshSource::SkeletalMesh) && TransferMethod == EChaosClothAssetTransferSkinWeightsMethod::InpaintWeights", EditConditionHides))
 	double RadiusPercentage = 0.05;
 
 	/**
 	 * Maximum angle difference (in degrees) between the target and source point normals to be considered a match for the InpaintWeights method.
 	 * If set to a negative value (e.g. -1), normals will be ignored.
 	 */
-	UPROPERTY(EditAnywhere, Category = "Transfer Skin Weights|Transfer Method", Meta = (UIMin = -1, UIMax = 180, ClampMin = -1, ClampMax = 180, EditCondition="TransferMethod==EChaosClothAssetTransferSkinWeightsMethod::InpaintWeights"))
+	UPROPERTY(EditAnywhere, Category = "Transfer Skin Weights", Meta = (UIMin = -1, UIMax = 180, ClampMin = -1, ClampMax = 180, EditCondition = "(TargetMeshType != EChaosClothAssetTransferTargetMeshType::Render || RenderMeshSourceType == EChaosClothAssetTransferRenderMeshSource::SkeletalMesh) && TransferMethod == EChaosClothAssetTransferSkinWeightsMethod::InpaintWeights", EditConditionHides))
 	double NormalThreshold = 30;
 
 	/** 
@@ -112,25 +132,25 @@ public:
 	 * This helps with layered meshes where the "inner" and "outer" layers are close to each other but whose normals 
 	 * are pointing in the opposite directions.
 	 */
-	UPROPERTY(EditAnywhere, Category = "Transfer Skin Weights|Transfer Method", Meta = (DisplayName = "Layered Mesh Support", EditCondition="TransferMethod==EChaosClothAssetTransferSkinWeightsMethod::InpaintWeights"))
+	UPROPERTY(EditAnywhere, Category = "Transfer Skin Weights", Meta = (EditCondition = "(TargetMeshType != EChaosClothAssetTransferTargetMeshType::Render || RenderMeshSourceType == EChaosClothAssetTransferRenderMeshSource::SkeletalMesh) && TransferMethod == EChaosClothAssetTransferSkinWeightsMethod::InpaintWeights", EditConditionHides))
 	bool LayeredMeshSupport = true;
 
 	/** The number of smoothing iterations applied to the vertices whose weights were automatically computed. */
-	UPROPERTY(EditAnywhere, Category = "Transfer Skin Weights|Transfer Method", Meta = (UIMin = 0, UIMax = 100, ClampMin = 0, ClampMax = 100, DisplayName = "Smoothing Iterations", EditCondition="TransferMethod==EChaosClothAssetTransferSkinWeightsMethod::InpaintWeights"))
+	UPROPERTY(EditAnywhere, Category = "Transfer Skin Weights", Meta = (UIMin = 0, UIMax = 100, ClampMin = 0, ClampMax = 100, DisplayName = "Smoothing Iterations", EditCondition = "(TargetMeshType != EChaosClothAssetTransferTargetMeshType::Render || RenderMeshSourceType == EChaosClothAssetTransferRenderMeshSource::SkeletalMesh) && TransferMethod == EChaosClothAssetTransferSkinWeightsMethod::InpaintWeights", EditConditionHides))
 	int32 NumSmoothingIterations = 10;
 
 	/** The smoothing strength of each smoothing iteration. */
-	UPROPERTY(EditAnywhere, Category = "Transfer Skin Weights|Transfer Method", Meta = (UIMin = 0, UIMax = 1, ClampMin = 0, ClampMax = 1, DisplayName = "Smoothing Strength", EditCondition="TransferMethod==EChaosClothAssetTransferSkinWeightsMethod::InpaintWeights"))
+	UPROPERTY(EditAnywhere, Category = "Transfer Skin Weights", Meta = (UIMin = 0, UIMax = 1, ClampMin = 0, ClampMax = 1, EditCondition = "(TargetMeshType != EChaosClothAssetTransferTargetMeshType::Render || RenderMeshSourceType == EChaosClothAssetTransferRenderMeshSource::SkeletalMesh) && TransferMethod == EChaosClothAssetTransferSkinWeightsMethod::InpaintWeights", EditConditionHides))
 	float SmoothingStrength = 0.1;
 	
+	/** Optional mask where a non-zero value indicates that we want the skinning weights for the vertex to be computed automatically instead of it being copied over from the source mesh. */
+	UPROPERTY(EditAnywhere, Category = "Transfer Skin Weights", Meta = (EditCondition = "(TargetMeshType != EChaosClothAssetTransferTargetMeshType::Render || RenderMeshSourceType == EChaosClothAssetTransferRenderMeshSource::SkeletalMesh) && TransferMethod == EChaosClothAssetTransferSkinWeightsMethod::InpaintWeights", EditConditionHides))
+	FChaosClothAssetWeightedValueNonAnimatableNoLowHighRange InpaintMask = { TEXT("InpaintMask") };
+
 	/** The maximum number of bones that will influence each vertex. */
-	UPROPERTY(EditAnywhere, Category = "Transfer Skin Weights|Transfer Method", Meta = (DisplayName = "Max Bone Influences"))
+	UPROPERTY(EditAnywhere, Category = "Transfer Skin Weights", Meta = (DisplayName = "Max Bone Influences"))
 	EChaosClothAssetMaxNumInfluences MaxNumInfluences = EChaosClothAssetMaxNumInfluences::Eight;
 	
-    /** Optional mask where a non-zero value indicates that we want the skinning weights for the vertex to be computed automatically instead of it being copied over from the source mesh. */
-	UPROPERTY(EditAnywhere, Category = "Transfer Skin Weights|Transfer Method", Meta = (DisplayName = "Inpaint Mask", EditCondition="TransferMethod==EChaosClothAssetTransferSkinWeightsMethod::InpaintWeights"))
-	FChaosClothAssetWeightedValueNonAnimatableNoLowHighRange InpaintMask = { TEXT("InpaintMask") };
-
 	FChaosClothAssetTransferSkinWeightsNode(const Dataflow::FNodeParameters& InParam, FGuid InGuid = FGuid::NewGuid());
 
 private:
