Cloth weight paint tool: draw 0.0 weight values as blue, 1.0 weight values as yellow. Make this behavior toggle-able with a checkbox in the UI.

diff --git a/Engine/Plugins/ChaosClothAssetEditor/Source/ChaosClothAssetEditorTools/Private/ChaosClothAsset/ClothWeightMapPaintTool.cpp b/Engine/Plugins/ChaosClothAssetEditor/Source/ChaosClothAssetEditorTools/Private/ChaosClothAsset/ClothWeightMapPaintTool.cpp
--- a/Engine/Plugins/ChaosClothAssetEditor/Source/ChaosClothAssetEditorTools/Private/ChaosClothAsset/ClothWeightMapPaintTool.cpp
+++ b/Engine/Plugins/ChaosClothAssetEditor/Source/ChaosClothAssetEditorTools/Private/ChaosClothAsset/ClothWeightMapPaintTool.cpp
@@ -212,6 +212,12 @@ void UClothEditorWeightMapPaintTool::Setup()
 			UpdateVertexColorOverlay(); 
 			DynamicMeshComponent->FastNotifyVertexAttributesUpdated(EMeshRenderAttributeFlags::VertexColors);
 		});
+	FilterProperties->WatchProperty(FilterProperties->bHighlightZeroAndOne,
+		[this](bool bNewValue)
+		{
+			UpdateVertexColorOverlay();
+			DynamicMeshComponent->FastNotifyVertexAttributesUpdated(EMeshRenderAttributeFlags::VertexColors);
+		});
 	FilterProperties->WatchProperty(FilterProperties->SubToolType,
 		[this](EClothEditorWeightMapPaintInteractionType NewType) { UpdateSubToolType(NewType); });
 	FilterProperties->WatchProperty(FilterProperties->BrushSize, [this](float NewSize) 
@@ -2315,7 +2321,15 @@ void UClothEditorWeightMapPaintTool::UpdateVertexColorOverlay(const TSet<int>* T
 			VertexWeight = FMath::Clamp(VertexWeight, 0.0f, 1.0f);
 
 			FVector4f NewColor;
-			if (FilterProperties->ColorMap == EClothEditorWeightMapDisplayType::BlackAndWhite)
+			if (FilterProperties->bHighlightZeroAndOne && VertexWeight == 0.0f)
+			{
+				NewColor = FVector4f(0.0f, 0.0f, 1.0f, 1.0f);
+			} 
+			else if (FilterProperties->bHighlightZeroAndOne && VertexWeight == 1.0f)
+			{
+				NewColor = FVector4f(1.0f, 1.0f, 0.0f, 1.0f);
+			}
+			else if (FilterProperties->ColorMap == EClothEditorWeightMapDisplayType::BlackAndWhite)
 			{
 				NewColor = FVector4f(VertexWeight, VertexWeight, VertexWeight, 1.0f);
 			}
diff --git a/Engine/Plugins/ChaosClothAssetEditor/Source/ChaosClothAssetEditorTools/Private/ChaosClothAsset/ClothWeightMapPaintTool.h b/Engine/Plugins/ChaosClothAssetEditor/Source/ChaosClothAssetEditorTools/Private/ChaosClothAsset/ClothWeightMapPaintTool.h
--- a/Engine/Plugins/ChaosClothAssetEditor/Source/ChaosClothAssetEditorTools/Private/ChaosClothAsset/ClothWeightMapPaintTool.h
+++ b/Engine/Plugins/ChaosClothAssetEditor/Source/ChaosClothAssetEditorTools/Private/ChaosClothAsset/ClothWeightMapPaintTool.h
@@ -120,6 +120,9 @@ public:
 	UPROPERTY(EditAnywhere, Category = Display)
 	EClothEditorWeightMapDisplayType ColorMap = EClothEditorWeightMapDisplayType::BlackAndWhite;
 
+	UPROPERTY(EditAnywhere, Category = Display)
+	bool bHighlightZeroAndOne = false;
+
 	UPROPERTY(EditAnywhere, Category = ActionType, meta = (DisplayName = "Action"))
 	EClothEditorWeightMapPaintInteractionType SubToolType = EClothEditorWeightMapPaintInteractionType::Brush;
 
