Cloth Editor: add ability to visualize surface normals in the construction viewport

diff --git a/Engine/Plugins/ChaosClothAssetEditor/Source/ChaosClothAssetEditor/Private/ChaosClothAsset/ClothEditorCommands.cpp b/Engine/Plugins/ChaosClothAssetEditor/Source/ChaosClothAssetEditor/Private/ChaosClothAsset/ClothEditorCommands.cpp
--- a/Engine/Plugins/ChaosClothAssetEditor/Source/ChaosClothAssetEditor/Private/ChaosClothAsset/ClothEditorCommands.cpp
+++ b/Engine/Plugins/ChaosClothAssetEditor/Source/ChaosClothAssetEditor/Private/ChaosClothAsset/ClothEditorCommands.cpp
@@ -24,6 +24,7 @@ const FString FChaosClothAssetEditorCommands::TogglePreviewWireframeIdentifier =
 const FString FChaosClothAssetEditorCommands::ToggleConstructionViewWireframeIdentifier = TEXT("ToggleConstructionViewWireframe");
 const FString FChaosClothAssetEditorCommands::ToggleConstructionViewSeamsIdentifier = TEXT("ToggleConstructionViewSeams");
 const FString FChaosClothAssetEditorCommands::ToggleConstructionViewSeamsCollapseIdentifier = TEXT("ToggleConstructionViewSeamsCollapse");
+const FString FChaosClothAssetEditorCommands::ToggleConstructionViewSurfaceNormalsIdentifier = TEXT("ToggleConstructionViewSurfaceNormals");
 
 FChaosClothAssetEditorCommands::FChaosClothAssetEditorCommands()
 	: TBaseCharacterFXEditorCommands<FChaosClothAssetEditorCommands>("ChaosClothAssetEditor",
@@ -59,6 +60,7 @@ void FChaosClothAssetEditorCommands::RegisterCommands()
 
 	UI_COMMAND(ToggleConstructionViewSeams, "ToggleConstructionViewSeams", "Toggle construction view seams", EUserInterfaceActionType::ToggleButton, FInputChord());
 	UI_COMMAND(ToggleConstructionViewSeamsCollapse, "ToggleConstructionViewSeamsCollapse", "Toggle drawing one line between seams in construction view", EUserInterfaceActionType::ToggleButton, FInputChord());
+	UI_COMMAND(ToggleConstructionViewSurfaceNormals, "ToggleConstructionViewSurfaceNormals", "Toggle construction view surface normals", EUserInterfaceActionType::ToggleButton, FInputChord());
 
 	UI_COMMAND(TogglePatternColor, "TogglePatternColor", "Toggle Pattern color in the construction view", EUserInterfaceActionType::ToggleButton, FInputChord());
 	UI_COMMAND(ToggleMeshStats, "ToggleMeshStats", "Toggle mesh information display", EUserInterfaceActionType::ToggleButton, FInputChord());
diff --git a/Engine/Plugins/ChaosClothAssetEditor/Source/ChaosClothAssetEditor/Private/ChaosClothAsset/ClothEditorCommands.h b/Engine/Plugins/ChaosClothAssetEditor/Source/ChaosClothAssetEditor/Private/ChaosClothAsset/ClothEditorCommands.h
--- a/Engine/Plugins/ChaosClothAssetEditor/Source/ChaosClothAssetEditor/Private/ChaosClothAsset/ClothEditorCommands.h
+++ b/Engine/Plugins/ChaosClothAssetEditor/Source/ChaosClothAssetEditor/Private/ChaosClothAsset/ClothEditorCommands.h
@@ -64,6 +64,9 @@ public:
 	TSharedPtr<FUICommandInfo> TogglePatternColor;
 	TSharedPtr<FUICommandInfo> ToggleMeshStats;
 
+	const static FString ToggleConstructionViewSurfaceNormalsIdentifier;
+	TSharedPtr<FUICommandInfo> ToggleConstructionViewSurfaceNormals;
+
 	TSharedPtr<FUICommandInfo> SetConstructionMode2D;
 	TSharedPtr<FUICommandInfo> SetConstructionMode3D;
 	TSharedPtr<FUICommandInfo> SetConstructionModeRender;
diff --git a/Engine/Plugins/ChaosClothAssetEditor/Source/ChaosClothAssetEditor/Private/ChaosClothAsset/ClothEditorMode.cpp b/Engine/Plugins/ChaosClothAssetEditor/Source/ChaosClothAssetEditor/Private/ChaosClothAsset/ClothEditorMode.cpp
--- a/Engine/Plugins/ChaosClothAssetEditor/Source/ChaosClothAssetEditor/Private/ChaosClothAsset/ClothEditorMode.cpp
+++ b/Engine/Plugins/ChaosClothAssetEditor/Source/ChaosClothAssetEditor/Private/ChaosClothAsset/ClothEditorMode.cpp
@@ -468,6 +468,12 @@ void UChaosClothAssetEditorMode::Exit()
 	}
 	ClothSeamDraw = nullptr;
 
+	if (SurfaceNormalDraw)
+	{
+		SurfaceNormalDraw->Disconnect();
+	}
+	SurfaceNormalDraw = nullptr;
+
 	if (DataflowComponent)
 	{
 		DataflowComponent->UnregisterComponent();
@@ -714,6 +720,50 @@ void UChaosClothAssetEditorMode::InitializeSeamDraw()
 }
 
 
+void UChaosClothAssetEditorMode::InitializeSurfaceNormalDraw()
+{
+	if (!SurfaceNormalDraw)
+	{
+		return;
+	}
+
+	SurfaceNormalDraw->RemoveAllLineSets();
+
+	if (!bConstructionViewNormalsVisible)
+	{
+		return;
+	}
+
+	if (!DynamicMeshComponent || !DynamicMeshComponent->GetMesh())
+	{
+		return;
+	}
+
+	const UE::Geometry::FDynamicMesh3& Mesh = *DynamicMeshComponent->GetMesh();
+
+	const float LineLength = 2.0f;
+	const float LineThickness = 1.0f;
+
+	if (Mesh.HasAttributes())
+	{
+		if (const UE::Geometry::FDynamicMeshNormalOverlay* const NormalOverlay = Mesh.Attributes()->PrimaryNormals())
+		{
+			SurfaceNormalDraw->CreateOrUpdateLineSet(TEXT("Normals"), NormalOverlay->MaxElementID(),
+			[&NormalOverlay, &Mesh, LineLength, LineThickness](int32 Index, TArray<FRenderableLine>& Lines)
+			{
+				if (NormalOverlay->IsElement(Index))
+				{
+					int32 ParentVtx = NormalOverlay->GetParentVertex(Index);
+					FVector3f Normal = NormalOverlay->GetElement(Index);
+					FVector3f Origin = (FVector3f)Mesh.GetVertex(ParentVtx);
+					Lines.Add(FRenderableLine((FVector)Origin, (FVector)Origin + LineLength * (FVector)Normal, FColor(15, 15, 240), LineThickness));
+				}
+			}, 1);
+		}
+	}
+}
+
+
 void UChaosClothAssetEditorMode::ReinitializeDynamicMeshComponents()
 {
 	using namespace UE::Chaos::ClothAsset;
@@ -811,11 +861,17 @@ void UChaosClothAssetEditorMode::ReinitializeDynamicMeshComponents()
 		ClothSeamDraw->Disconnect();
 	}
 
+	if (SurfaceNormalDraw)
+	{
+		SurfaceNormalDraw->Disconnect();
+	}
+
 	PropertyObjectsToTick.Empty();	// TODO: We only want to empty the wireframe display properties. Is anything else using this array?
 	DynamicMeshComponent = nullptr;
 	DynamicMeshComponentParentActor = nullptr;
 	WireframeDraw = nullptr;
 	ClothSeamDraw = nullptr;
+	SurfaceNormalDraw = nullptr;
 	bDynamicMeshComponentInitDeferred = false;
 
 	TSharedPtr<FManagedArrayCollection> Collection = bDynamicMeshUseInputCollection ? GetInputClothCollection() : GetClothCollection();
@@ -917,12 +973,20 @@ void UChaosClothAssetEditorMode::ReinitializeDynamicMeshComponents()
 	InitializeSeamDraw();
 	ClothSeamDraw->SetAllVisible(bRestSpaceMeshVisible && bConstructionViewSeamsVisible);
 
+	SurfaceNormalDraw = NewObject<UPreviewGeometry>(this);
+	SurfaceNormalDraw->CreateInWorld(GetWorld(), FTransform::Identity);
+	InitializeSurfaceNormalDraw();
+	SurfaceNormalDraw->SetAllVisible(bRestSpaceMeshVisible && bConstructionViewNormalsVisible);
+
 	DynamicMeshComponent->OnMeshChanged.Add(
 		FSimpleMulticastDelegate::FDelegate::CreateLambda([this]()
 		{
 			InitializeSeamDraw();
 			const bool bRestSpaceMeshVisible = DynamicMeshComponent->GetVisibleFlag();
 			ClothSeamDraw->SetAllVisible(bRestSpaceMeshVisible&& bConstructionViewSeamsVisible);
+
+			InitializeSurfaceNormalDraw();
+			SurfaceNormalDraw->SetAllVisible(bRestSpaceMeshVisible && bConstructionViewNormalsVisible);
 		}));
 
 
@@ -943,6 +1007,10 @@ void UChaosClothAssetEditorMode::ReinitializeDynamicMeshComponents()
 			{
 				ClothSeamDraw->SetAllVisible(bRestSpaceMeshVisible && bConstructionViewSeamsVisible);
 			}
+			if (SurfaceNormalDraw)
+			{
+				SurfaceNormalDraw->SetAllVisible(bRestSpaceMeshVisible && bConstructionViewNormalsVisible);
+			}
 		});
 
 
@@ -1444,6 +1512,22 @@ bool UChaosClothAssetEditorMode::CanSetConstructionViewSeamsCollapse() const
 	return bConstructionViewSeamsVisible && (ConstructionViewMode == UE::Chaos::ClothAsset::EClothPatternVertexType::Sim2D);
 }
 
+void UChaosClothAssetEditorMode::ToggleConstructionViewSurfaceNormals()
+{
+	bConstructionViewNormalsVisible = !bConstructionViewNormalsVisible;
+	ReinitializeDynamicMeshComponents();
+}
+
+bool UChaosClothAssetEditorMode::CanSetConstructionViewSurfaceNormalsActive() const
+{
+	if (GetToolManager()->HasActiveTool(EToolSide::Left))
+	{
+		return false;
+	}
+
+	return true;
+}
+
 
 void UChaosClothAssetEditorMode::TogglePatternColor()
 {
diff --git a/Engine/Plugins/ChaosClothAssetEditor/Source/ChaosClothAssetEditor/Private/ChaosClothAsset/ClothEditorMode.h b/Engine/Plugins/ChaosClothAssetEditor/Source/ChaosClothAssetEditor/Private/ChaosClothAsset/ClothEditorMode.h
--- a/Engine/Plugins/ChaosClothAssetEditor/Source/ChaosClothAssetEditor/Private/ChaosClothAsset/ClothEditorMode.h
+++ b/Engine/Plugins/ChaosClothAssetEditor/Source/ChaosClothAssetEditor/Private/ChaosClothAsset/ClothEditorMode.h
@@ -96,6 +96,13 @@ public:
 		return bPatternColors;
 	}
 
+	void ToggleConstructionViewSurfaceNormals();
+	bool CanSetConstructionViewSurfaceNormalsActive() const;
+	bool IsConstructionViewSurfaceNormalsActive() const
+	{
+		return bConstructionViewNormalsVisible;
+	}
+
 	void ToggleMeshStats();
 	bool CanSetMeshStats() const;
 	bool IsMeshStatsActive() const
@@ -235,6 +242,8 @@ private:
 	UPROPERTY()
 	TObjectPtr<UPreviewGeometry> ClothSeamDraw = nullptr;
 
+	UPROPERTY()
+	TObjectPtr<UPreviewGeometry> SurfaceNormalDraw = nullptr;
 
 	// Preview Scene, here largely for convenience to avoid having to pass it around functions. Owned by the ClothEditorToolkit.
 	UE::Chaos::ClothAsset::FChaosClothPreviewScene* PreviewScene = nullptr;
@@ -286,6 +295,9 @@ private:
 	bool bConstructionViewSeamsCollapse = false;
 	void InitializeSeamDraw();
 
+	bool bConstructionViewNormalsVisible = false;
+	void InitializeSurfaceNormalDraw();
+
 	bool bPatternColors = false;
 	bool bMeshStats = false;
 
diff --git a/Engine/Plugins/ChaosClothAssetEditor/Source/ChaosClothAssetEditor/Private/ChaosClothAsset/SClothEditorRestSpaceViewport.cpp b/Engine/Plugins/ChaosClothAssetEditor/Source/ChaosClothAssetEditor/Private/ChaosClothAsset/SClothEditorRestSpaceViewport.cpp
--- a/Engine/Plugins/ChaosClothAssetEditor/Source/ChaosClothAssetEditor/Private/ChaosClothAsset/SClothEditorRestSpaceViewport.cpp
+++ b/Engine/Plugins/ChaosClothAssetEditor/Source/ChaosClothAssetEditor/Private/ChaosClothAsset/SClothEditorRestSpaceViewport.cpp
@@ -258,6 +258,33 @@ void SChaosClothAssetEditorRestSpaceViewport::BindCommands()
 				return false;
 			}));
 
+	CommandList->MapAction(
+		CommandInfos.ToggleConstructionViewSurfaceNormals,
+		FExecuteAction::CreateLambda([this]()
+			{
+				if (UChaosClothAssetEditorMode* const ClothEdMode = GetEdMode())
+				{
+					ClothEdMode->ToggleConstructionViewSurfaceNormals();
+				}
+			}),
+		FCanExecuteAction::CreateLambda([this]()
+			{
+				if (const UChaosClothAssetEditorMode* const ClothEdMode = GetEdMode())
+				{
+					return ClothEdMode->CanSetConstructionViewSurfaceNormalsActive();
+				}
+				return false;
+			}),
+		FIsActionChecked::CreateLambda([this]()
+			{
+				if (const UChaosClothAssetEditorMode* const ClothEdMode = GetEdMode())
+				{
+					return ClothEdMode->IsConstructionViewSurfaceNormalsActive();
+				}
+				return false;
+			}));
+
+
 }
 
 TSharedPtr<SWidget> SChaosClothAssetEditorRestSpaceViewport::MakeViewportToolbar()
diff --git a/Engine/Plugins/ChaosClothAssetEditor/Source/ChaosClothAssetEditor/Private/ChaosClothAsset/SClothEditorRestSpaceViewportToolBar.cpp b/Engine/Plugins/ChaosClothAssetEditor/Source/ChaosClothAssetEditor/Private/ChaosClothAsset/SClothEditorRestSpaceViewportToolBar.cpp
--- a/Engine/Plugins/ChaosClothAssetEditor/Source/ChaosClothAssetEditor/Private/ChaosClothAsset/SClothEditorRestSpaceViewportToolBar.cpp
+++ b/Engine/Plugins/ChaosClothAssetEditor/Source/ChaosClothAssetEditor/Private/ChaosClothAsset/SClothEditorRestSpaceViewportToolBar.cpp
@@ -128,6 +128,10 @@ TSharedRef<SWidget> SChaosClothAssetEditorRestSpaceViewportToolBar::GenerateClot
 				LOCTEXT("ToggleMeshStatsLabel", "Mesh Stats"),
 				LOCTEXT("ToggleMeshStatsTooltip", "Show mesh stats in the viewport"));
 
+			OptionsMenuBuilder.AddMenuEntry(FChaosClothAssetEditorCommands::Get().ToggleConstructionViewSurfaceNormals,
+				NAME_None,
+				LOCTEXT("ShowNormalsLabel", "Show Normals"),
+				LOCTEXT("ShowNormalsTooltip", "Display surface normals"));
 		}
 		OptionsMenuBuilder.EndSection();
 
