Cloth Editor: add option to pause Preview viewport updates while either PIE or SIE is running in the level editor

diff --git a/Engine/Plugins/ChaosClothAssetEditor/Source/ChaosClothAssetEditor/Private/ChaosClothAsset/ClothEditorMode.cpp b/Engine/Plugins/ChaosClothAssetEditor/Source/ChaosClothAssetEditor/Private/ChaosClothAsset/ClothEditorMode.cpp
--- a/Engine/Plugins/ChaosClothAssetEditor/Source/ChaosClothAssetEditor/Private/ChaosClothAsset/ClothEditorMode.cpp
+++ b/Engine/Plugins/ChaosClothAssetEditor/Source/ChaosClothAssetEditor/Private/ChaosClothAsset/ClothEditorMode.cpp
@@ -1289,9 +1289,15 @@ void UChaosClothAssetEditorMode::ModeTick(float DeltaTime)
 		FirstTimeFocusRestSpaceViewport();
 	}
 
-	if (PreviewScene->GetWorld())
+	const bool bIsInPIEOrSIE = GEditor->PlayWorld != NULL || GEditor->bIsSimulatingInEditor;
+	const bool bShouldPause = PreviewScene->GetPreviewSceneDescription()->bPauseWhilePlayingInEditor && bIsInPIEOrSIE;
+
+	if (!bShouldPause)
 	{
-		PreviewScene->GetWorld()->Tick(ELevelTick::LEVELTICK_All, DeltaTime);
+		if (PreviewScene->GetWorld())
+		{
+			PreviewScene->GetWorld()->Tick(ELevelTick::LEVELTICK_All, DeltaTime);
+		}
 	}
 }
 
diff --git a/Engine/Plugins/ChaosClothAssetEditor/Source/ChaosClothAssetEditor/Private/ChaosClothAsset/ClothEditorPreviewScene.h b/Engine/Plugins/ChaosClothAssetEditor/Source/ChaosClothAssetEditor/Private/ChaosClothAsset/ClothEditorPreviewScene.h
--- a/Engine/Plugins/ChaosClothAssetEditor/Source/ChaosClothAssetEditor/Private/ChaosClothAsset/ClothEditorPreviewScene.h
+++ b/Engine/Plugins/ChaosClothAssetEditor/Source/ChaosClothAssetEditor/Private/ChaosClothAsset/ClothEditorPreviewScene.h
@@ -39,6 +39,10 @@ public:
 
 	void SetPreviewScene(UE::Chaos::ClothAsset::FChaosClothPreviewScene* PreviewScene);
 
+	/* Whether the preview viewport should pause animation and simulation while Play In Editor (PIE) or Simulate In Editor is active */
+	UPROPERTY(EditAnywhere, Transient, Category = "Viewport")
+	bool bPauseWhilePlayingInEditor = true;
+
 	// Skeletal Mesh source asset
 	UPROPERTY(EditAnywhere, Transient, Category="SkeletalMesh")
 	TObjectPtr<USkeletalMesh> SkeletalMeshAsset;
