Cloth Editor: remove references to DataflowComponent (unused in the cloth editor)

diff --git a/Engine/Plugins/ChaosClothAssetEditor/Source/ChaosClothAssetEditor/Private/ChaosClothAsset/ClothEditor3DViewportClient.cpp b/Engine/Plugins/ChaosClothAssetEditor/Source/ChaosClothAssetEditor/Private/ChaosClothAsset/ClothEditor3DViewportClient.cpp
--- a/Engine/Plugins/ChaosClothAssetEditor/Source/ChaosClothAssetEditor/Private/ChaosClothAsset/ClothEditor3DViewportClient.cpp
+++ b/Engine/Plugins/ChaosClothAssetEditor/Source/ChaosClothAssetEditor/Private/ChaosClothAsset/ClothEditor3DViewportClient.cpp
@@ -127,58 +127,6 @@ void FChaosClothAssetEditor3DViewportClient::AddReferencedObjects(FReferenceColl
 	Collector.AddReferencedObject(InputBehaviorSet);
 }
 
-void FChaosClothAssetEditor3DViewportClient::Tick(float DeltaSeconds)
-{
-	FEditorViewportClient::Tick(DeltaSeconds);
-
-	auto GetLatestTimestamp = [](const UDataflow* Dataflow, const Dataflow::FContext* Context) -> Dataflow::FTimestamp
-	{
-		if (Dataflow && Context)
-		{
-			return FMath::Max(Dataflow->GetRenderingTimestamp().Value, Context->GetTimestamp().Value);
-		}
-		return Dataflow::FTimestamp::Invalid;
-	};
-
-	if (TSharedPtr<const FChaosClothAssetEditorToolkit> PinnedClothToolkit = ClothToolkit.Pin())
-	{
-		if (const TSharedPtr<Dataflow::FContext> Context = PinnedClothToolkit->GetDataflowContext())
-		{
-			if (const UDataflow* const Dataflow = PinnedClothToolkit->GetDataflow())
-			{
-				if (UDataflowComponent* const DataflowComponent = ClothEdMode->GetDataflowComponent())
-				{
-					const Dataflow::FTimestamp SystemTimestamp = GetLatestTimestamp(Dataflow, Context.Get());
-
-					if (SystemTimestamp >= LastModifiedTimestamp)
-					{
-						if (Dataflow->GetRenderTargets().Num())
-						{
-							// Component Object Rendering
-							DataflowComponent->ResetRenderTargets();
-							DataflowComponent->SetDataflow(Dataflow);
-							DataflowComponent->SetContext(Context);
-							for (const UDataflowEdNode* const Node : Dataflow->GetRenderTargets())
-							{
-								DataflowComponent->AddRenderTarget(Node);
-							}
-						}
-						else
-						{
-							DataflowComponent->ResetRenderTargets();
-						}
-
-						LastModifiedTimestamp = GetLatestTimestamp(Dataflow, Context.Get()).Value + 1;
-					}
-				}
-			}
-		}
-	}
-
-	// Note: we don't tick the PreviewWorld here, that is done in UChaosClothAssetEditorMode::ModeTick()
-
-}
-
 void FChaosClothAssetEditor3DViewportClient::EnableRenderMeshWireframe(bool bEnable)
 {
 	bRenderMeshWireframe = bEnable;
diff --git a/Engine/Plugins/ChaosClothAssetEditor/Source/ChaosClothAssetEditor/Private/ChaosClothAsset/ClothEditor3DViewportClient.h b/Engine/Plugins/ChaosClothAssetEditor/Source/ChaosClothAssetEditor/Private/ChaosClothAsset/ClothEditor3DViewportClient.h
--- a/Engine/Plugins/ChaosClothAssetEditor/Source/ChaosClothAssetEditor/Private/ChaosClothAsset/ClothEditor3DViewportClient.h
+++ b/Engine/Plugins/ChaosClothAssetEditor/Source/ChaosClothAssetEditor/Private/ChaosClothAsset/ClothEditor3DViewportClient.h
@@ -90,7 +90,6 @@ private:
 	virtual bool CanSetWidgetMode(UE::Widget::EWidgetMode NewMode) const override { return false; }
 	virtual void SetWidgetMode(UE::Widget::EWidgetMode NewMode) override {}
 	virtual UE::Widget::EWidgetMode GetWidgetMode() const override { return UE::Widget::EWidgetMode::WM_None; }
-	virtual void Tick(float DeltaSeconds) override;
 	virtual void ProcessClick(FSceneView& View, HHitProxy* HitProxy, FKey Key, EInputEvent Event, uint32 HitX, uint32 HitY) override;
 	virtual void Draw(const FSceneView* View, FPrimitiveDrawInterface* PDI) override;
 	virtual void DrawCanvas(FViewport& InViewport, FSceneView& View, FCanvas& Canvas) override;
diff --git a/Engine/Plugins/ChaosClothAssetEditor/Source/ChaosClothAssetEditor/Private/ChaosClothAsset/ClothEditorMode.cpp b/Engine/Plugins/ChaosClothAssetEditor/Source/ChaosClothAssetEditor/Private/ChaosClothAsset/ClothEditorMode.cpp
--- a/Engine/Plugins/ChaosClothAssetEditor/Source/ChaosClothAssetEditor/Private/ChaosClothAsset/ClothEditorMode.cpp
+++ b/Engine/Plugins/ChaosClothAssetEditor/Source/ChaosClothAssetEditor/Private/ChaosClothAsset/ClothEditorMode.cpp
@@ -474,12 +474,6 @@ void UChaosClothAssetEditorMode::Exit()
 	}
 	SurfaceNormalDraw = nullptr;
 
-	if (DataflowComponent)
-	{
-		DataflowComponent->UnregisterComponent();
-		DataflowComponent->DestroyComponent();
-	}
-
 	PropertyObjectsToTick.Empty();
 	PreviewScene = nullptr;
 
@@ -1084,9 +1078,6 @@ void UChaosClothAssetEditorMode::InitializeTargets(const TArray<TObjectPtr<UObje
 	check(PreviewScene);
 
 	UBaseCharacterFXEditorMode::InitializeTargets(AssetsIn);
-
-	DataflowComponent = NewObject<UDataflowComponent>();
-	DataflowComponent->RegisterComponentWithWorld(PreviewScene->GetWorld());
 }
 
 void UChaosClothAssetEditorMode::SoftResetSimulation()
@@ -1200,11 +1191,6 @@ int32 UChaosClothAssetEditorMode::GetNumLODs() const
 	return 0;
 }
 
-UDataflowComponent* UChaosClothAssetEditorMode::GetDataflowComponent() const
-{
-	return DataflowComponent;
-}
-
 void UChaosClothAssetEditorMode::ModeTick(float DeltaTime)
 {
 	Super::ModeTick(DeltaTime);
diff --git a/Engine/Plugins/ChaosClothAssetEditor/Source/ChaosClothAssetEditor/Private/ChaosClothAsset/ClothEditorMode.h b/Engine/Plugins/ChaosClothAssetEditor/Source/ChaosClothAssetEditor/Private/ChaosClothAsset/ClothEditorMode.h
--- a/Engine/Plugins/ChaosClothAssetEditor/Source/ChaosClothAssetEditor/Private/ChaosClothAsset/ClothEditorMode.h
+++ b/Engine/Plugins/ChaosClothAssetEditor/Source/ChaosClothAssetEditor/Private/ChaosClothAsset/ClothEditorMode.h
@@ -27,7 +27,6 @@ class UChaosClothComponent;
 class FEditorViewportClient;
 class FViewport;
 class UDataflow;
-class UDataflowComponent;
 class SDataflowGraphEditor;
 struct FManagedArrayCollection;
 
@@ -128,8 +127,6 @@ public:
 	int32 GetLODModel() const;
 	int32 GetNumLODs() const;
 
-	UDataflowComponent* GetDataflowComponent() const;
-
 	TObjectPtr<UEditorInteractiveToolsContext> GetActiveToolsContext()
 	{
 		return ActiveToolsContext;
@@ -309,9 +306,6 @@ private:
 	bool bHardReset = false;
 	bool bShouldClearTeleportFlag = false;
 
-	UPROPERTY()
-	TObjectPtr<UDataflowComponent> DataflowComponent = nullptr;
-
 	TWeakObjectPtr<UDataflow> DataflowGraph = nullptr;
 
 	TWeakPtr<SDataflowGraphEditor> DataflowGraphEditor;
