Chaos Cloth Asset: Add kinematic vertices set that's generated by the MaxDistanceConfig node. Switch LRA to use a set rather than a map.

diff --git a/Engine/Plugins/ChaosCloth/Source/ChaosCloth/Private/ChaosCloth/ChaosClothingSimulationCloth.cpp b/Engine/Plugins/ChaosCloth/Source/ChaosCloth/Private/ChaosCloth/ChaosClothingSimulationCloth.cpp
--- a/Engine/Plugins/ChaosCloth/Source/ChaosCloth/Private/ChaosCloth/ChaosClothingSimulationCloth.cpp
+++ b/Engine/Plugins/ChaosCloth/Source/ChaosCloth/Private/ChaosCloth/ChaosClothingSimulationCloth.cpp
@@ -114,6 +114,7 @@ struct FClothingSimulationCloth::FLODData
 	UE_CHAOS_DECLARE_INDEXLESS_PROPERTYCOLLECTION_NAME(MassValue, float);
 	UE_CHAOS_DECLARE_INDEXLESS_PROPERTYCOLLECTION_NAME(MultiResCoarseLODIndex, int32);
 	UE_CHAOS_DECLARE_INDEXLESS_PROPERTYCOLLECTION_NAME(IsCoarseMultiResLOD, bool);
+	UE_CHAOS_DECLARE_INDEXLESS_PROPERTYCOLLECTION_NAME(KinematicVertices3D, bool); // Selection set name string property; bool value is not actually used.
 
 };
 
@@ -216,12 +217,14 @@ void FClothingSimulationCloth::FLODData::AddParticles(FClothingSimulationSolver*
 		WeightMaps.FindRef(GetMaxDistanceString(ConfigProperties, MaxDistanceName.ToString())),
 		NumParticles);
 
+	const TSet<int32>* const KinematicVertices3DSet = VertexSets.FindRef(GetKinematicVertices3DString(ConfigProperties, KinematicVertices3DName.ToString()), nullptr);
+
 	// Set the particle masses
 	static const FRealSingle KinematicDistanceThreshold = 0.1f;  // TODO: This is not the same value as set in the painting UI but we might want to expose this value as parameter
 	auto KinematicPredicate =
-		[&MaxDistances](int32 Index)
+		[&MaxDistances, KinematicVertices3DSet](int32 Index)
 	{
-		return MaxDistances.GetValue(Index) < KinematicDistanceThreshold;
+		return MaxDistances.GetValue(Index) < KinematicDistanceThreshold || (KinematicVertices3DSet && KinematicVertices3DSet->Contains(Index));
 	};
 
 	const int32 MassMode = ConfigProperties.GetValue<int32>(TEXT("MassMode"), ClothingSimulationClothDefault::MassMode);
diff --git a/Engine/Plugins/ChaosClothAsset/Source/ChaosClothAsset/Private/ChaosClothAsset/ClothGeometryTools.cpp b/Engine/Plugins/ChaosClothAsset/Source/ChaosClothAsset/Private/ChaosClothAsset/ClothGeometryTools.cpp
--- a/Engine/Plugins/ChaosClothAsset/Source/ChaosClothAsset/Private/ChaosClothAsset/ClothGeometryTools.cpp
+++ b/Engine/Plugins/ChaosClothAsset/Source/ChaosClothAsset/Private/ChaosClothAsset/ClothGeometryTools.cpp
@@ -1491,4 +1491,41 @@ namespace UE::Chaos::ClothAsset
 		}
 	}
 
+	TSet<int32> FClothGeometryTools::GenerateKinematicVertices3D(const TSharedRef<FManagedArrayCollection>& ClothCollection, const FName& MaxDistanceMapName, const FVector2f& MaxDistanceValue, const FName& InputKinematicVertices, float KinematicDistanceThreshold)
+	{
+		TSet<int32> KinematicVertices;
+
+		// Add InputKinematicVertices
+		FCollectionClothSelectionConstFacade SelectionFacade(ClothCollection);
+		if (InputKinematicVertices != NAME_None && SelectionFacade.IsValid() && SelectionFacade.HasSelection(InputKinematicVertices) && SelectionFacade.GetSelectionGroup(InputKinematicVertices) == UE::Chaos::ClothAsset::ClothCollectionGroup::SimVertices3D)
+		{
+			KinematicVertices = SelectionFacade.GetSelectionSet(InputKinematicVertices);
+		}
+
+		FCollectionClothFacade ClothFacade(ClothCollection);
+		if (ClothFacade.IsValid())
+		{
+			if (ClothFacade.HasWeightMap(MaxDistanceMapName))
+			{
+				TConstArrayView<float> MaxDistanceMap = ClothFacade.GetWeightMap(MaxDistanceMapName);
+				const FVector2f MaxDistanceOffsetRange(MaxDistanceValue[0], MaxDistanceValue[1] - MaxDistanceValue[0]);
+				for (int32 Index = 0; Index < MaxDistanceMap.Num(); ++Index)
+				{
+					if (MaxDistanceOffsetRange[0] + MaxDistanceMap[Index] * MaxDistanceOffsetRange[1] < KinematicDistanceThreshold)
+					{
+						KinematicVertices.Add(Index);
+					}
+				}
+			}
+			else if (MaxDistanceValue[0] < KinematicDistanceThreshold)
+			{
+				KinematicVertices.Reserve(ClothFacade.GetNumSimVertices3D());
+				for (int32 Index = 0; Index < ClothFacade.GetNumSimVertices3D(); ++Index)
+				{
+					KinematicVertices.Add(Index);
+				}
+			}
+		}
+		return KinematicVertices;
+	}
 }  // End namespace UE::Chaos::ClothAsset
diff --git a/Engine/Plugins/ChaosClothAsset/Source/ChaosClothAsset/Public/ChaosClothAsset/ClothGeometryTools.h b/Engine/Plugins/ChaosClothAsset/Source/ChaosClothAsset/Public/ChaosClothAsset/ClothGeometryTools.h
--- a/Engine/Plugins/ChaosClothAsset/Source/ChaosClothAsset/Public/ChaosClothAsset/ClothGeometryTools.h
+++ b/Engine/Plugins/ChaosClothAsset/Source/ChaosClothAsset/Public/ChaosClothAsset/ClothGeometryTools.h
@@ -141,5 +141,10 @@ namespace UE::Chaos::ClothAsset
 			const TConstArrayView<FVector3f>& TargetNormals,
 			const TConstArrayView<FIntVector3>& InTargetIndices,
 			const TArrayView<float>& TargetWeights);
+
+		/**
+		 * Generate KinematicVertices3D set from the given MaxDistance weight map, MaxDistance values, and any additional kinematic vertices.
+		 */
+		static TSet<int32> GenerateKinematicVertices3D(const TSharedRef<FManagedArrayCollection>& ClothCollection, const FName& MaxDistanceMap, const FVector2f& MaxDistanceValues, const FName& InputKinematicVertices, float KinematicDistanceThreshold = 0.1f);
 	};
 }  // End namespace UE::Chaos::ClothAsset
\ No newline at end of file
diff --git a/Engine/Plugins/ChaosClothAsset/Source/ChaosClothAssetEngine/Private/ChaosClothAsset/ClothEngineTools.cpp b/Engine/Plugins/ChaosClothAsset/Source/ChaosClothAssetEngine/Private/ChaosClothAsset/ClothEngineTools.cpp
--- a/Engine/Plugins/ChaosClothAsset/Source/ChaosClothAssetEngine/Private/ChaosClothAsset/ClothEngineTools.cpp
+++ b/Engine/Plugins/ChaosClothAsset/Source/ChaosClothAssetEngine/Private/ChaosClothAsset/ClothEngineTools.cpp
@@ -2,13 +2,39 @@
 
 #include "ChaosClothAsset/ClothEngineTools.h"
 #include "ChaosClothAsset/CollectionClothFacade.h"
+#include "ChaosClothAsset/CollectionClothSelectionFacade.h"
+#include "ChaosClothAsset/ClothCollectionGroup.h"
 #include "ChaosClothAsset/ClothGeometryTools.h"
 #include "ClothTetherData.h"
 
 namespace UE::Chaos::ClothAsset
 {
 
-void FClothEngineTools::GenerateTethers(const TSharedRef<FManagedArrayCollection>& ClothCollection, const FName& WeightMapName, const bool bGenerateGeodesicTethers)
+namespace Private
+{
+static void AppendTetherData(FCollectionClothFacade& ClothFacade, const FClothTetherData& TetherData)
+{
+	// Append new tethers
+	TArrayView<TArray<int32>> TetherKinematicIndex = ClothFacade.GetTetherKinematicIndex();
+	TArrayView<TArray<float>> TetherReferenceLength = ClothFacade.GetTetherReferenceLength();
+	for (const TArray<TTuple<int32, int32, float>>& TetherBatch : TetherData.Tethers)
+	{
+		for (const TTuple<int32, int32, float>& Tether : TetherBatch)
+		{
+			// Tuple is Kinematic, Dynamic, RefLength
+			const int32 DynamicIndex = Tether.Get<1>();
+			TArray<int32>& KinematicIndex = TetherKinematicIndex[DynamicIndex];
+			TArray<float>& ReferenceLength = TetherReferenceLength[DynamicIndex];
+			check(KinematicIndex.Num() == ReferenceLength.Num());
+			checkSlow(KinematicIndex.Find(Tether.Get<0>()) == INDEX_NONE);
+			KinematicIndex.Add(Tether.Get<0>());
+			ReferenceLength.Add(Tether.Get<2>());
+		}
+	}
+}
+}
+
+void FClothEngineTools::GenerateTethers(const TSharedRef<FManagedArrayCollection>& ClothCollection, const FName& WeightMapName, const bool bGenerateGeodesicTethers, const FVector2f& MaxDistanceValue)
 {
 	FCollectionClothFacade ClothFacade(ClothCollection);
 	FClothGeometryTools::DeleteTethers(ClothCollection);
@@ -23,25 +49,42 @@ void FClothEngineTools::GenerateTethers(const TSharedRef<FManagedArrayCollection
 			SimIndices.Add(Face[1]);
 			SimIndices.Add(Face[2]);
 		}
-		TetherData.GenerateTethers(ClothFacade.GetSimPosition3D(), TConstArrayView<uint32>(SimIndices), ClothFacade.GetWeightMap(WeightMapName), bGenerateGeodesicTethers);
 
-		// Append new tethers
-		TArrayView<TArray<int32>> TetherKinematicIndex = ClothFacade.GetTetherKinematicIndex();
-		TArrayView<TArray<float>> TetherReferenceLength = ClothFacade.GetTetherReferenceLength();
-		for(const TArray<TTuple<int32, int32, float>>& TetherBatch : TetherData.Tethers)
+		if (MaxDistanceValue.Equals(FVector2f(0.f, 1.f)))
+		{
+			TetherData.GenerateTethers(ClothFacade.GetSimPosition3D(), TConstArrayView<uint32>(SimIndices), ClothFacade.GetWeightMap(WeightMapName), bGenerateGeodesicTethers);
+		}
+		else
+		{
+
+			const TSet<int32> KinematicVertices = FClothGeometryTools::GenerateKinematicVertices3D(ClothCollection, WeightMapName, MaxDistanceValue, NAME_None);
+			TetherData.GenerateTethers(ClothFacade.GetSimPosition3D(), TConstArrayView<uint32>(SimIndices), KinematicVertices, bGenerateGeodesicTethers);
+		}
+
+		Private::AppendTetherData(ClothFacade, TetherData);
+	}
+}
+
+void FClothEngineTools::GenerateTethersFromSelectionSet(const TSharedRef<FManagedArrayCollection>& ClothCollection, const FName& FixedEndSet, const bool bGeodesicTethers)
+{
+	FClothGeometryTools::DeleteTethers(ClothCollection);
+	FCollectionClothSelectionConstFacade SelectionFacade(ClothCollection);
+
+	if (SelectionFacade.HasSelection(FixedEndSet) && SelectionFacade.GetSelectionGroup(FixedEndSet) == UE::Chaos::ClothAsset::ClothCollectionGroup::SimVertices3D)
+	{
+		FCollectionClothFacade ClothFacade(ClothCollection);
+		FClothTetherData TetherData;
+		TArray<uint32> SimIndices;
+		SimIndices.Reserve(ClothFacade.GetNumSimFaces() * 3);
+		for (const FIntVector3& Face : ClothFacade.GetSimIndices3D())
 		{
-			for (const TTuple<int32, int32, float>& Tether : TetherBatch)
-			{
-				// Tuple is Kinematic, Dynamic, RefLength
-				const int32 DynamicIndex = Tether.Get<1>();
-				TArray<int32>& KinematicIndex = TetherKinematicIndex[DynamicIndex];
-				TArray<float>& ReferenceLength = TetherReferenceLength[DynamicIndex];
-				check(KinematicIndex.Num() == ReferenceLength.Num());
-				checkSlow(KinematicIndex.Find(Tether.Get<0>()) == INDEX_NONE);
-				KinematicIndex.Add(Tether.Get<0>());
-				ReferenceLength.Add(Tether.Get<2>());
-			}
+			SimIndices.Add(Face[0]);
+			SimIndices.Add(Face[1]);
+			SimIndices.Add(Face[2]);
 		}
+
+		TetherData.GenerateTethers(ClothFacade.GetSimPosition3D(), TConstArrayView<uint32>(SimIndices), SelectionFacade.GetSelectionSet(FixedEndSet), bGeodesicTethers);
+		Private::AppendTetherData(ClothFacade, TetherData);
 	}
 }
 
diff --git a/Engine/Plugins/ChaosClothAsset/Source/ChaosClothAssetEngine/Public/ChaosClothAsset/ClothEngineTools.h b/Engine/Plugins/ChaosClothAsset/Source/ChaosClothAssetEngine/Public/ChaosClothAsset/ClothEngineTools.h
--- a/Engine/Plugins/ChaosClothAsset/Source/ChaosClothAssetEngine/Public/ChaosClothAsset/ClothEngineTools.h
+++ b/Engine/Plugins/ChaosClothAsset/Source/ChaosClothAssetEngine/Public/ChaosClothAsset/ClothEngineTools.h
@@ -3,6 +3,7 @@
 #pragma once
 
 #include "Templates/SharedPointer.h"
+#include "Math/Vector2D.h"
 
 struct FManagedArrayCollection;
 class FName;
@@ -16,6 +17,7 @@ namespace UE::Chaos::ClothAsset
 	struct CHAOSCLOTHASSETENGINE_API FClothEngineTools
 	{
 		/** Generate tether data. */
-		static void GenerateTethers(const TSharedRef<FManagedArrayCollection>& ClothCollection, const FName& WeightMap, const bool bGeodesicTethers);
+		static void GenerateTethers(const TSharedRef<FManagedArrayCollection>& ClothCollection, const FName& WeightMap, const bool bGeodesicTethers, const FVector2f& MaxDistanceValue = FVector2f(0.f, 1.f));
+		static void GenerateTethersFromSelectionSet(const TSharedRef<FManagedArrayCollection>& ClothCollection, const FName& FixedEndSet, const bool bGeodesicTethers);
 	};
 }  // End namespace UE::Chaos::ClothAsset
\ No newline at end of file
diff --git a/Engine/Plugins/ChaosClothAssetEditor/Source/ChaosClothAssetDataflowNodes/Private/ChaosClothAsset/ChaosClothAssetDataflowNodesModule.cpp b/Engine/Plugins/ChaosClothAssetEditor/Source/ChaosClothAssetDataflowNodes/Private/ChaosClothAsset/ChaosClothAssetDataflowNodesModule.cpp
--- a/Engine/Plugins/ChaosClothAssetEditor/Source/ChaosClothAssetDataflowNodes/Private/ChaosClothAsset/ChaosClothAssetDataflowNodesModule.cpp
+++ b/Engine/Plugins/ChaosClothAssetEditor/Source/ChaosClothAssetDataflowNodes/Private/ChaosClothAsset/ChaosClothAssetDataflowNodesModule.cpp
@@ -109,6 +109,7 @@ namespace UE::Chaos::ClothAsset
 			DATAFLOW_NODE_REGISTER_CREATION_FACTORY(FChaosClothAssetSimulationDefaultConfigNode);
 			DATAFLOW_NODE_REGISTER_CREATION_FACTORY(FChaosClothAssetSimulationGravityConfigNode);
 			DATAFLOW_NODE_REGISTER_CREATION_FACTORY(FChaosClothAssetSimulationLongRangeAttachmentConfigNode);
+			DATAFLOW_NODE_REGISTER_CREATION_FACTORY(FChaosClothAssetSimulationLongRangeAttachmentConfigNode_v2);
 			DATAFLOW_NODE_REGISTER_CREATION_FACTORY(FChaosClothAssetSimulationMassConfigNode);
 			DATAFLOW_NODE_REGISTER_CREATION_FACTORY(FChaosClothAssetSimulationMaxDistanceConfigNode);
 			DATAFLOW_NODE_REGISTER_CREATION_FACTORY(FChaosClothAssetSimulationPBDAreaSpringConfigNode);
diff --git a/Engine/Plugins/ChaosClothAssetEditor/Source/ChaosClothAssetDataflowNodes/Private/ChaosClothAsset/RemeshNode.cpp b/Engine/Plugins/ChaosClothAssetEditor/Source/ChaosClothAssetDataflowNodes/Private/ChaosClothAsset/RemeshNode.cpp
--- a/Engine/Plugins/ChaosClothAssetEditor/Source/ChaosClothAssetDataflowNodes/Private/ChaosClothAsset/RemeshNode.cpp
+++ b/Engine/Plugins/ChaosClothAssetEditor/Source/ChaosClothAssetDataflowNodes/Private/ChaosClothAsset/RemeshNode.cpp
@@ -12,6 +12,7 @@
 #include "ChaosClothAsset/ClothDataflowTools.h"
 #include "ChaosClothAsset/SimulationSelfCollisionSpheresConfigNode.h"
 #include "ChaosClothAsset/SimulationLongRangeAttachmentConfigNode.h"
+#include "ChaosClothAsset/SimulationMaxDistanceConfigNode.h"
 #include "Dataflow/DataflowInputOutput.h"
 #include "DynamicMesh/MeshTangents.h"
 #include "DynamicMesh/MeshNormals.h"
@@ -1204,6 +1205,8 @@ void FChaosClothAssetRemeshNode::RebuildTopologyDependentSimData(const TSharedRe
 
 	FCollectionClothConstFacade InClothFacade(InClothCollection);
 	FCollectionClothFacade OutClothFacade(OutClothCollection);
+	FCollectionClothSelectionFacade OutSelectionFacade(OutClothCollection);
+	OutSelectionFacade.DefineSchema();
 
 	// Check that weight maps and skinning info have been interpolated over
 	for (const FName& InWeightMapName : InClothFacade.GetWeightMapNames())
@@ -1222,6 +1225,16 @@ void FChaosClothAssetRemeshNode::RebuildTopologyDependentSimData(const TSharedRe
 
 	Chaos::Softs::FCollectionPropertyConstFacade InProperties(InClothCollection);
 
+	// Reconstruct KinematicVertexSet
+	const FString MaxDistanceString = GET_MEMBER_NAME_STRING_CHECKED(FChaosClothAssetSimulationMaxDistanceConfigNode, MaxDistance);
+	const FName KinematicVertices3DName = GET_MEMBER_NAME_CHECKED(FChaosClothAssetSimulationMaxDistanceConfigNode, KinematicVertices3D);
+	if (InProperties.GetKeyIndex(MaxDistanceString) != INDEX_NONE)
+	{
+		const FName MaxDistanceMapName(InProperties.GetStringValue(MaxDistanceString, MaxDistanceString));
+		OutSelectionFacade.FindOrAddSelectionSet(KinematicVertices3DName, ClothCollectionGroup::SimVertices3D) =
+			FClothGeometryTools::GenerateKinematicVertices3D(OutClothCollection, MaxDistanceMapName, InProperties.GetWeightedFloatValue(MaxDistanceString, FVector2f(0.f, 1.f)), NAME_None);
+	}
+
 	// Reconstruct collision spheres
 	const FString SelfCollisionSphereStiffnessString = GET_MEMBER_NAME_STRING_CHECKED(FChaosClothAssetSimulationSelfCollisionSpheresConfigNode, SelfCollisionSphereStiffness);
 	const FString SelfCollisionSphereRadiusString = GET_MEMBER_NAME_STRING_CHECKED(FChaosClothAssetSimulationSelfCollisionSpheresConfigNode, SelfCollisionSphereRadius);
@@ -1240,27 +1253,35 @@ void FChaosClothAssetRemeshNode::RebuildTopologyDependentSimData(const TSharedRe
 			TSet<int32> VertexSet;
 			FClothGeometryTools::SampleVertices(SimPositions, CullDiameterSq, VertexSet);
 
-			FCollectionClothSelectionFacade Selection(OutClothCollection);
-			Selection.DefineSchema();
-
 			const FName SelectionSetName(*InProperties.GetStringValue(SelfCollisionSphereSetNameString, SelfCollisionSphereSetNameString));
-			Selection.FindOrAddSelectionSet(SelectionSetName, ClothCollectionGroup::SimVertices3D) = VertexSet;
+			OutSelectionFacade.FindOrAddSelectionSet(SelectionSetName, ClothCollectionGroup::SimVertices3D) = VertexSet;
 		}
 	}
 
 	// Reconstruct long-range attachments
-	const FString TetherStiffnessString = GET_MEMBER_NAME_STRING_CHECKED(FChaosClothAssetSimulationLongRangeAttachmentConfigNode, TetherStiffness);
-	const FString FixedEndWeightMapString = GET_MEMBER_NAME_STRING_CHECKED(FChaosClothAssetSimulationLongRangeAttachmentConfigNode, FixedEndWeightMap);
-	FString UseGeodesicTethersString = GET_MEMBER_NAME_STRING_CHECKED(FChaosClothAssetSimulationLongRangeAttachmentConfigNode, bUseGeodesicTethers);
-	UseGeodesicTethersString.RemoveFromStart(TEXT("b"), ESearchCase::CaseSensitive);  // Property collection names doesn't use the b prefix for booleans
 
-	if (InProperties.GetKeyIndex(TetherStiffnessString) != INDEX_NONE)
+	// v1 (weight map)
+	const FString FixedEndWeightMapString = GET_MEMBER_NAME_STRING_CHECKED(FChaosClothAssetSimulationLongRangeAttachmentConfigNode, FixedEndWeightMap);
+	if (InProperties.GetKeyIndex(FixedEndWeightMapString) != INDEX_NONE)
 	{
+		FString UseGeodesicTethersString = GET_MEMBER_NAME_STRING_CHECKED(FChaosClothAssetSimulationLongRangeAttachmentConfigNode, bUseGeodesicTethers);
+		UseGeodesicTethersString.RemoveFromStart(TEXT("b"), ESearchCase::CaseSensitive);  // Property collection names doesn't use the b prefix for booleans
 		const bool bUseGeodesicTethers = InProperties.GetValue<bool>(UseGeodesicTethersString);
 		const FName FixedEndWeightMap(InProperties.GetStringValue(FixedEndWeightMapString));
 
 		UE::Chaos::ClothAsset::FClothEngineTools::GenerateTethers(OutClothCollection, FixedEndWeightMap, bUseGeodesicTethers);
 	}
+
+	// v2 (vertex set)
+	const FString FixedEndSetString = GET_MEMBER_NAME_STRING_CHECKED(FChaosClothAssetSimulationLongRangeAttachmentConfigNode_v2, FixedEndSet);
+	if (InProperties.GetKeyIndex(FixedEndSetString) != INDEX_NONE)
+	{
+		// Regenerate using the Kinematic vertices.
+		FString UseGeodesicTethersString = GET_MEMBER_NAME_STRING_CHECKED(FChaosClothAssetSimulationLongRangeAttachmentConfigNode_v2, bUseGeodesicTethers);
+		UseGeodesicTethersString.RemoveFromStart(TEXT("b"), ESearchCase::CaseSensitive);  // Property collection names doesn't use the b prefix for booleans
+		const bool bUseGeodesicTethers = InProperties.GetValue<bool>(UseGeodesicTethersString);
+		UE::Chaos::ClothAsset::FClothEngineTools::GenerateTethersFromSelectionSet(OutClothCollection, KinematicVertices3DName, bUseGeodesicTethers);
+	}
 }
 
 
diff --git a/Engine/Plugins/ChaosClothAssetEditor/Source/ChaosClothAssetDataflowNodes/Private/ChaosClothAsset/SimulationDefaultConfigNode.cpp b/Engine/Plugins/ChaosClothAssetEditor/Source/ChaosClothAssetDataflowNodes/Private/ChaosClothAsset/SimulationDefaultConfigNode.cpp
--- a/Engine/Plugins/ChaosClothAssetEditor/Source/ChaosClothAssetDataflowNodes/Private/ChaosClothAsset/SimulationDefaultConfigNode.cpp
+++ b/Engine/Plugins/ChaosClothAssetEditor/Source/ChaosClothAssetDataflowNodes/Private/ChaosClothAsset/SimulationDefaultConfigNode.cpp
@@ -70,8 +70,9 @@ void FChaosClothAssetSimulationDefaultConfigNode::Evaluate(Dataflow::FContext& C
 				const bool bUseGeodesicTethers = Properties.GetValue<bool>(TEXT("UseGeodesicTethers"), bUseGeodesicTethersDefault);
 				// Use the "MaxDistance" weight map to generate tethers. This follows legacy behavior.
 				static const FName MaxDistanceName(TEXT("MaxDistance"));
+				const FVector2f MaxDistanceValue = Properties.GetWeightedFloatValue(TEXT("MaxDistance"), FVector2f(0.f, 1.f));
 
-				UE::Chaos::ClothAsset::FClothEngineTools::GenerateTethers(ClothCollection, MaxDistanceName, bUseGeodesicTethers);
+				UE::Chaos::ClothAsset::FClothEngineTools::GenerateTethers(ClothCollection, MaxDistanceName, bUseGeodesicTethers, MaxDistanceValue);
 			}
 			else
 			{
diff --git a/Engine/Plugins/ChaosClothAssetEditor/Source/ChaosClothAssetDataflowNodes/Private/ChaosClothAsset/SimulationLongRangeAttachmentConfigNode.cpp b/Engine/Plugins/ChaosClothAssetEditor/Source/ChaosClothAssetDataflowNodes/Private/ChaosClothAsset/SimulationLongRangeAttachmentConfigNode.cpp
--- a/Engine/Plugins/ChaosClothAssetEditor/Source/ChaosClothAssetDataflowNodes/Private/ChaosClothAsset/SimulationLongRangeAttachmentConfigNode.cpp
+++ b/Engine/Plugins/ChaosClothAssetEditor/Source/ChaosClothAssetDataflowNodes/Private/ChaosClothAsset/SimulationLongRangeAttachmentConfigNode.cpp
@@ -7,6 +7,35 @@
 
 #include UE_INLINE_GENERATED_CPP_BY_NAME(SimulationLongRangeAttachmentConfigNode)
 
+FChaosClothAssetSimulationLongRangeAttachmentConfigNode_v2::FChaosClothAssetSimulationLongRangeAttachmentConfigNode_v2(const Dataflow::FNodeParameters& InParam, FGuid InGuid)
+	: FChaosClothAssetSimulationBaseConfigNode(InParam, InGuid)
+{
+	RegisterCollectionConnections();
+	RegisterInputConnection(&FixedEndSet.StringValue, GET_MEMBER_NAME_CHECKED(FChaosClothAssetConnectableIStringValue, StringValue));
+	RegisterInputConnection(&TetherStiffness.WeightMap)
+		.SetCanHidePin(true)
+		.SetPinIsHidden(true);
+	RegisterInputConnection(&TetherScale.WeightMap)
+		.SetCanHidePin(true)
+		.SetPinIsHidden(true);
+}
+
+void FChaosClothAssetSimulationLongRangeAttachmentConfigNode_v2::AddProperties(FPropertyHelper& PropertyHelper) const
+{
+	PropertyHelper.SetPropertyWeighted(this, &TetherStiffness);
+	PropertyHelper.SetPropertyWeighted(this, &TetherScale);
+	PropertyHelper.SetPropertyBool(this, &bUseGeodesicTethers);
+	PropertyHelper.SetPropertyString(this, &FixedEndSet);
+}
+
+void FChaosClothAssetSimulationLongRangeAttachmentConfigNode_v2::EvaluateClothCollection(Dataflow::FContext& Context, const TSharedRef<FManagedArrayCollection>& ClothCollection) const
+{
+	const FString InFixedEndSetString = GetValue<FString>(Context, &FixedEndSet.StringValue);
+	const FName InFixedEndSet(InFixedEndSetString);
+	UE::Chaos::ClothAsset::FClothEngineTools::GenerateTethersFromSelectionSet(ClothCollection, InFixedEndSet, bUseGeodesicTethers);
+}
+
+
 FChaosClothAssetSimulationLongRangeAttachmentConfigNode::FChaosClothAssetSimulationLongRangeAttachmentConfigNode(const Dataflow::FNodeParameters& InParam, FGuid InGuid)
 	: FChaosClothAssetSimulationBaseConfigNode(InParam, InGuid)
 {
diff --git a/Engine/Plugins/ChaosClothAssetEditor/Source/ChaosClothAssetDataflowNodes/Private/ChaosClothAsset/SimulationMaxDistanceConfigNode.cpp b/Engine/Plugins/ChaosClothAssetEditor/Source/ChaosClothAssetDataflowNodes/Private/ChaosClothAsset/SimulationMaxDistanceConfigNode.cpp
--- a/Engine/Plugins/ChaosClothAssetEditor/Source/ChaosClothAssetDataflowNodes/Private/ChaosClothAsset/SimulationMaxDistanceConfigNode.cpp
+++ b/Engine/Plugins/ChaosClothAssetEditor/Source/ChaosClothAssetDataflowNodes/Private/ChaosClothAsset/SimulationMaxDistanceConfigNode.cpp
@@ -1,6 +1,9 @@
 // Copyright Epic Games, Inc. All Rights Reserved.
 
 #include "ChaosClothAsset/SimulationMaxDistanceConfigNode.h"
+#include "ChaosClothAsset/ClothCollectionGroup.h"
+#include "ChaosClothAsset/ClothGeometryTools.h"
+#include "ChaosClothAsset/CollectionClothSelectionFacade.h"
 #include "Chaos/CollectionPropertyFacade.h"
 
 #include UE_INLINE_GENERATED_CPP_BY_NAME(SimulationMaxDistanceConfigNode)
@@ -8,13 +11,41 @@
 FChaosClothAssetSimulationMaxDistanceConfigNode::FChaosClothAssetSimulationMaxDistanceConfigNode(const Dataflow::FNodeParameters& InParam, FGuid InGuid)
 	: FChaosClothAssetSimulationBaseConfigNode(InParam, InGuid)
 {
+	KinematicVertices3D = TEXT("KinematicVertices3D");
 	RegisterCollectionConnections();
 	RegisterInputConnection(&MaxDistance.WeightMap)
 		.SetCanHidePin(true)
 		.SetPinIsHidden(true);
+	RegisterInputConnection(&InKinematic, GET_MEMBER_NAME_CHECKED(FChaosClothAssetConnectableIStringValue, StringValue));
+	RegisterOutputConnection(&KinematicVertices3D);
+}
+
+void FChaosClothAssetSimulationMaxDistanceConfigNode::Evaluate(Dataflow::FContext& Context, const FDataflowOutput* Out) const
+{
+	Super::Evaluate(Context, Out);
+
+	if (Out->IsA(&KinematicVertices3D))
+	{
+		SetValue(Context, KinematicVertices3D, &KinematicVertices3D);
+	}
 }
 
 void FChaosClothAssetSimulationMaxDistanceConfigNode::AddProperties(FPropertyHelper& PropertyHelper) const
 {
 	PropertyHelper.SetPropertyWeighted(this, &MaxDistance, {}, ECollectionPropertyFlags::Intrinsic);  // Intrinsic since the deformer weights needs to be recalculated
+	PropertyHelper.SetPropertyString(this, &KinematicVertices3D);
 }
+
+void FChaosClothAssetSimulationMaxDistanceConfigNode::EvaluateClothCollection(Dataflow::FContext& Context, const TSharedRef<FManagedArrayCollection>& ClothCollection) const
+{
+	using namespace UE::Chaos::ClothAsset;
+	const FName MaxDistanceString(GetValue<FString>(Context, &MaxDistance.WeightMap)); //  Override for this is already set by AddProperties
+	const FName InputKinematicString(GetValue<FString>(Context, &InKinematic.StringValue));
+	InKinematic.StringValue_Override = GetValue<FString>(Context, &InKinematic.StringValue, FWeightMapTools::NotOverridden);
+
+	FCollectionClothSelectionFacade SelectionFacade(ClothCollection);
+	SelectionFacade.DefineSchema();
+	SelectionFacade.FindOrAddSelectionSet(FName(KinematicVertices3D), ClothCollectionGroup::SimVertices3D) =
+		FClothGeometryTools::GenerateKinematicVertices3D(ClothCollection, MaxDistanceString, FVector2f(MaxDistance.Low, MaxDistance.High), InputKinematicString);
+
+}
\ No newline at end of file
diff --git a/Engine/Plugins/ChaosClothAssetEditor/Source/ChaosClothAssetDataflowNodes/Public/ChaosClothAsset/SimulationLongRangeAttachmentConfigNode.h b/Engine/Plugins/ChaosClothAssetEditor/Source/ChaosClothAssetDataflowNodes/Public/ChaosClothAsset/SimulationLongRangeAttachmentConfigNode.h
--- a/Engine/Plugins/ChaosClothAssetEditor/Source/ChaosClothAssetDataflowNodes/Public/ChaosClothAsset/SimulationLongRangeAttachmentConfigNode.h
+++ b/Engine/Plugins/ChaosClothAssetEditor/Source/ChaosClothAssetDataflowNodes/Public/ChaosClothAsset/SimulationLongRangeAttachmentConfigNode.h
@@ -8,6 +8,56 @@
 
 /** Long range attachment constraint property configuration node. */
 USTRUCT(Meta = (DataflowCloth))
+struct FChaosClothAssetSimulationLongRangeAttachmentConfigNode_v2 : public FChaosClothAssetSimulationBaseConfigNode
+{
+	GENERATED_USTRUCT_BODY()
+	DATAFLOW_NODE_DEFINE_INTERNAL(FChaosClothAssetSimulationLongRangeAttachmentConfigNode_v2, "SimulationLongRangeAttachmentConfig", "Cloth", "Cloth Simulation Long Range Attachment Config")
+
+public:
+
+	/**
+	 * The tethers' stiffness of the long range attachment constraints.
+	 * The long range attachment connects each of the cloth particles to its closest fixed point with a spring constraint.
+	 * This can be used to compensate for a lack of stretch resistance when the iterations count is kept low for performance reasons.
+	 * Can lead to an unnatural pull string puppet like behavior.
+	 * If a valid weight map is found with the given Weight Map name, then both Low and High values
+	 * are interpolated with the per particle weight to make the final value used for the simulation.
+	 * Otherwise all particles are considered to have a zero weight, and only the Low value is meaningful.
+	 */
+	UPROPERTY(EditAnywhere, Category = "Long Range Attachment Properties", Meta = (UIMin = "0", UIMax = "1", ClampMin = "0", ClampMax = "1"))
+	FChaosClothAssetWeightedValue TetherStiffness = { true, 1.f, 1.f, TEXT("TetherStiffness") };
+
+	/**
+	 * The limit scale of the long range attachment constraints (aka tether limit).
+	 * If a valid weight map is found with the given Weight Map name, then both Low and High values
+	 * are interpolated with the per particle weight to make the final value used for the simulation.
+	 * Otherwise all particles are considered to have a zero weight, and only the Low value is meaningful.
+	*/
+	UPROPERTY(EditAnywhere, Category = "Long Range Attachment Properties", Meta = (UIMin = "1.", UIMax = "1.1", ClampMin = "0.01", ClampMax = "10"))
+	FChaosClothAssetWeightedValue TetherScale = { true, 1.f, 1.f, TEXT("TetherScale") };
+
+	/**
+	 * Use geodesic instead of euclidean distance calculations for the Long Range Attachment constraint,
+	 * which is slower at setup but more accurate at establishing the correct position and length of the tethers,
+	 * and therefore is less prone to artifacts during the simulation.
+	 */
+	UPROPERTY(EditAnywhere, Category = "Long Range Attachment Properties")
+	bool bUseGeodesicTethers = true;
+
+	/** The name of the vertex selection set used as fixed tether ends. */
+	UPROPERTY(EditAnywhere, Category = "Long Range Attachment Properties")
+	FChaosClothAssetConnectableIStringValue FixedEndSet = { TEXT("KinematicVertices3D") };
+
+	FChaosClothAssetSimulationLongRangeAttachmentConfigNode_v2(const Dataflow::FNodeParameters& InParam, FGuid InGuid = FGuid::NewGuid());
+
+private:
+	virtual void AddProperties(FPropertyHelper& PropertyHelper) const override;
+
+	virtual void EvaluateClothCollection(Dataflow::FContext& Context, const TSharedRef<FManagedArrayCollection>& ClothCollection) const override;
+};
+
+/** Long range attachment constraint property configuration node. */
+USTRUCT(Meta = (DataflowCloth, Deprecated = "5.5"))
 struct FChaosClothAssetSimulationLongRangeAttachmentConfigNode : public FChaosClothAssetSimulationBaseConfigNode
 {
 	GENERATED_USTRUCT_BODY()
diff --git a/Engine/Plugins/ChaosClothAssetEditor/Source/ChaosClothAssetDataflowNodes/Public/ChaosClothAsset/SimulationMaxDistanceConfigNode.h b/Engine/Plugins/ChaosClothAssetEditor/Source/ChaosClothAssetDataflowNodes/Public/ChaosClothAsset/SimulationMaxDistanceConfigNode.h
--- a/Engine/Plugins/ChaosClothAssetEditor/Source/ChaosClothAssetDataflowNodes/Public/ChaosClothAsset/SimulationMaxDistanceConfigNode.h
+++ b/Engine/Plugins/ChaosClothAssetEditor/Source/ChaosClothAssetDataflowNodes/Public/ChaosClothAsset/SimulationMaxDistanceConfigNode.h
@@ -24,8 +24,25 @@ public:
 	UPROPERTY(EditAnywhere, Category = "MaxDistance Properties", Meta = (UIMin = "0", UIMax = "100", ClampMin = "0"))
 	FChaosClothAssetWeightedValue MaxDistance = { true, 0.f, 100.f, TEXT("MaxDistance") };
 
+	/**
+	 * Selection set of SimVertices3D that should be made kinematic regardless of the MaxDistance map value.
+	 */
+	UPROPERTY(VisibleAnywhere, Category = "MaxDistance Properties")
+	FChaosClothAssetConnectableIStringValue InKinematic = { TEXT("") };
+
+	/**
+	 * The name of the kinematic vertex set generated by this node.
+	 * This is the union of InKinematic and any vertices which are below the solver max distance threshold.
+	 * The name of this set cannot be changed and is only provided for further tweaking.
+	 */
+	UPROPERTY(VisibleAnywhere, Category = "MaxDistance Properties", Meta = (DataflowOutput))
+	FString KinematicVertices3D;
+
 	FChaosClothAssetSimulationMaxDistanceConfigNode(const Dataflow::FNodeParameters& InParam, FGuid InGuid = FGuid::NewGuid());
 
 private:
 	virtual void AddProperties(FPropertyHelper& PropertyHelper) const override;
+	virtual void EvaluateClothCollection(Dataflow::FContext& Context, const TSharedRef<FManagedArrayCollection>& ClothCollection) const override;
+
+	virtual void Evaluate(Dataflow::FContext& Context, const FDataflowOutput* Out) const override;
 };
diff --git a/Engine/Source/Runtime/ClothingSystemRuntimeCommon/Private/ClothTetherData.cpp b/Engine/Source/Runtime/ClothingSystemRuntimeCommon/Private/ClothTetherData.cpp
--- a/Engine/Source/Runtime/ClothingSystemRuntimeCommon/Private/ClothTetherData.cpp
+++ b/Engine/Source/Runtime/ClothingSystemRuntimeCommon/Private/ClothTetherData.cpp
@@ -15,7 +15,7 @@ public:
 	FClothTetherDataPrivate(
 		const TConstArrayView<FVector3f>& Points,
 		const TConstArrayView<uint32>& Indices,
-		const TConstArrayView<float>& MaxDistances,
+		TFunctionRef<bool(int32 Index)> IsKinematic,
 		bool bUseGeodesicDistance);
 
 	FClothTetherDataPrivate(
@@ -34,7 +34,7 @@ private:
 
 	// Generate a map of tethers by following the triangle mesh network
 	// The choice of the tether is determined by the shortest geodesic (curvature) distance.
-	void GenerateGeodesicTethers(const TConstArrayView<FVector3f>& Points, const TConstArrayView<float>& MaxDistances);
+	void GenerateGeodesicTethers(const TConstArrayView<FVector3f>& Points, TFunctionRef<bool(int32 Index)> IsKinematic);
 
 	// Update TetherNums after computing tethers.
 	void UpdateCounts();
@@ -62,7 +62,31 @@ void FClothTetherData::GenerateTethers(
 	}
 
 	// Calculate the tethers
-	const FClothTetherDataPrivate ClothTetherData(Points, Indices, MaxDistances, bUseGeodesicDistance);
+	const FClothTetherDataPrivate ClothTetherData(Points, Indices, [&MaxDistances](int32 Node)
+	{
+		return MaxDistances[Node] < FClothTetherDataPrivate::KinematicDistanceThreshold;
+	}, bUseGeodesicDistance);
+	ClothTetherData.GetBatchedTetherData(Tethers);
+}
+
+void FClothTetherData::GenerateTethers(
+	const TConstArrayView<FVector3f>& Points,
+	const TConstArrayView<uint32>& Indices,
+	const TSet<int32>& KinematicNodes,
+	bool bUseGeodesicDistance)
+{
+	// Early exit if there is no KinematicNodes set
+	if (KinematicNodes.IsEmpty())
+	{
+		Tethers.Empty();
+		return;
+	}
+
+	// Calculate the tethers
+	const FClothTetherDataPrivate ClothTetherData(Points, Indices, [&KinematicNodes](int32 Node)
+	{
+		return KinematicNodes.Contains(Node);
+	}, bUseGeodesicDistance);
 	ClothTetherData.GetBatchedTetherData(Tethers);
 }
 
@@ -91,7 +115,7 @@ bool FClothTetherData::Serialize(FArchive& Ar)
 FClothTetherDataPrivate::FClothTetherDataPrivate(
 	const TConstArrayView<FVector3f>& Points,
 	const TConstArrayView<uint32>& Indices,
-	const TConstArrayView<float>& MaxDistances,
+	TFunctionRef<bool(int32 Index)> IsKinematic,
 	bool bUseGeodesicDistance)
 {
 	// Calculate the points' neighbors map
@@ -113,17 +137,17 @@ FClothTetherDataPrivate::FClothTetherDataPrivate(
 	NodeToNeighbors.GenerateKeyArray(Nodes);
 
 	// Find all kinematic points to use as anchor points for the tethers and seed the path finding
-	DynamicNodes.Reserve(2 * (Nodes.Num() / 3));  // Start at 66% of the number of nodes to minimize this array's reallocations
-	KinematicNodes.Reserve(Nodes.Num() / 3);  // Start at 33% of the number of nodes to minimize this array's reallocations
+	DynamicNodes.Reserve(Points.Num());
+	KinematicNodes.Reserve(Points.Num());
 	for (const int32 Node : Nodes)
 	{
-		if (MaxDistances[Node] >= KinematicDistanceThreshold)
+		if (IsKinematic(Node))
 		{
-			DynamicNodes.Add(Node);
+			KinematicNodes.Add(Node);
 		}
 		else
 		{
-			KinematicNodes.Add(Node);
+			DynamicNodes.Add(Node);
 		}
 	}
 
@@ -147,7 +171,7 @@ FClothTetherDataPrivate::FClothTetherDataPrivate(
 		}
 		else
 		{
-			GenerateGeodesicTethers(Points, MaxDistances);
+			GenerateGeodesicTethers(Points, IsKinematic);
 		}
 
 		UpdateCounts();
@@ -350,7 +374,7 @@ void FClothTetherDataPrivate::GenerateEuclideanTethers(const TConstArrayView<FVe
 	});
 }
 
-void FClothTetherDataPrivate::GenerateGeodesicTethers(const TConstArrayView<FVector3f>& Points, const TConstArrayView<float>& MaxDistances)
+void FClothTetherDataPrivate::GenerateGeodesicTethers(const TConstArrayView<FVector3f>& Points, TFunctionRef<bool(int32 Index)> IsKinematic)
 {
 	check(KinematicNodeIslands.Num());
 
@@ -369,7 +393,7 @@ void FClothTetherDataPrivate::GenerateGeodesicTethers(const TConstArrayView<FVec
 			const TSet<int32>& Neighbors = NodeToNeighbors[KinematicNode];
 			for (const int32 Neighbor : Neighbors)
 			{
-				if (MaxDistances[Neighbor] >= KinematicDistanceThreshold)
+				if (!IsKinematic(Neighbor))
 				{
 					SeedIslands[IslandIndex].Add(KinematicNode);
 					++NumSeeds;
@@ -405,7 +429,7 @@ void FClothTetherDataPrivate::GenerateGeodesicTethers(const TConstArrayView<FVec
 		}
 	}
 
-	ParallelFor(Seeds.Num(), [this, &Seeds, &GeodesicDistances, &Points, &MaxDistances](int32 Index)
+	ParallelFor(Seeds.Num(), [this, &Seeds, &GeodesicDistances, &Points, &IsKinematic](int32 Index)
 	{
 		const int32 Seed = Seeds[Index];
 		TMap<int32, float>& SeedGeodesicDistances = GeodesicDistances[Seed];
@@ -448,7 +472,7 @@ void FClothTetherDataPrivate::GenerateGeodesicTethers(const TConstArrayView<FVec
 				check(NeighborNode != ParentNode);
 
 				// Do not progress onto kinematic nodes
-				if (MaxDistances[NeighborNode] < KinematicDistanceThreshold)
+				if (IsKinematic(NeighborNode))
 				{
 					continue;
 				}
diff --git a/Engine/Source/Runtime/ClothingSystemRuntimeCommon/Public/ClothTetherData.h b/Engine/Source/Runtime/ClothingSystemRuntimeCommon/Public/ClothTetherData.h
--- a/Engine/Source/Runtime/ClothingSystemRuntimeCommon/Public/ClothTetherData.h
+++ b/Engine/Source/Runtime/ClothingSystemRuntimeCommon/Public/ClothTetherData.h
@@ -34,6 +34,16 @@ struct FClothTetherData
 		const TConstArrayView<uint32>& Indices,  // Triangle mesh
 		const TConstArrayView<float>& MaxDistances,  // Mask for sorting the kinematic from the dynamic points
 		bool bUseGeodesicDistance);  // Whether to use geodesic (walking along the sruface) or euclidean (beeline) distances to find the tethers.
+	
+	/**
+	 * Generate the tethers by following the triangle mesh network from the closest kinematic to each dynamic point.
+	 * Inside array items can be processed concurrently, but the outside array must be iterated on sequentially.
+	 */
+	CLOTHINGSYSTEMRUNTIMECOMMON_API void GenerateTethers(
+		const TConstArrayView<FVector3f>& Points,  // Reference pose
+		const TConstArrayView<uint32>& Indices,  // Triangle mesh
+		const TSet<int32>& KinematicNodes,  // kinematic indices
+		bool bUseGeodesicDistance);  // Whether to use geodesic (walking along the sruface) or euclidean (beeline) distances to find the tethers.
 
 	/**
 	 * Generate the tether batches from existing tether data.
