Add NotifyFinishedChangingProperties to ToggleEditConditionState. This ensures a NotifyFinishedChangingProperties after toggling an InlineEditConditionToggle property.

diff --git a/Engine/Plugins/ChaosClothAssetEditor/Source/ChaosClothAssetDataflowNodes/Public/ChaosClothAsset/SimulationVelocityScaleConfigNode.h b/Engine/Plugins/ChaosClothAssetEditor/Source/ChaosClothAssetDataflowNodes/Public/ChaosClothAsset/SimulationVelocityScaleConfigNode.h
--- a/Engine/Plugins/ChaosClothAssetEditor/Source/ChaosClothAssetDataflowNodes/Public/ChaosClothAsset/SimulationVelocityScaleConfigNode.h
+++ b/Engine/Plugins/ChaosClothAssetEditor/Source/ChaosClothAssetDataflowNodes/Public/ChaosClothAsset/SimulationVelocityScaleConfigNode.h
@@ -26,7 +26,7 @@ public:
 	/**
 	 * Enable linear velocity clamping.
 	 */
-	UPROPERTY(EditAnywhere, Category = "Velocity Scale Properties")
+	UPROPERTY(EditAnywhere, Category = "Velocity Scale Properties", meta = (InlineEditConditionToggle))
 	bool bEnableLinearVelocityClamping = false;
 
 	/**
@@ -39,7 +39,7 @@ public:
 	/**
 	 * Enable linear acceleration clamping.
 	 */
-	UPROPERTY(EditAnywhere, Category = "Velocity Scale Properties")
+	UPROPERTY(EditAnywhere, Category = "Velocity Scale Properties", meta = (InlineEditConditionToggle))
 	bool bEnableLinearAccelerationClamping = false;
 
 	/**
@@ -60,7 +60,7 @@ public:
 	/**
 	 * Enable angular velocity clamping.
 	 */
-	UPROPERTY(EditAnywhere, Category = "Velocity Scale Properties")
+	UPROPERTY(EditAnywhere, Category = "Velocity Scale Properties", meta = (InlineEditConditionToggle))
 	bool bEnableAngularVelocityClamping = false;
 
 	/**
@@ -73,7 +73,7 @@ public:
 	/**
 	 * Enable angular acceleration clamping.
 	 */
-	UPROPERTY(EditAnywhere, Category = "Velocity Scale Properties")
+	UPROPERTY(EditAnywhere, Category = "Velocity Scale Properties", meta = (InlineEditConditionToggle))
 	bool bEnableAngularAccelerationClamping = false;
 
 	/**
diff --git a/Engine/Plugins/Experimental/Dataflow/Source/DataflowEditor/Private/Dataflow/DataflowEditorCommands.cpp b/Engine/Plugins/Experimental/Dataflow/Source/DataflowEditor/Private/Dataflow/DataflowEditorCommands.cpp
--- a/Engine/Plugins/Experimental/Dataflow/Source/DataflowEditor/Private/Dataflow/DataflowEditorCommands.cpp
+++ b/Engine/Plugins/Experimental/Dataflow/Source/DataflowEditor/Private/Dataflow/DataflowEditorCommands.cpp
@@ -357,6 +357,7 @@ void FDataflowEditorCommands::OnPropertyValueChanged(UDataflow* OutDataflow, TSh
 	case EPropertyChangeType::ArrayClear:
 	case EPropertyChangeType::ArrayMove:
 	case EPropertyChangeType::Duplicate:
+	case EPropertyChangeType::ToggleEditable:
 		if (ensure(OutDataflow && InPropertyChangedEvent.Property && InPropertyChangedEvent.Property->GetOwnerUObject()))
 		{
 			OutDataflow->Modify();  // Modify must be called even if SelectedNodes is empty because comment nodes aren't part of the selection set but still have properties
diff --git a/Engine/Source/Editor/PropertyEditor/Private/Presentation/PropertyEditor/PropertyEditor.cpp b/Engine/Source/Editor/PropertyEditor/Private/Presentation/PropertyEditor/PropertyEditor.cpp
--- a/Engine/Source/Editor/PropertyEditor/Private/Presentation/PropertyEditor/PropertyEditor.cpp
+++ b/Engine/Source/Editor/PropertyEditor/Private/Presentation/PropertyEditor/PropertyEditor.cpp
@@ -566,6 +566,7 @@ void FPropertyEditor::ToggleEditConditionState()
 	FPropertyChangedEvent ChangeEvent(PropertyNode->GetProperty(), EPropertyChangeType::ToggleEditable);
 	ChangeEvent.SetArrayIndexPerObject(ArrayIndicesPerObject);
 	PropertyNode->NotifyPostChange( ChangeEvent, PropertyUtilities->GetNotifyHook() );
+	PropertyUtilities->NotifyFinishedChangingProperties(ChangeEvent);
 }
 
 void FPropertyEditor::OnGetClassesForAssetPicker( TArray<const UClass*>& OutClasses )
