Cloth Asset - Deprecated all mutables on struct customizations, and replaced the need to keep them updated with an empty context evaluation in the customization code.

#rnx
#tests Connected/disconnected name properties to various nodes while observing the name property being overriden in the UI.

diff --git a/Engine/Plugins/ChaosClothAssetEditor/Source/ChaosClothAssetDataflowNodes/Private/ChaosClothAsset/AddStitchNode.cpp b/Engine/Plugins/ChaosClothAssetEditor/Source/ChaosClothAssetDataflowNodes/Private/ChaosClothAsset/AddStitchNode.cpp
--- a/Engine/Plugins/ChaosClothAssetEditor/Source/ChaosClothAssetDataflowNodes/Private/ChaosClothAsset/AddStitchNode.cpp
+++ b/Engine/Plugins/ChaosClothAssetEditor/Source/ChaosClothAssetDataflowNodes/Private/ChaosClothAsset/AddStitchNode.cpp
@@ -34,7 +34,6 @@ void FChaosClothAssetAddStitchNode::Evaluate(Dataflow::FContext& Context, const
 		FCollectionClothFacade Cloth(ClothCollection);
 		FCollectionClothSelectionConstFacade SelectionFacade(ClothCollection);
 		const FName InSelectionName(*GetValue<FString>(Context, &MergeToSingleVertexSelection.StringValue));
-		MergeToSingleVertexSelection.StringValue_Override = GetValue<FString>(Context, &MergeToSingleVertexSelection.StringValue, UE::Chaos::ClothAsset::FWeightMapTools::NotOverridden);
 
 		if (Cloth.IsValid() && SelectionFacade.IsValid() && InSelectionName != NAME_None)
 		{
diff --git a/Engine/Plugins/ChaosClothAssetEditor/Source/ChaosClothAssetDataflowNodes/Private/ChaosClothAsset/AddWeightMapNode.cpp b/Engine/Plugins/ChaosClothAssetEditor/Source/ChaosClothAssetDataflowNodes/Private/ChaosClothAsset/AddWeightMapNode.cpp
--- a/Engine/Plugins/ChaosClothAssetEditor/Source/ChaosClothAssetDataflowNodes/Private/ChaosClothAsset/AddWeightMapNode.cpp
+++ b/Engine/Plugins/ChaosClothAssetEditor/Source/ChaosClothAssetDataflowNodes/Private/ChaosClothAsset/AddWeightMapNode.cpp
@@ -258,7 +258,6 @@ void FChaosClothAssetAddWeightMapNode::Evaluate(Dataflow::FContext& Context, con
 	{
 		// Evaluate InputName
 		const FName InInputName = GetInputName(Context);
-		InputName.StringValue_Override = GetValue<FString>(Context, &InputName.StringValue, UE::Chaos::ClothAsset::FWeightMapTools::NotOverridden);
 
 		// Evaluate in collection
 		FManagedArrayCollection InCollection = GetValue<FManagedArrayCollection>(Context, &Collection);
diff --git a/Engine/Plugins/ChaosClothAssetEditor/Source/ChaosClothAssetDataflowNodes/Private/ChaosClothAsset/AttributeNode.cpp b/Engine/Plugins/ChaosClothAssetEditor/Source/ChaosClothAssetDataflowNodes/Private/ChaosClothAsset/AttributeNode.cpp
--- a/Engine/Plugins/ChaosClothAssetEditor/Source/ChaosClothAssetDataflowNodes/Private/ChaosClothAsset/AttributeNode.cpp
+++ b/Engine/Plugins/ChaosClothAssetEditor/Source/ChaosClothAssetDataflowNodes/Private/ChaosClothAsset/AttributeNode.cpp
@@ -22,8 +22,6 @@ FChaosClothAssetAttributeNode_V2::FChaosClothAssetAttributeNode_V2(const Dataflo
 
 void FChaosClothAssetAttributeNode_V2::Evaluate(Dataflow::FContext& Context, const FDataflowOutput* Out) const
 {
-	Name.StringValue_Override = GetValue<FString>(Context, &Name.StringValue, UE::Chaos::ClothAsset::FWeightMapTools::NotOverridden);
-
 	if (Out->IsA<FManagedArrayCollection>(&Collection))
 	{
 		using namespace UE::Chaos::ClothAsset;
diff --git a/Engine/Plugins/ChaosClothAssetEditor/Source/ChaosClothAssetDataflowNodes/Private/ChaosClothAsset/ClothDataflowTools.cpp b/Engine/Plugins/ChaosClothAssetEditor/Source/ChaosClothAssetDataflowNodes/Private/ChaosClothAsset/ClothDataflowTools.cpp
--- a/Engine/Plugins/ChaosClothAssetEditor/Source/ChaosClothAssetDataflowNodes/Private/ChaosClothAsset/ClothDataflowTools.cpp
+++ b/Engine/Plugins/ChaosClothAssetEditor/Source/ChaosClothAssetDataflowNodes/Private/ChaosClothAsset/ClothDataflowTools.cpp
@@ -12,6 +12,7 @@
 #include "Widgets/Notifications/SNotificationList.h"
 #include "MeshDescriptionToDynamicMesh.h"
 #include "MeshUtilities.h"
+#include "PropertyHandle.h"
 #include "ReferenceSkeleton.h"
 #include "SkeletalMeshAttributes.h"
 #include "ToDynamicMesh.h"
@@ -539,4 +540,23 @@ namespace UE::Chaos::ClothAsset
 		}
 		return true;
 	}
+
+	FDataflowNode* FClothDataflowTools::GetPropertyOwnerDataflowNode(const TSharedPtr<IPropertyHandle>& PropertyHandle, const UStruct* DataflowNodeStruct)
+	{
+		for (TSharedPtr<IPropertyHandle> OwnerHandle = PropertyHandle->GetParentHandle(); OwnerHandle; OwnerHandle = OwnerHandle->GetParentHandle())
+		{
+			if (const TSharedPtr<IPropertyHandleStruct> OwnerHandleStruct = OwnerHandle->AsStruct())
+			{
+				if (TSharedPtr<FStructOnScope> StructOnScope = OwnerHandleStruct->GetStructData())
+				{
+					if (StructOnScope->GetStruct()->IsChildOf(DataflowNodeStruct))
+					{
+						return reinterpret_cast<FDataflowNode*>(StructOnScope->GetStructMemory());
+					}
+				}
+			}
+		}
+		return nullptr;
+	}
+
 }  // End namespace UE::Chaos::ClothAsset
diff --git a/Engine/Plugins/ChaosClothAssetEditor/Source/ChaosClothAssetDataflowNodes/Private/ChaosClothAsset/ClothDataflowTools.h b/Engine/Plugins/ChaosClothAssetEditor/Source/ChaosClothAssetDataflowNodes/Private/ChaosClothAsset/ClothDataflowTools.h
--- a/Engine/Plugins/ChaosClothAssetEditor/Source/ChaosClothAssetDataflowNodes/Private/ChaosClothAsset/ClothDataflowTools.h
+++ b/Engine/Plugins/ChaosClothAssetEditor/Source/ChaosClothAssetDataflowNodes/Private/ChaosClothAsset/ClothDataflowTools.h
@@ -11,6 +11,7 @@ struct FMeshDescription;
 struct FDataflowNode;
 class FSkeletalMeshLODModel;
 class FString;
+class IPropertyHandle;
 
 namespace UE::Chaos::ClothAsset
 {
@@ -34,6 +35,17 @@ namespace UE::Chaos::ClothAsset
 		static void MakeCollectionName(FString& InOutString);
 
 		static bool BuildSkeletalMeshModelFromMeshDescription(const FMeshDescription* const InMeshDescription, const FMeshBuildSettings& InBuildSettings, FSkeletalMeshLODModel& SkeletalMeshModel);
+
+		/** Return the Dataflow node owning by this property, and cast it to the desired node type. */
+		template<typename T = FDataflowNode>
+		static T* GetPropertyOwnerDataflowNode(const TSharedPtr<IPropertyHandle>& PropertyHandle)
+		{
+			return static_cast<FDataflowNode*>(GetPropertyOwnerDataflowNode(PropertyHandle, T::StaticStruct()));
+		}
+
+	private:
+		/** Return the Dataflow node owning by this property. */
+		static FDataflowNode* GetPropertyOwnerDataflowNode(const TSharedPtr<IPropertyHandle>& PropertyHandle, const UStruct* DataflowNodeStruct);
 	};
 }  // End namespace UE::Chaos::ClothAsset
 
diff --git a/Engine/Plugins/ChaosClothAssetEditor/Source/ChaosClothAssetDataflowNodes/Private/ChaosClothAsset/ConnectableValueCustomization.cpp b/Engine/Plugins/ChaosClothAssetEditor/Source/ChaosClothAssetDataflowNodes/Private/ChaosClothAsset/ConnectableValueCustomization.cpp
--- a/Engine/Plugins/ChaosClothAssetEditor/Source/ChaosClothAssetDataflowNodes/Private/ChaosClothAsset/ConnectableValueCustomization.cpp
+++ b/Engine/Plugins/ChaosClothAssetEditor/Source/ChaosClothAssetDataflowNodes/Private/ChaosClothAsset/ConnectableValueCustomization.cpp
@@ -1,15 +1,18 @@
 // Copyright Epic Games, Inc. All Rights Reserved.
 
 #include "ChaosClothAsset/ConnectableValueCustomization.h"
-#include "ChaosClothAsset/ImportedValueCustomization.h"
 #include "ChaosClothAsset/ClothAssetEditorStyle.h"
+#include "ChaosClothAsset/ClothDataflowTools.h"
+#include "ChaosClothAsset/ImportedValueCustomization.h"
 #include "ChaosClothAsset/WeightedValue.h"
+#include "Dataflow/DataflowNode.h"
+#include "Dataflow/DataflowNodeParameters.h"
 #include "Widgets/Input/SCheckBox.h"
 #include "Widgets/Input/SNumericEntryBox.h"
 #include "Widgets/Input/SEditableTextBox.h"
 #include "DetailLayoutBuilder.h"
 #include "DetailWidgetRow.h"
-#include "Editor.h"
+//#include "Editor.h"
 
 #define LOCTEXT_NAMESPACE "ChaosClothAssetWeightedValueCustomization"
 
@@ -17,23 +20,27 @@ namespace UE::Chaos::ClothAsset
 {
 	namespace Private
 	{
-		static const FString OverridePrefix = TEXT("_Override");
+		static const FString OverridePrefix = TEXT("_Override");  // UE_DEPRECATED(5.5, "Override properties are no longer used.")
 		static const FString BuildFabricMaps = TEXT("BuildFabricMaps");
 		static const FString CouldUseFabrics = TEXT("CouldUseFabrics");
-
-		
 	}
 	
+	// UE_DEPRECATED(5.5, "Override properties are no longer used.")
 	bool FConnectableValueCustomization::IsOverrideProperty(const TSharedPtr<IPropertyHandle>& Property)
 	{
+PRAGMA_DISABLE_DEPRECATION_WARNINGS
 		const FStringView PropertyPath = Property ? Property->GetPropertyPath() : FStringView();
 		return PropertyPath.EndsWith(Private::OverridePrefix, ESearchCase::CaseSensitive);
+PRAGMA_ENABLE_DEPRECATION_WARNINGS
 	}
+	// UE_DEPRECATED(5.5, "Override properties are no longer used.")
 	bool FConnectableValueCustomization::IsOverridePropertyOf(const TSharedPtr<IPropertyHandle>& OverrideProperty, const TSharedPtr<IPropertyHandle>& Property)
 	{
+PRAGMA_DISABLE_DEPRECATION_WARNINGS
 		const FStringView OverridePropertyPath = OverrideProperty ? OverrideProperty->GetPropertyPath() : FStringView();
 		const FStringView PropertyPath = Property ? Property->GetPropertyPath() : FStringView();
 		return OverridePropertyPath == FString(PropertyPath) + Private::OverridePrefix;
+PRAGMA_ENABLE_DEPRECATION_WARNINGS
 	}
 	bool FConnectableValueCustomization::BuildFabricMapsProperty(const TSharedPtr<IPropertyHandle>& Property)
 	{
@@ -108,10 +115,12 @@ namespace UE::Chaos::ClothAsset
 		{
 			TSharedRef<IPropertyHandle> ChildHandle = SortedChildHandles[ChildIndex];
 
+PRAGMA_DISABLE_DEPRECATION_WARNINGS
 			if (IsOverrideProperty(ChildHandle))
 			{
 				continue;  // Skip overrides
 			}
+PRAGMA_ENABLE_DEPRECATION_WARNINGS
 
 			const bool bLastChild = SortedChildHandles.Num() - 1 == ChildIndex;
 
@@ -135,42 +144,33 @@ namespace UE::Chaos::ClothAsset
 		
 		if (PropertyClass == FStrProperty::StaticClass())
 		{
-			// Manage override property values (properties ending with _Override)
-			TWeakPtr<IPropertyHandle> OverrideHandleWeakPtr;
-
-			for (int32 ChildIndex = 0; ChildIndex < SortedChildHandles.Num(); ++ChildIndex)
-			{
-				const bool bLastChild = SortedChildHandles.Num() - 1 == ChildIndex;
-				const TSharedRef<IPropertyHandle>& ChildHandle = SortedChildHandles[ChildIndex];
-
-				if (IsOverridePropertyOf(ChildHandle, PropertyHandle))
-				{
-					OverrideHandleWeakPtr = ChildHandle;
-					break;
-				}
-			}
-
 			TWeakPtr<IPropertyHandle> HandleWeakPtr = PropertyHandle;
 			return
 				SNew(SEditableTextBox)
 				.ToolTipText(PropertyHandle->GetToolTipText())
-				.Text_Lambda([HandleWeakPtr, OverrideHandleWeakPtr]() -> FText
+				.Text_Lambda([HandleWeakPtr, StructurePropertyHandle]() -> FText
 					{
+						using namespace UE::Chaos::ClothAsset;
+
 						FString Text;
-						if (const TSharedPtr<IPropertyHandle> OverrideHandlePtr = OverrideHandleWeakPtr.Pin())
-						{
-							OverrideHandlePtr->GetValue(Text);
-						}
-						if (Text == UE::Chaos::ClothAsset::FWeightMapTools::NotOverridden)
+						if (const TSharedPtr<IPropertyHandle> HandlePtr = HandleWeakPtr.Pin())
 						{
-							Text.Empty();  // GetValue seems to concatenate the text if the string isn't emptied first
-							if (const TSharedPtr<IPropertyHandle> HandlePtr = HandleWeakPtr.Pin())
+							const FDataflowNode* const DataflowNode = FClothDataflowTools::GetPropertyOwnerDataflowNode(StructurePropertyHandle);
+							if (ensure(DataflowNode))
 							{
-								HandlePtr->GetValue(Text);
+								void* Data;
+								if (ensure(HandlePtr->GetValueData(Data) == FPropertyAccess::Success))
+								{
+									Text = *static_cast<FString*>(Data);  // Default value if the property isn't an input, or isn't connected
+									if (const FDataflowInput* const DataflowInput = DataflowNode->FindInput(Data))
+									{
+										Dataflow::FContextThreaded Context(FPlatformTime::Cycles64());
+										Text = DataflowInput->GetValue<FString>(Context, Text);
+									}
+								}
 							}
 						}
 						return FText::FromString(Text);
-
 					})
 				.OnTextCommitted_Lambda([HandleWeakPtr](const FText& Text, ETextCommit::Type)
 					{
@@ -179,14 +179,24 @@ namespace UE::Chaos::ClothAsset
 							HandlePtr->SetValue(Text.ToString(), EPropertyValueSetFlags::DefaultFlags);
 						}
 					})
-				.IsEnabled_Lambda([OverrideHandleWeakPtr]() -> bool
+				.IsEnabled_Lambda([HandleWeakPtr, StructurePropertyHandle]() -> bool
 					{
-						FString Text;
-						if (const TSharedPtr<IPropertyHandle> OverrideHandlePtr = OverrideHandleWeakPtr.Pin())
+						if (const TSharedPtr<IPropertyHandle> HandlePtr = HandleWeakPtr.Pin())
 						{
-							OverrideHandlePtr->GetValue(Text);
+							const FDataflowNode* const DataflowNode = FClothDataflowTools::GetPropertyOwnerDataflowNode(StructurePropertyHandle);
+							if (ensure(DataflowNode))
+							{
+								void* Data;
+								if (HandlePtr->GetValueData(Data) == FPropertyAccess::Success)
+								{
+									if (const FDataflowInput* const DataflowInput = DataflowNode->FindInput(Data))
+									{
+										return !DataflowInput->HasAnyConnections();
+									}
+								}
+							}
 						}
-						return Text == UE::Chaos::ClothAsset::FWeightMapTools::NotOverridden;
+						return true;
 					})
 				.Font(IPropertyTypeCustomizationUtils::GetRegularFont());
 		}
diff --git a/Engine/Plugins/ChaosClothAssetEditor/Source/ChaosClothAssetDataflowNodes/Private/ChaosClothAsset/ConnectableValueCustomization.h b/Engine/Plugins/ChaosClothAssetEditor/Source/ChaosClothAssetDataflowNodes/Private/ChaosClothAsset/ConnectableValueCustomization.h
--- a/Engine/Plugins/ChaosClothAssetEditor/Source/ChaosClothAssetDataflowNodes/Private/ChaosClothAsset/ConnectableValueCustomization.h
+++ b/Engine/Plugins/ChaosClothAssetEditor/Source/ChaosClothAssetDataflowNodes/Private/ChaosClothAsset/ConnectableValueCustomization.h
@@ -22,7 +22,9 @@ namespace UE::Chaos::ClothAsset
 		virtual void MakeHeaderRow(TSharedRef<class IPropertyHandle>& StructPropertyHandle, FDetailWidgetRow& Row) override;
 		virtual TSharedRef<SWidget> MakeChildWidget(TSharedRef<IPropertyHandle>& StructurePropertyHandle, TSharedRef<IPropertyHandle>& PropertyHandle) override;
 		
+		UE_DEPRECATED(5.5, "Override properties are no longer used.")
 		static bool IsOverrideProperty(const TSharedPtr<IPropertyHandle>& Property);
+		UE_DEPRECATED(5.5, "Override properties are no longer used.")
 		static bool IsOverridePropertyOf(const TSharedPtr<IPropertyHandle>& OverrideProperty, const TSharedPtr<IPropertyHandle>& Property);
 		static bool BuildFabricMapsProperty(const TSharedPtr<IPropertyHandle>& Property);
 		static bool CouldUseFabricsProperty(const TSharedPtr<IPropertyHandle>& Property);
diff --git a/Engine/Plugins/ChaosClothAssetEditor/Source/ChaosClothAssetDataflowNodes/Private/ChaosClothAsset/DeleteElementNode.cpp b/Engine/Plugins/ChaosClothAssetEditor/Source/ChaosClothAssetDataflowNodes/Private/ChaosClothAsset/DeleteElementNode.cpp
--- a/Engine/Plugins/ChaosClothAssetEditor/Source/ChaosClothAssetDataflowNodes/Private/ChaosClothAsset/DeleteElementNode.cpp
+++ b/Engine/Plugins/ChaosClothAssetEditor/Source/ChaosClothAssetDataflowNodes/Private/ChaosClothAsset/DeleteElementNode.cpp
@@ -290,7 +290,6 @@ void FChaosClothAssetDeleteElementNode::Evaluate(Dataflow::FContext& Context, co
 			// Selection set
 			FCollectionClothSelectionFacade SelectionFacade(ClothCollection);
 			const FName InSelectionName(*GetValue<FString>(Context, &SelectionName.StringValue));
-			SelectionName.StringValue_Override = GetValue<FString>(Context, &SelectionName.StringValue, UE::Chaos::ClothAsset::FWeightMapTools::NotOverridden);
 
 			FName SelectionGroup;
 			TSet<int32> SelectedElements;
diff --git a/Engine/Plugins/ChaosClothAssetEditor/Source/ChaosClothAssetDataflowNodes/Private/ChaosClothAsset/ProxyDeformerNode.cpp b/Engine/Plugins/ChaosClothAssetEditor/Source/ChaosClothAssetDataflowNodes/Private/ChaosClothAsset/ProxyDeformerNode.cpp
--- a/Engine/Plugins/ChaosClothAssetEditor/Source/ChaosClothAssetDataflowNodes/Private/ChaosClothAsset/ProxyDeformerNode.cpp
+++ b/Engine/Plugins/ChaosClothAssetEditor/Source/ChaosClothAssetDataflowNodes/Private/ChaosClothAsset/ProxyDeformerNode.cpp
@@ -415,13 +415,6 @@ void FChaosClothAssetProxyDeformerNode_v2::Evaluate(Dataflow::FContext& Context,
 	{
 		using namespace UE::Chaos::ClothAsset;
 
-		// Update the selection names override
-		SimVertexSelection.StringValue_Override = GetValue<FString>(Context, &SimVertexSelection.StringValue, UE::Chaos::ClothAsset::FWeightMapTools::NotOverridden);
-		for (int32 FilterSetIndex = 0; FilterSetIndex < SelectionFilterSets.Num(); ++FilterSetIndex)
-		{
-			SelectionFilterSets[FilterSetIndex].RenderSelection.StringValue_Override = GetValue<FString>(Context, GetRenderConnectionReference(FilterSetIndex), UE::Chaos::ClothAsset::FWeightMapTools::NotOverridden);
-			SelectionFilterSets[FilterSetIndex].SimSelection.StringValue_Override = GetValue<FString>(Context, GetSimConnectionReference(FilterSetIndex), UE::Chaos::ClothAsset::FWeightMapTools::NotOverridden);
-		}
 
 		// Evaluate in collection
 		FManagedArrayCollection InCollection = GetValue<FManagedArrayCollection>(Context, &Collection);
@@ -595,14 +588,10 @@ void FChaosClothAssetProxyDeformerNode::Evaluate(Dataflow::FContext& Context, co
 	{
 		using namespace UE::Chaos::ClothAsset;
 
-		// Update the selection names override
-		SimVertexSelection.StringValue_Override = GetValue<FString>(Context, &SimVertexSelection.StringValue, UE::Chaos::ClothAsset::FWeightMapTools::NotOverridden);
-		SelectionFilterSet0.StringValue_Override = GetValue<FString>(Context, &SelectionFilterSet0.StringValue, UE::Chaos::ClothAsset::FWeightMapTools::NotOverridden);
 		const TArray<const FChaosClothAssetConnectableStringValue*> Non0SelectionFilterSets = Get1To9SelectionFilterSets();
 		for (int32 FilterSetIndex = 1; FilterSetIndex < NumFilterSets; ++FilterSetIndex)
 		{
 			const FChaosClothAssetConnectableStringValue& SelectionFilterSet = *Non0SelectionFilterSets[FilterSetIndex - 1];
-			SelectionFilterSet.StringValue_Override = GetValue<FString>(Context, &SelectionFilterSet.StringValue, UE::Chaos::ClothAsset::FWeightMapTools::NotOverridden);
 		}
 
 		// Evaluate in collection
diff --git a/Engine/Plugins/ChaosClothAssetEditor/Source/ChaosClothAssetDataflowNodes/Private/ChaosClothAsset/SelectionNode.cpp b/Engine/Plugins/ChaosClothAssetEditor/Source/ChaosClothAssetDataflowNodes/Private/ChaosClothAsset/SelectionNode.cpp
--- a/Engine/Plugins/ChaosClothAssetEditor/Source/ChaosClothAssetDataflowNodes/Private/ChaosClothAsset/SelectionNode.cpp
+++ b/Engine/Plugins/ChaosClothAssetEditor/Source/ChaosClothAssetDataflowNodes/Private/ChaosClothAsset/SelectionNode.cpp
@@ -366,13 +366,10 @@ void FChaosClothAssetSelectionNode::Evaluate(Dataflow::FContext& Context, const
 		}
 	};
 
-
 	if (Out->IsA<FManagedArrayCollection>(&Collection))
 	{
 		// Evaluate InputName
 		const FName InInputName = GetInputName(Context);
-		InputName.StringValue_Override = GetValue<FString>(Context, &InputName.StringValue, UE::Chaos::ClothAsset::FWeightMapTools::NotOverridden);
-
 		const FName SelectionName(Name.IsEmpty() ? InInputName : FName(Name));
 
 		if (SelectionName == NAME_None || Group.Name.IsEmpty())
diff --git a/Engine/Plugins/ChaosClothAssetEditor/Source/ChaosClothAssetDataflowNodes/Private/ChaosClothAsset/SelectionToIntMapNode.cpp b/Engine/Plugins/ChaosClothAssetEditor/Source/ChaosClothAssetDataflowNodes/Private/ChaosClothAsset/SelectionToIntMapNode.cpp
--- a/Engine/Plugins/ChaosClothAssetEditor/Source/ChaosClothAssetDataflowNodes/Private/ChaosClothAsset/SelectionToIntMapNode.cpp
+++ b/Engine/Plugins/ChaosClothAssetEditor/Source/ChaosClothAssetDataflowNodes/Private/ChaosClothAsset/SelectionToIntMapNode.cpp
@@ -33,9 +33,7 @@ void FChaosClothAssetSelectionToIntMapNode::Evaluate(Dataflow::FContext& Context
 
 		const FCollectionClothSelectionConstFacade SelectionFacade(ClothCollection);
 		const FName InSelectionName(*GetValue<FString>(Context, &SelectionName.StringValue));
-		SelectionName.StringValue_Override = GetValue<FString>(Context, &SelectionName.StringValue, UE::Chaos::ClothAsset::FWeightMapTools::NotOverridden);
 		const FString& InIntMapNameString = GetValue<FString>(Context, &IntMapName.StringValue);
-		IntMapName.StringValue_Override = GetValue<FString>(Context, &IntMapName.StringValue, UE::Chaos::ClothAsset::FWeightMapTools::NotOverridden);
 		const FName InIntMapName = InIntMapNameString.IsEmpty() ? InSelectionName : FName(*InIntMapNameString);
 		if (SelectionFacade.IsValid() && ClothFacade.IsValid() && InSelectionName != NAME_None && InIntMapName != NAME_None)
 		{
diff --git a/Engine/Plugins/ChaosClothAssetEditor/Source/ChaosClothAssetDataflowNodes/Private/ChaosClothAsset/SimulationBaseConfigNode.cpp b/Engine/Plugins/ChaosClothAssetEditor/Source/ChaosClothAssetDataflowNodes/Private/ChaosClothAsset/SimulationBaseConfigNode.cpp
--- a/Engine/Plugins/ChaosClothAssetEditor/Source/ChaosClothAssetDataflowNodes/Private/ChaosClothAsset/SimulationBaseConfigNode.cpp
+++ b/Engine/Plugins/ChaosClothAssetEditor/Source/ChaosClothAssetDataflowNodes/Private/ChaosClothAsset/SimulationBaseConfigNode.cpp
@@ -270,7 +270,6 @@ int32 FChaosClothAssetSimulationBaseConfigNode::FPropertyHelper::SetPropertyWeig
 	const float& PropertyLow,
 	const float& PropertyHigh,
 	const FString& WeightMap,
-	FString& MapOverride,
 	const TArray<FName>& SimilarPropertyNames,
 	ECollectionPropertyFlags PropertyFlags) const
 {
@@ -288,7 +287,6 @@ int32 FChaosClothAssetSimulationBaseConfigNode::FPropertyHelper::SetPropertyWeig
 	const int32 PropertyKeyIndex = ConfigNode.AddPropertyHelper(Properties, PropertyName, SimilarPropertyNames, PropertyFlags);
 	Properties.SetWeightedValue(PropertyKeyIndex, PropertyLow, PropertyHigh);
 	Properties.SetStringValue(PropertyKeyIndex, ConfigNode.GetValue<FString>(Context, &WeightMap));
-	MapOverride = ConfigNode.GetValue<FString>(Context, &WeightMap, UE::Chaos::ClothAsset::FWeightMapTools::NotOverridden);
 	return PropertyKeyIndex;
 }
 
@@ -299,7 +297,7 @@ int32 FChaosClothAssetSimulationBaseConfigNode::FPropertyHelper::SetPropertyWeig
 	ECollectionPropertyFlags PropertyFlags)
 {
 	return SetPropertyWeighted(PropertyName, PropertyValue.bIsAnimatable, PropertyValue.Low,
-		PropertyValue.High, PropertyValue.WeightMap, PropertyValue.WeightMap_Override, SimilarPropertyNames, PropertyFlags);
+		PropertyValue.High, PropertyValue.WeightMap, SimilarPropertyNames, PropertyFlags);
 }
 
 int32 FChaosClothAssetSimulationBaseConfigNode::FPropertyHelper::SetPropertyWeighted(
@@ -309,7 +307,7 @@ int32 FChaosClothAssetSimulationBaseConfigNode::FPropertyHelper::SetPropertyWeig
 	ECollectionPropertyFlags PropertyFlags)
 {
 	return SetPropertyWeighted(PropertyName, false, PropertyValue.Low,
-		PropertyValue.High, PropertyValue.WeightMap, PropertyValue.WeightMap_Override, SimilarPropertyNames, PropertyFlags);
+		PropertyValue.High, PropertyValue.WeightMap, SimilarPropertyNames, PropertyFlags);
 }
 
 int32 FChaosClothAssetSimulationBaseConfigNode::FPropertyHelper::SetPropertyWeighted(
@@ -319,7 +317,7 @@ int32 FChaosClothAssetSimulationBaseConfigNode::FPropertyHelper::SetPropertyWeig
 	ECollectionPropertyFlags PropertyFlags)
 {
 	return SetPropertyWeighted(PropertyName, false, 0.0f,
-		1.0f, PropertyValue.WeightMap, PropertyValue.WeightMap_Override, SimilarPropertyNames, PropertyFlags);
+		1.0f, PropertyValue.WeightMap, SimilarPropertyNames, PropertyFlags);
 }
 
 
@@ -390,7 +388,6 @@ int32 FChaosClothAssetSimulationBaseConfigNode::FPropertyHelper::SetPropertyStri
 {
 	const int32 PropertyKeyIndex = ConfigNode.AddPropertyHelper(Properties, PropertyName, SimilarPropertyNames, PropertyFlags);
 	Properties.SetStringValue(PropertyKeyIndex, ConfigNode.GetValue<FString>(Context, &PropertyValue.StringValue));
-	PropertyValue.StringValue_Override = ConfigNode.GetValue<FString>(Context, &PropertyValue.StringValue, UE::Chaos::ClothAsset::FWeightMapTools::NotOverridden);
 	return PropertyKeyIndex;
 }
 
diff --git a/Engine/Plugins/ChaosClothAssetEditor/Source/ChaosClothAssetDataflowNodes/Private/ChaosClothAsset/SimulationMaxDistanceConfigNode.cpp b/Engine/Plugins/ChaosClothAssetEditor/Source/ChaosClothAssetDataflowNodes/Private/ChaosClothAsset/SimulationMaxDistanceConfigNode.cpp
--- a/Engine/Plugins/ChaosClothAssetEditor/Source/ChaosClothAssetDataflowNodes/Private/ChaosClothAsset/SimulationMaxDistanceConfigNode.cpp
+++ b/Engine/Plugins/ChaosClothAssetEditor/Source/ChaosClothAssetDataflowNodes/Private/ChaosClothAsset/SimulationMaxDistanceConfigNode.cpp
@@ -41,7 +41,6 @@ void FChaosClothAssetSimulationMaxDistanceConfigNode::EvaluateClothCollection(Da
 	using namespace UE::Chaos::ClothAsset;
 	const FName MaxDistanceString(GetValue<FString>(Context, &MaxDistance.WeightMap)); //  Override for this is already set by AddProperties
 	const FName InputKinematicString(GetValue<FString>(Context, &InKinematic.StringValue));
-	InKinematic.StringValue_Override = GetValue<FString>(Context, &InKinematic.StringValue, FWeightMapTools::NotOverridden);
 
 	FCollectionClothSelectionFacade SelectionFacade(ClothCollection);
 	SelectionFacade.DefineSchema();
diff --git a/Engine/Plugins/ChaosClothAssetEditor/Source/ChaosClothAssetDataflowNodes/Private/ChaosClothAsset/TransferSkinWeightsNode.cpp b/Engine/Plugins/ChaosClothAssetEditor/Source/ChaosClothAssetDataflowNodes/Private/ChaosClothAsset/TransferSkinWeightsNode.cpp
--- a/Engine/Plugins/ChaosClothAssetEditor/Source/ChaosClothAssetDataflowNodes/Private/ChaosClothAsset/TransferSkinWeightsNode.cpp
+++ b/Engine/Plugins/ChaosClothAssetEditor/Source/ChaosClothAssetDataflowNodes/Private/ChaosClothAsset/TransferSkinWeightsNode.cpp
@@ -766,9 +766,6 @@ void FChaosClothAssetTransferSkinWeightsNode::Evaluate(Dataflow::FContext& Conte
 
 	if (Out->IsA<FManagedArrayCollection>(&Collection))
 	{
-		// Update the weight map override
-		InpaintMask.WeightMap_Override = GetValue<FString>(Context, &InpaintMask.WeightMap, UE::Chaos::ClothAsset::FWeightMapTools::NotOverridden);
-
 		// Evaluate inputs
 		FManagedArrayCollection InputCollection = GetValue<FManagedArrayCollection>(Context, &Collection);
 
diff --git a/Engine/Plugins/ChaosClothAssetEditor/Source/ChaosClothAssetDataflowNodes/Private/ChaosClothAsset/WeightMapNode.cpp b/Engine/Plugins/ChaosClothAssetEditor/Source/ChaosClothAssetDataflowNodes/Private/ChaosClothAsset/WeightMapNode.cpp
--- a/Engine/Plugins/ChaosClothAssetEditor/Source/ChaosClothAssetDataflowNodes/Private/ChaosClothAsset/WeightMapNode.cpp
+++ b/Engine/Plugins/ChaosClothAssetEditor/Source/ChaosClothAssetDataflowNodes/Private/ChaosClothAsset/WeightMapNode.cpp
@@ -152,7 +152,6 @@ void FChaosClothAssetWeightMapNode::Evaluate(Dataflow::FContext& Context, const
 	{
 		// Evaluate InputName
 		const FName InInputName = GetInputName(Context);
-		InputName.StringValue_Override = GetValue<FString>(Context, &InputName.StringValue, UE::Chaos::ClothAsset::FWeightMapTools::NotOverridden);
 
 		// Evaluate in collection
 		FManagedArrayCollection InCollection = GetValue<FManagedArrayCollection>(Context, &Collection);
diff --git a/Engine/Plugins/ChaosClothAssetEditor/Source/ChaosClothAssetDataflowNodes/Private/ChaosClothAsset/WeightedValueCustomization.cpp b/Engine/Plugins/ChaosClothAssetEditor/Source/ChaosClothAssetDataflowNodes/Private/ChaosClothAsset/WeightedValueCustomization.cpp
--- a/Engine/Plugins/ChaosClothAssetEditor/Source/ChaosClothAssetDataflowNodes/Private/ChaosClothAsset/WeightedValueCustomization.cpp
+++ b/Engine/Plugins/ChaosClothAssetEditor/Source/ChaosClothAssetDataflowNodes/Private/ChaosClothAsset/WeightedValueCustomization.cpp
@@ -91,10 +91,12 @@ namespace UE::Chaos::ClothAsset
 		{
 			TSharedRef<IPropertyHandle> ChildHandle = SortedChildHandles[ChildIndex];
 
+PRAGMA_DISABLE_DEPRECATION_WARNINGS
 			if (IsOverrideProperty(ChildHandle))
 			{
 				continue;  // Skip overrides
 			}
+PRAGMA_ENABLE_DEPRECATION_WARNINGS
 
 			// Propagate metadata to child properties so that it's reflected in the nested, individual spin boxes
 			ChildHandle->SetInstanceMetaData(TEXT("UIMin"), StructPropertyHandle->GetMetaData(TEXT("UIMin")));
@@ -171,64 +173,7 @@ namespace UE::Chaos::ClothAsset
 						});
 			}
 		}
-		if (PropertyClass == FStrProperty::StaticClass())
-		{
-			// Manage override property values (properties ending with _Override)
-			TWeakPtr<IPropertyHandle> OverrideHandleWeakPtr;
-
-			for (int32 ChildIndex = 0; ChildIndex < SortedChildHandles.Num(); ++ChildIndex)
-			{
-				const bool bLastChild = SortedChildHandles.Num() - 1 == ChildIndex;
-				const TSharedRef<IPropertyHandle>& ChildHandle = SortedChildHandles[ChildIndex];
-
-				if (IsOverridePropertyOf(ChildHandle, PropertyHandle))
-				{
-					OverrideHandleWeakPtr = ChildHandle;
-					break;
-				}
-			}
-
-			TWeakPtr<IPropertyHandle> HandleWeakPtr = PropertyHandle;
-			return
-				SNew(SEditableTextBox)
-				.ToolTipText(PropertyHandle->GetToolTipText())
-				.Text_Lambda([HandleWeakPtr, OverrideHandleWeakPtr]() -> FText
-					{
-						FString Text;
-						if (const TSharedPtr<IPropertyHandle> OverrideHandlePtr = OverrideHandleWeakPtr.Pin())
-						{
-							OverrideHandlePtr->GetValue(Text);
-						}
-						if (Text == UE::Chaos::ClothAsset::FWeightMapTools::NotOverridden)
-						{
-							Text.Empty();  // GetValue seems to concatenate the text if the string isn't emptied first
-							if (const TSharedPtr<IPropertyHandle> HandlePtr = HandleWeakPtr.Pin())
-							{
-								HandlePtr->GetValue(Text);
-							}
-						}
-						return FText::FromString(Text);
-
-					})
-				.OnTextCommitted_Lambda([HandleWeakPtr](const FText& Text, ETextCommit::Type)
-					{
-						if (const TSharedPtr<IPropertyHandle> HandlePtr = HandleWeakPtr.Pin())
-						{
-							HandlePtr->SetValue(Text.ToString(), EPropertyValueSetFlags::DefaultFlags);
-						}
-					})
-				.IsEnabled_Lambda([OverrideHandleWeakPtr]() -> bool
-					{
-						FString Text;
-						if (const TSharedPtr<IPropertyHandle> OverrideHandlePtr = OverrideHandleWeakPtr.Pin())
-						{
-							OverrideHandlePtr->GetValue(Text);
-						}
-						return Text == UE::Chaos::ClothAsset::FWeightMapTools::NotOverridden;
-					})
-				.Font(IPropertyTypeCustomizationUtils::GetRegularFont());
-		}
-		return SNullWidget::NullWidget;
+		return FConnectableValueCustomization::MakeChildWidget(StructurePropertyHandle, PropertyHandle);
 	}
 
 	TSharedRef<SWidget> FWeightedValueCustomization::MakeFloatWidget(
diff --git a/Engine/Plugins/ChaosClothAssetEditor/Source/ChaosClothAssetDataflowNodes/Public/ChaosClothAsset/ConnectableValue.h b/Engine/Plugins/ChaosClothAssetEditor/Source/ChaosClothAssetDataflowNodes/Public/ChaosClothAsset/ConnectableValue.h
--- a/Engine/Plugins/ChaosClothAssetEditor/Source/ChaosClothAssetDataflowNodes/Public/ChaosClothAsset/ConnectableValue.h
+++ b/Engine/Plugins/ChaosClothAssetEditor/Source/ChaosClothAssetDataflowNodes/Public/ChaosClothAsset/ConnectableValue.h
@@ -14,9 +14,9 @@ struct FChaosClothAssetConnectableStringValue
 	UPROPERTY(EditAnywhere, Category = "Value")
 	FString StringValue;
 
-	/** The string value override value for when the StringValue has a connection that replaces the provided string value. */
-	UPROPERTY(VisibleAnywhere, Category = "Value", Transient)
-	mutable FString StringValue_Override;  // _Override has a special meaning to the property customization, mutable because this property is set while getting the original value
+	UE_DEPRECATED(5.5, "Override properties are no longer used.")
+	UPROPERTY(Transient)
+	mutable FString StringValue_Override;
 };
 
 USTRUCT()
@@ -34,9 +34,9 @@ struct FChaosClothAssetConnectableIStringValue
 	UPROPERTY(VisibleAnywhere, Transient, Category = "Value")
 	bool bCouldUseFabrics = false;
 
-	/** The string value override value for when the StringValue has a connection that replaces the provided string value. */
-	UPROPERTY(VisibleAnywhere, Category = "Value", Transient)
-	mutable FString StringValue_Override;  // _Override has a special meaning to the property customization, mutable because this property is set while getting the original value
+	UE_DEPRECATED(5.5, "Override properties are no longer used.")
+	UPROPERTY(Transient)
+	mutable FString StringValue_Override;
 	
 	/**
 	 * Whether the property can override the weight map based on the imported fabrics
@@ -54,7 +54,7 @@ struct FChaosClothAssetConnectableIOStringValue
 	UPROPERTY(EditAnywhere, Category = "Value", Meta = (DataflowInput, DataflowOutput, DataflowPassthrough = "StringValue"))
 	FString StringValue;
 
-	/** The string value override value for when the StringValue has a connection that replaces the provided string value. */
-	UPROPERTY(VisibleAnywhere, Category = "Value", Transient)
-	mutable FString StringValue_Override;  // _Override has a special meaning to the property customization, mutable because this property is set while getting the original value
+	UE_DEPRECATED(5.5, "Override properties are no longer used.")
+	UPROPERTY(Transient)
+	mutable FString StringValue_Override;
 };
diff --git a/Engine/Plugins/ChaosClothAssetEditor/Source/ChaosClothAssetDataflowNodes/Public/ChaosClothAsset/SimulationBaseConfigNode.h b/Engine/Plugins/ChaosClothAssetEditor/Source/ChaosClothAssetDataflowNodes/Public/ChaosClothAsset/SimulationBaseConfigNode.h
--- a/Engine/Plugins/ChaosClothAssetEditor/Source/ChaosClothAssetDataflowNodes/Public/ChaosClothAsset/SimulationBaseConfigNode.h
+++ b/Engine/Plugins/ChaosClothAssetEditor/Source/ChaosClothAssetDataflowNodes/Public/ChaosClothAsset/SimulationBaseConfigNode.h
@@ -121,7 +121,7 @@ protected:
 		template<typename T UE_REQUIRES(std::is_base_of_v<FChaosClothAssetSimulationBaseConfigNode, T>)>
 		inline FName FindPropertyNameByAddress(const T* ConfigStruct, const void* PropertyAddress);
 		
-		int32 SetPropertyWeighted(const FName& PropertyName, const bool bIsAnimatable, const float& PropertyLow, const float& PropertyHigh, const FString& WeightMap, FString& MapOverride, const TArray<FName>& SimilarPropertyNames, ECollectionPropertyFlags PropertyFlags) const;
+		int32 SetPropertyWeighted(const FName& PropertyName, const bool bIsAnimatable, const float& PropertyLow, const float& PropertyHigh, const FString& WeightMap, const TArray<FName>& SimilarPropertyNames, ECollectionPropertyFlags PropertyFlags) const;
 
 		const FChaosClothAssetSimulationBaseConfigNode& ConfigNode;
 		Dataflow::FContext& Context;
diff --git a/Engine/Plugins/ChaosClothAssetEditor/Source/ChaosClothAssetDataflowNodes/Public/ChaosClothAsset/WeightedValue.h b/Engine/Plugins/ChaosClothAssetEditor/Source/ChaosClothAssetDataflowNodes/Public/ChaosClothAsset/WeightedValue.h
--- a/Engine/Plugins/ChaosClothAssetEditor/Source/ChaosClothAssetDataflowNodes/Public/ChaosClothAsset/WeightedValue.h
+++ b/Engine/Plugins/ChaosClothAssetEditor/Source/ChaosClothAssetDataflowNodes/Public/ChaosClothAsset/WeightedValue.h
@@ -9,11 +9,10 @@ namespace UE::Chaos::ClothAsset
 {
 	struct FWeightMapTools
 	{
-		/**
-		 * String code used to disable the override field in the UI.
-		 * Uses a character that is removed from the names made by MakeWeightMapName to avoid conflicts.
-		 */
+PRAGMA_DISABLE_DEPRECATION_WARNINGS
+		UE_DEPRECATED(5.5, "Override properties are no longer used.")
 		inline static const FString NotOverridden = TEXT("#");
+PRAGMA_ENABLE_DEPRECATION_WARNINGS
 
 		/**
 		 * Modifies a string to make it suitable as a weight map name.
@@ -73,9 +72,9 @@ struct FChaosClothAssetWeightedValue
 	UPROPERTY(VisibleAnywhere, Transient, Category = "Weighted Value")
 	bool bCouldUseFabrics = false;
 
-	/** The weight map override value for when the WeightMap has a connection that replaces the provided weight map value. */
-	UPROPERTY(VisibleAnywhere, Transient, Category = "Weighted Value")
-	mutable FString WeightMap_Override = UE::Chaos::ClothAsset::FWeightMapTools::NotOverridden;  // _Override has a special meaning to the property customization, mutable because this property is set while getting the original value
+	UE_DEPRECATED(5.5, "Override properties are no longer used.")
+	UPROPERTY(Transient)
+	mutable FString WeightMap_Override;
 	
 	/**
 	 * Whether the property can use the fabric bounds (low/high) values imported from USD
@@ -127,10 +126,10 @@ struct FChaosClothAssetWeightedValueNonAnimatable
 	UPROPERTY(VisibleAnywhere, Transient, Category = "Weighted Value")
 	bool bCouldUseFabrics = false;
 	
-	/** The weight map override value for when the WeightMap has a connection that replaces the provided weight map value. */
-	UPROPERTY(VisibleAnywhere, Transient, Category = "Weighted Value")
-	mutable FString WeightMap_Override = UE::Chaos::ClothAsset::FWeightMapTools::NotOverridden;  // _Override has a special meaning to the property customization, mutable because this property is set while getting the original value
-	
+	UE_DEPRECATED(5.5, "Override properties are no longer used.")
+	UPROPERTY(Transient)
+	mutable FString WeightMap_Override;
+
 	/**
 	 * Whether the property can use the fabric bounds (low/high) values imported from USD
 	 */
@@ -153,9 +152,9 @@ struct FChaosClothAssetWeightedValueNonAnimatableNoLowHighRange
 	UPROPERTY(EditAnywhere, Category = "Weighted Value", Meta = (DataflowInput))
 	FString WeightMap = TEXT("WeightMap");
 
-	/** The weight map override value for when the WeightMap has a connection that replaces the provided weight map value. */
-	UPROPERTY(VisibleAnywhere, Transient, Category = "Weighted Value")
-	mutable FString WeightMap_Override = UE::Chaos::ClothAsset::FWeightMapTools::NotOverridden;  // _Override has a special meaning to the property customization, mutable because this property is set while getting the original value
+	UE_DEPRECATED(5.5, "Override properties are no longer used.")
+	UPROPERTY(Transient)
+	mutable FString WeightMap_Override;
 };
 
 USTRUCT()
