Cloth Asset - Trivial change, removing a commented out header inclusion from a private file.

diff --git a/Engine/Plugins/ChaosClothAssetEditor/Source/ChaosClothAssetDataflowNodes/Private/ChaosClothAsset/ConnectableValueCustomization.cpp b/Engine/Plugins/ChaosClothAssetEditor/Source/ChaosClothAssetDataflowNodes/Private/ChaosClothAsset/ConnectableValueCustomization.cpp
--- a/Engine/Plugins/ChaosClothAssetEditor/Source/ChaosClothAssetDataflowNodes/Private/ChaosClothAsset/ConnectableValueCustomization.cpp
+++ b/Engine/Plugins/ChaosClothAssetEditor/Source/ChaosClothAssetDataflowNodes/Private/ChaosClothAsset/ConnectableValueCustomization.cpp
@@ -12,7 +12,6 @@
 #include "Widgets/Input/SEditableTextBox.h"
 #include "DetailLayoutBuilder.h"
 #include "DetailWidgetRow.h"
-//#include "Editor.h"
 
 #define LOCTEXT_NAMESPACE "ChaosClothAssetWeightedValueCustomization"
 
