CIS - Fixed issue 511509: warning: 'WeightMap_Override' is deprecated.

#rnx

diff --git a/Engine/Plugins/ChaosClothAssetEditor/Source/ChaosClothAssetDataflowNodes/Public/ChaosClothAsset/ConnectableValue.h b/Engine/Plugins/ChaosClothAssetEditor/Source/ChaosClothAssetDataflowNodes/Public/ChaosClothAsset/ConnectableValue.h
--- a/Engine/Plugins/ChaosClothAssetEditor/Source/ChaosClothAssetDataflowNodes/Public/ChaosClothAsset/ConnectableValue.h
+++ b/Engine/Plugins/ChaosClothAssetEditor/Source/ChaosClothAssetDataflowNodes/Public/ChaosClothAsset/ConnectableValue.h
@@ -5,10 +5,12 @@
 #include "Containers/UnrealString.h"
 #include "ConnectableValue.generated.h"
 
+PRAGMA_DISABLE_DEPRECATION_WARNINGS
 USTRUCT()
 struct FChaosClothAssetConnectableStringValue
 {
 	GENERATED_BODY()
+PRAGMA_ENABLE_DEPRECATION_WARNINGS
 
 	/** The value for this property. */
 	UPROPERTY(EditAnywhere, Category = "Value")
@@ -19,10 +21,12 @@ struct FChaosClothAssetConnectableStringValue
 	mutable FString StringValue_Override;
 };
 
+PRAGMA_DISABLE_DEPRECATION_WARNINGS
 USTRUCT()
 struct FChaosClothAssetConnectableIStringValue
 {
 	GENERATED_BODY()
+PRAGMA_ENABLE_DEPRECATION_WARNINGS
 
 	/** The value for this property. */
 	UPROPERTY(EditAnywhere, Category = "Value", Meta = (DataflowInput))
@@ -45,10 +49,12 @@ struct FChaosClothAssetConnectableIStringValue
 	bool bBuildFabricMaps = false;
 };
 
+PRAGMA_DISABLE_DEPRECATION_WARNINGS
 USTRUCT()
 struct FChaosClothAssetConnectableIOStringValue
 {
 	GENERATED_BODY()
+PRAGMA_ENABLE_DEPRECATION_WARNINGS
 
 	/** The value for this property. */
 	UPROPERTY(EditAnywhere, Category = "Value", Meta = (DataflowInput, DataflowOutput, DataflowPassthrough = "StringValue"))
diff --git a/Engine/Plugins/ChaosClothAssetEditor/Source/ChaosClothAssetDataflowNodes/Public/ChaosClothAsset/WeightedValue.h b/Engine/Plugins/ChaosClothAssetEditor/Source/ChaosClothAssetDataflowNodes/Public/ChaosClothAsset/WeightedValue.h
--- a/Engine/Plugins/ChaosClothAssetEditor/Source/ChaosClothAssetDataflowNodes/Public/ChaosClothAsset/WeightedValue.h
+++ b/Engine/Plugins/ChaosClothAssetEditor/Source/ChaosClothAssetDataflowNodes/Public/ChaosClothAsset/WeightedValue.h
@@ -30,10 +30,12 @@ PRAGMA_ENABLE_DEPRECATION_WARNINGS
 	};
 }
 
+PRAGMA_DISABLE_DEPRECATION_WARNINGS
 USTRUCT()
 struct FChaosClothAssetWeightedValue
 {
 	GENERATED_BODY()
+PRAGMA_ENABLE_DEPRECATION_WARNINGS
 
 	/**
 	 * Whether the property can ever be updated/animated in real time.
@@ -89,10 +91,12 @@ struct FChaosClothAssetWeightedValue
 	bool bBuildFabricMaps = false;
 };
 
+PRAGMA_DISABLE_DEPRECATION_WARNINGS
 USTRUCT()
 struct FChaosClothAssetWeightedValueNonAnimatable
 {
 	GENERATED_BODY()
+PRAGMA_ENABLE_DEPRECATION_WARNINGS
 
 	static constexpr bool bIsAnimatable = false;
 
@@ -143,10 +147,12 @@ struct FChaosClothAssetWeightedValueNonAnimatable
 	bool bBuildFabricMaps = false;
 };
 
+PRAGMA_DISABLE_DEPRECATION_WARNINGS
 USTRUCT()
 struct FChaosClothAssetWeightedValueNonAnimatableNoLowHighRange
 {
 	GENERATED_BODY()
+PRAGMA_ENABLE_DEPRECATION_WARNINGS
 
 	/** The name of the weight map for this property. */
 	UPROPERTY(EditAnywhere, Category = "Weighted Value", Meta = (DataflowInput))
