Cloth Editor and Dataflow Editor: allow tools to start when PIE is running

diff --git a/Engine/Plugins/ChaosClothAssetEditor/Source/ChaosClothAssetEditor/Private/ChaosClothAsset/ClothEditorMode.cpp b/Engine/Plugins/ChaosClothAssetEditor/Source/ChaosClothAssetEditor/Private/ChaosClothAsset/ClothEditorMode.cpp
--- a/Engine/Plugins/ChaosClothAssetEditor/Source/ChaosClothAssetEditor/Private/ChaosClothAsset/ClothEditorMode.cpp
+++ b/Engine/Plugins/ChaosClothAssetEditor/Source/ChaosClothAssetEditor/Private/ChaosClothAsset/ClothEditorMode.cpp
@@ -74,6 +74,10 @@ const FEditorModeID UChaosClothAssetEditorMode::EM_ChaosClothAssetEditorModeId =
 
 namespace UE::Chaos::ClothAsset::Private
 {
+	bool bClothEditorEnableToolsInPIE = true;
+	FAutoConsoleVariableRef CVARClothEditorEnableToolsInPIE(TEXT("p.ChaosCloth.EnableToolsInPIE"), bClothEditorEnableToolsInPIE,
+		TEXT("Enable Cloth Editor tools while Play In Editor is running [def:true]"));
+
 	void RemoveClothWeightMaps(UE::Chaos::ClothAsset::FCollectionClothFacade& ClothFacade, const TArray<FName>& WeightMapNames)
 	{
 		for (const FName& WeightMapName : WeightMapNames)
@@ -367,7 +371,15 @@ bool UChaosClothAssetEditorMode::ShouldToolStartBeAllowed(const FString& ToolIde
 		}
 	}
 
-	return Super::ShouldToolStartBeAllowed(ToolIdentifier);
+	if (UE::Chaos::ClothAsset::Private::bClothEditorEnableToolsInPIE)
+	{
+		// UEdMode::ShouldToolStartBeAllowed returns (!GEditor->PlayWorld && !GIsPlayInEditorWorld) but we want to allow tools to start while in PIE
+		return true;
+	}
+	else
+	{
+		return UBaseCharacterFXEditorMode::ShouldToolStartBeAllowed(ToolIdentifier);
+	}
 }
 
 void UChaosClothAssetEditorMode::CreateToolkit()
diff --git a/Engine/Plugins/Experimental/Dataflow/Source/DataflowEditor/Private/Dataflow/DataflowEditorMode.cpp b/Engine/Plugins/Experimental/Dataflow/Source/DataflowEditor/Private/Dataflow/DataflowEditorMode.cpp
--- a/Engine/Plugins/Experimental/Dataflow/Source/DataflowEditor/Private/Dataflow/DataflowEditorMode.cpp
+++ b/Engine/Plugins/Experimental/Dataflow/Source/DataflowEditor/Private/Dataflow/DataflowEditorMode.cpp
@@ -50,6 +50,13 @@
 
 const FEditorModeID UDataflowEditorMode::EM_DataflowEditorModeId = TEXT("EM_DataflowAssetEditorMode");
 
+namespace Dataflow::Private
+{
+	bool bDataflowEditorEnableToolsInPIE = true;
+	FAutoConsoleVariableRef CVARDataflowEditorEnableToolsInPIE(TEXT("p.Dataflow.EnableToolsInPIE"), bDataflowEditorEnableToolsInPIE,
+		TEXT("Enable Dataflow Editor tools while Play In Editor is running [def:true]"));
+}
+
 UDataflowEditorMode::UDataflowEditorMode()
 {
 	Info = FEditorModeInfo(
@@ -201,7 +208,16 @@ bool UDataflowEditorMode::ShouldToolStartBeAllowed(const FString& ToolIdentifier
 		}
 	}
 
-	return Super::ShouldToolStartBeAllowed(ToolIdentifier);
+
+	if (Dataflow::Private::bDataflowEditorEnableToolsInPIE)
+	{
+		// UEdMode::ShouldToolStartBeAllowed returns (!GEditor->PlayWorld && !GIsPlayInEditorWorld) but we want to allow tools to start while in PIE
+		return true;
+	}
+	else
+	{
+		return UBaseCharacterFXEditorMode::ShouldToolStartBeAllowed(ToolIdentifier);
+	}
 }
 
 void UDataflowEditorMode::CreateToolkit()
