Enable opening Cloth Assets in the Dataflow Editor via content browser context menu

diff --git a/Engine/Plugins/ChaosClothAssetEditor/Source/ChaosClothAssetEditor/Private/ChaosClothAsset/ChaosClothAssetEditorModule.cpp b/Engine/Plugins/ChaosClothAssetEditor/Source/ChaosClothAssetEditor/Private/ChaosClothAsset/ChaosClothAssetEditorModule.cpp
--- a/Engine/Plugins/ChaosClothAssetEditor/Source/ChaosClothAssetEditor/Private/ChaosClothAsset/ChaosClothAssetEditorModule.cpp
+++ b/Engine/Plugins/ChaosClothAssetEditor/Source/ChaosClothAssetEditor/Private/ChaosClothAsset/ChaosClothAssetEditorModule.cpp
@@ -1,19 +1,84 @@
 // Copyright Epic Games, Inc. All Rights Reserved.
 
 #include "ChaosClothAsset/ChaosClothAssetEditorModule.h"
+#include "ChaosClothAsset/ClothAsset.h"
 #include "ChaosClothAsset/ClothEditorCommands.h"
 #include "ChaosClothAsset/ClothEditorMode.h"
 #include "ChaosClothAsset/ClothEditorStyle.h"
 #include "Modules/ModuleManager.h"
 #include "EditorModeRegistry.h"
+#include "Algo/AllOf.h"
+#include "ContentBrowserMenuContexts.h"
+#include "Dataflow/DataflowEditorModeUILayer.h"
+#include "Dataflow/DataflowEditor.h"
+#include "ToolMenus.h"
+
+#define LOCTEXT_NAMESPACE "ChaosClothAssetEditorModule"
 
 namespace UE::Chaos::ClothAsset
 {
+	void FChaosClothAssetEditorModule::RegisterMenus()
+	{
+		// Allows cleanup when module unloads.
+		FToolMenuOwnerScoped OwnerScoped(this);
+
+		// Enable opening ChaosClothAssets in the Dataflow Editor via the Content Browser context menu
+		// (Note: this should be temporary until the Dataflow Editor becomes *the* editor for ChaosClothAssets)
+
+		UToolMenu* const ClothContextMenu = UToolMenus::Get()->ExtendMenu("ContentBrowser.AssetContextMenu.ChaosClothAsset");
+		FToolMenuSection& Section = ClothContextMenu->FindOrAddSection("GetAssetActions");
+		Section.AddDynamicEntry("OpenInDataflowEditor", FNewToolMenuSectionDelegate::CreateLambda([](FToolMenuSection& Section)
+		{
+			// We'll need to get the target assets out of the context
+			if (const UContentBrowserAssetContextMenuContext* const Context = Section.FindContext<UContentBrowserAssetContextMenuContext>())
+			{
+				// We are deliberately not using Context->GetSelectedObjects() here to avoid triggering a load from right clicking
+				// an asset in the content browser.
+				const bool bAllSelectedAssetsAreCloth = Algo::AllOf(Context->SelectedAssets, [](const FAssetData& Asset)
+				{
+					return Asset.AssetClassPath == UChaosClothAsset::StaticClass()->GetClassPathName();
+				});
+
+				if (bAllSelectedAssetsAreCloth)
+				{
+					UDataflowEditorUISubsystem* const DataflowEditorSubsystem = GEditor->GetEditorSubsystem<UDataflowEditorUISubsystem>();
+					check(DataflowEditorSubsystem);
+
+					const TSharedPtr<class FUICommandList> CommandListToBind = MakeShared<FUICommandList>();
+
+					CommandListToBind->MapAction(
+						FChaosClothAssetEditorCommands::Get().OpenClothAssetInDataflowEditor,
+						FExecuteAction::CreateWeakLambda(DataflowEditorSubsystem, [Context, DataflowEditorSubsystem]()
+						{
+							// When we actually do want to open the editor, trigger the load to get the objects
+							TArray<TObjectPtr<UObject>> AssetsToEdit;
+							AssetsToEdit.Append(Context->LoadSelectedObjects<UObject>());
+
+							UDataflowEditor* const AssetEditor = NewObject<UDataflowEditor>(DataflowEditorSubsystem, NAME_None, RF_Transient);
+							AssetEditor->Initialize({AssetsToEdit[0]});
+						}),
+						FCanExecuteAction::CreateWeakLambda(Context, [Context]() { return Context->bCanBeModified; }));
+
+					const TAttribute<FText> ToolTipOverride = Context->bCanBeModified ? TAttribute<FText>() : LOCTEXT("ReadOnlyAssetWarning", "The selected asset(s) are read-only and cannot be edited.");
+					Section.AddMenuEntryWithCommandList(FChaosClothAssetEditorCommands::Get().OpenClothAssetInDataflowEditor,
+						CommandListToBind,
+						TAttribute<FText>(),
+						ToolTipOverride,
+						FSlateIcon());		// TODO: If DataflowEditorStyle.h was public we could do this: FSlateIcon(FDataflowEditorStyle::Get().GetStyleSetName(), "ClassThumbnail.Dataflow")
+				}
+			}
+		}));
+	}
+
 	void FChaosClothAssetEditorModule::StartupModule()
 	{
 		FChaosClothAssetEditorStyle::Get(); // Causes the constructor to be called
 
 		FChaosClothAssetEditorCommands::Register();
+
+		// Menus need to be registered in a callback to make sure the system is ready for them.
+		StartupCallbackDelegateHandle = UToolMenus::RegisterStartupCallback(
+			FSimpleMulticastDelegate::FDelegate::CreateRaw(this, &FChaosClothAssetEditorModule::RegisterMenus));
 	}
 
 	void FChaosClothAssetEditorModule::ShutdownModule()
@@ -21,7 +86,12 @@ namespace UE::Chaos::ClothAsset
 		FChaosClothAssetEditorCommands::Unregister();
 
 		FEditorModeRegistry::Get().UnregisterMode(UChaosClothAssetEditorMode::EM_ChaosClothAssetEditorModeId);
+
+		UToolMenus::UnRegisterStartupCallback(StartupCallbackDelegateHandle);
 	}
+
 } // namespace UE::Chaos::ClothAsset
 
 IMPLEMENT_MODULE(UE::Chaos::ClothAsset::FChaosClothAssetEditorModule, ChaosClothAssetEditor)
+
+#undef LOCTEXT_NAMESPACE
diff --git a/Engine/Plugins/ChaosClothAssetEditor/Source/ChaosClothAssetEditor/Private/ChaosClothAsset/ChaosClothAssetEditorModule.h b/Engine/Plugins/ChaosClothAssetEditor/Source/ChaosClothAssetEditor/Private/ChaosClothAsset/ChaosClothAssetEditorModule.h
--- a/Engine/Plugins/ChaosClothAssetEditor/Source/ChaosClothAssetEditor/Private/ChaosClothAsset/ChaosClothAssetEditorModule.h
+++ b/Engine/Plugins/ChaosClothAssetEditor/Source/ChaosClothAssetEditor/Private/ChaosClothAsset/ChaosClothAssetEditorModule.h
@@ -12,5 +12,12 @@ namespace UE::Chaos::ClothAsset
 		/** IModuleInterface implementation */
 		virtual void StartupModule() override;
 		virtual void ShutdownModule() override;
+
+	private:
+
+		FDelegateHandle StartupCallbackDelegateHandle;
+
+		void RegisterMenus();
+		
 	};
 } // namespace UE::Chaos::ClothAsset
diff --git a/Engine/Plugins/ChaosClothAssetEditor/Source/ChaosClothAssetEditor/Private/ChaosClothAsset/ClothEditorCommands.cpp b/Engine/Plugins/ChaosClothAssetEditor/Source/ChaosClothAssetEditor/Private/ChaosClothAsset/ClothEditorCommands.cpp
--- a/Engine/Plugins/ChaosClothAssetEditor/Source/ChaosClothAssetEditor/Private/ChaosClothAsset/ClothEditorCommands.cpp
+++ b/Engine/Plugins/ChaosClothAssetEditor/Source/ChaosClothAssetEditor/Private/ChaosClothAsset/ClothEditorCommands.cpp
@@ -40,6 +40,8 @@ void FChaosClothAssetEditorCommands::RegisterCommands()
 
 	UI_COMMAND(OpenClothEditor, "Cloth Editor", "Open the Cloth Editor window", EUserInterfaceActionType::Button, FInputChord());
 
+	UI_COMMAND(OpenClothAssetInDataflowEditor, "Open in Dataflow Editor", "Open the Cloth Asset in Dataflow Editor", EUserInterfaceActionType::Button, FInputChord());
+
 	UI_COMMAND(BeginRemeshTool, "Remesh", "Remesh the selected mesh", EUserInterfaceActionType::Button, FInputChord());
 	UI_COMMAND(BeginAttributeEditorTool, "AttrEd", "Edit/configure mesh attributes", EUserInterfaceActionType::Button, FInputChord());
 	
diff --git a/Engine/Plugins/ChaosClothAssetEditor/Source/ChaosClothAssetEditor/Private/ChaosClothAsset/ClothEditorCommands.h b/Engine/Plugins/ChaosClothAssetEditor/Source/ChaosClothAssetEditor/Private/ChaosClothAsset/ClothEditorCommands.h
--- a/Engine/Plugins/ChaosClothAssetEditor/Source/ChaosClothAssetEditor/Private/ChaosClothAsset/ClothEditorCommands.h
+++ b/Engine/Plugins/ChaosClothAssetEditor/Source/ChaosClothAssetEditor/Private/ChaosClothAsset/ClothEditorCommands.h
@@ -30,6 +30,8 @@ public:
 
 	TSharedPtr<FUICommandInfo> OpenClothEditor;
 
+	TSharedPtr<FUICommandInfo> OpenClothAssetInDataflowEditor;
+
 	const static FString BeginRemeshToolIdentifier;
 	TSharedPtr<FUICommandInfo> BeginRemeshTool;
 
