Add "experimental" text to the context menu item for loading a Cloth Asset in the Dataflow Editor. This will signal to the user that the asset will be loaded in a different editor than the usual Cloth Asset Editor.

diff --git a/Engine/Plugins/ChaosClothAssetEditor/Source/ChaosClothAssetEditor/Private/ChaosClothAsset/ClothEditorCommands.cpp b/Engine/Plugins/ChaosClothAssetEditor/Source/ChaosClothAssetEditor/Private/ChaosClothAsset/ClothEditorCommands.cpp
--- a/Engine/Plugins/ChaosClothAssetEditor/Source/ChaosClothAssetEditor/Private/ChaosClothAsset/ClothEditorCommands.cpp
+++ b/Engine/Plugins/ChaosClothAssetEditor/Source/ChaosClothAssetEditor/Private/ChaosClothAsset/ClothEditorCommands.cpp
@@ -40,7 +40,7 @@ void FChaosClothAssetEditorCommands::RegisterCommands()
 
 	UI_COMMAND(OpenClothEditor, "Cloth Editor", "Open the Cloth Editor window", EUserInterfaceActionType::Button, FInputChord());
 
-	UI_COMMAND(OpenClothAssetInDataflowEditor, "Open in Dataflow Editor", "Open the Cloth Asset in Dataflow Editor", EUserInterfaceActionType::Button, FInputChord());
+	UI_COMMAND(OpenClothAssetInDataflowEditor, "Open in Dataflow Editor (experimental)", "Open the Cloth Asset in Dataflow Editor", EUserInterfaceActionType::Button, FInputChord());
 
 	UI_COMMAND(BeginRemeshTool, "Remesh", "Remesh the selected mesh", EUserInterfaceActionType::Button, FInputChord());
 	UI_COMMAND(BeginAttributeEditorTool, "AttrEd", "Edit/configure mesh attributes", EUserInterfaceActionType::Button, FInputChord());
