Dataflow and Cloth Asset Editors: make undo/redo work when changing node property values. Also fix some other undo/redo issues related to connections and option pins.

These changes make undo/redo work for node properties
* Add a NotifyPreChange hook to the node property editor which will mark the selected EdNode(s) with Modify (this will cause them to Serialize for Undo/Redo).
* Make EdNodes Serialize their associated FDataflowNode when transacting (i.e., doing undo/redo)

When doing other transactions (e.g., add/remove pins, add/remove connections), we're still serializing the whole graph. We should make our own transactions for the Undo queue for this (not part of this CL). In the mean time, I made the following fixes to support undo/redo:
* Fixed the Serialize methods for existing nodes that use option pins to correctly add/remove pins when transacting.
* Updated the ClothEditorToolkit to hold onto the SelectedDataflowNode via Guid rather than pointer since the pointer goes stale when deserializing the whole graph.

Also:
* Updated the WeightedValueCustomization with the NotTransactable flag when using the slider. This matches the MathStructCustomization and fixed an issue with undo/redo when using the slider.

diff --git a/Engine/Plugins/ChaosClothAssetEditor/Source/ChaosClothAssetDataflowNodes/Private/ChaosClothAsset/MergeClothCollectionsNode.cpp b/Engine/Plugins/ChaosClothAssetEditor/Source/ChaosClothAssetDataflowNodes/Private/ChaosClothAsset/MergeClothCollectionsNode.cpp
--- a/Engine/Plugins/ChaosClothAssetEditor/Source/ChaosClothAssetDataflowNodes/Private/ChaosClothAsset/MergeClothCollectionsNode.cpp
+++ b/Engine/Plugins/ChaosClothAssetEditor/Source/ChaosClothAssetDataflowNodes/Private/ChaosClothAsset/MergeClothCollectionsNode.cpp
@@ -332,9 +332,13 @@ namespace UE::Chaos::ClothAsset::Private
 FChaosClothAssetMergeClothCollectionsNode_v2::FChaosClothAssetMergeClothCollectionsNode_v2(const Dataflow::FNodeParameters& InParam, FGuid InGuid)
 	: FDataflowNode(InParam, InGuid)
 {
+	check(GetNumInputs() == NumRequiredInputs);
+
 	// Add two sets of pins to start.
-	AddPins();
-	AddPins();
+	for (int32 Index = 0; Index < NumInitialOptionalInputs; ++Index)
+	{
+		AddPins();
+	}
 	RegisterOutputConnection(&Collection)
 		.SetPassthroughInput(GetConnectionReference(0));
 }
@@ -474,13 +478,38 @@ void FChaosClothAssetMergeClothCollectionsNode_v2::Serialize(FArchive& Ar)
 {
 	if (Ar.IsLoading())
 	{
-		check(Collections.Num() > 1);
-		check(FindInput(GetConnectionReference(0)));
-		check(FindInput(GetConnectionReference(1)));
+		check(Collections.Num() >= NumInitialOptionalInputs);
 
-		for (int32 Index = 2; Index < Collections.Num(); ++Index)
+		for (int32 Index = 0; Index < NumInitialOptionalInputs; ++Index)
+		{
+			check(FindInput(GetConnectionReference(Index)));
+		}
+
+		for (int32 Index = NumInitialOptionalInputs; Index < Collections.Num(); ++Index)
+		{
+			FindOrRegisterInputArrayConnection(GetConnectionReference(Index));
+		}
+		if (Ar.IsTransacting())
+		{
+			const int32 OrigNumRegisteredInputs = GetNumInputs();
+			check(OrigNumRegisteredInputs >= NumRequiredInputs + NumInitialOptionalInputs);
+			const int32 OrigNumCollections = Collections.Num();
+			const int32 OrigNumRegisteredCollections = OrigNumRegisteredInputs - NumRequiredInputs;
+			if (OrigNumRegisteredCollections > OrigNumCollections)
+			{
+				// Inputs have been removed.
+				// Temporarily expand Collections so we can get connection references.
+				Collections.SetNum(GetNumInputs() - 1);
+				for (int32 Index = OrigNumCollections; Index < Collections.Num(); ++Index)
+				{
+					UnregisterInputConnection(GetConnectionReference(Index));
+				}
+				Collections.SetNum(OrigNumCollections);
+			}
+		}
+		else
 		{
-			RegisterInputArrayConnection(GetConnectionReference(Index));
+			ensureAlways(Collections.Num() == GetNumInputs());
 		}
 	}
 }
@@ -499,6 +528,8 @@ FChaosClothAssetMergeClothCollectionsNode::FChaosClothAssetMergeClothCollections
 	RegisterInputConnection(&Collection);
 	RegisterOutputConnection(&Collection)
 		.SetPassthroughInput(&Collection);
+
+	check(GetNumInputs() == NumRequiredInputs + NumInitialOptionalInputs); // Update NumRequiredInputs if you add more Inputs. This is used by Serialize.
 }
 
 void FChaosClothAssetMergeClothCollectionsNode::Evaluate(Dataflow::FContext& Context, const FDataflowOutput* Out) const
@@ -708,17 +739,46 @@ TArray<const FManagedArrayCollection*> FChaosClothAssetMergeClothCollectionsNode
 	return Collections;
 }
 
+
+const FManagedArrayCollection* FChaosClothAssetMergeClothCollectionsNode::GetCollection(int32 Index) const
+{
+	switch (Index)
+	{
+	case 0: return &Collection;
+	case 1: return &Collection1;
+	case 2: return &Collection2;
+	case 3: return &Collection3;
+	case 4: return &Collection4;
+	case 5: return &Collection5;
+	default: check(false) return nullptr;
+	}
+}
+
 void FChaosClothAssetMergeClothCollectionsNode::Serialize(FArchive& Ar)
 {
 	if (Ar.IsLoading())
 	{
-		const int32 NumInputsToAdd = NumInputs - 1;
-		NumInputs = 1;  // AddPin will increment it again
-		for (int32 InputIndex = 0; InputIndex < NumInputsToAdd; ++InputIndex)
+		const int32 OrigNumRegisteredInputs = GetNumInputs() - NumRequiredInputs;
+		const int32 OrigNumInputs = NumInputs;
+		const int32 NumInputsToAdd = OrigNumInputs - OrigNumRegisteredInputs;
+		check(Ar.IsTransacting() || OrigNumRegisteredInputs == NumInitialOptionalInputs)
+		if (NumInputsToAdd > 0)
 		{
-			AddPins();
+			NumInputs = OrigNumRegisteredInputs;  // AddPin will increment it again
+			for (int32 InputIndex = 0; InputIndex < NumInputsToAdd; ++InputIndex)
+			{
+				AddPins();
+			}
+		}
+		else if (NumInputsToAdd < 0)
+		{
+			check(Ar.IsTransacting());
+			for (int32 Index = NumInputs; Index < OrigNumRegisteredInputs; ++Index)
+			{
+				UnregisterInputConnection(GetCollection(Index));
+			}
 		}
-		check(NumInputsToAdd == NumInputs - 1);
+		check(NumInputs + NumRequiredInputs == GetNumInputs());
 	}
 }
 
diff --git a/Engine/Plugins/ChaosClothAssetEditor/Source/ChaosClothAssetDataflowNodes/Private/ChaosClothAsset/ProxyDeformerNode.cpp b/Engine/Plugins/ChaosClothAssetEditor/Source/ChaosClothAssetDataflowNodes/Private/ChaosClothAsset/ProxyDeformerNode.cpp
--- a/Engine/Plugins/ChaosClothAssetEditor/Source/ChaosClothAssetDataflowNodes/Private/ChaosClothAsset/ProxyDeformerNode.cpp
+++ b/Engine/Plugins/ChaosClothAssetEditor/Source/ChaosClothAssetDataflowNodes/Private/ChaosClothAsset/ProxyDeformerNode.cpp
@@ -406,7 +406,12 @@ FChaosClothAssetProxyDeformerNode_v2::FChaosClothAssetProxyDeformerNode_v2(const
 	RegisterOutputConnection(&SkinningBlendName);
 
 	// Start with one set of option pins.
-	AddPins();
+	for (int32 Index = 0; Index < NumInitialSelectionFilterSets; ++Index)
+	{
+		AddPins();
+	}
+
+	check(GetNumInputs() == NumRequiredInputs + NumInitialSelectionFilterSets); // Update NumRequiredInputs if you add more Inputs. This is used by Serialize.
 }
 
 void FChaosClothAssetProxyDeformerNode_v2::Evaluate(Dataflow::FContext& Context, const FDataflowOutput* Out) const
@@ -530,14 +535,42 @@ void FChaosClothAssetProxyDeformerNode_v2::Serialize(FArchive& Ar)
 	// Restore the pins when re-loading so they can get properly reconnected
 	if (Ar.IsLoading())
 	{
-		check(SelectionFilterSets.Num() > 0);
-		check(FindInput(GetRenderConnectionReference(0)));
-		check(FindInput(GetSimConnectionReference(0)));
+		check(SelectionFilterSets.Num() >= NumInitialSelectionFilterSets);
+		for (int32 Index = 0; Index < NumInitialSelectionFilterSets; ++Index)
+		{
+			check(FindInput(GetRenderConnectionReference(Index)));
+			check(FindInput(GetSimConnectionReference(Index)));
+		}
 
-		for (int32 Index = 1; Index < SelectionFilterSets.Num(); ++Index)
+		for (int32 Index = NumInitialSelectionFilterSets; Index < SelectionFilterSets.Num(); ++Index)
 		{
-			RegisterInputArrayConnection(GetRenderConnectionReference(Index), GET_MEMBER_NAME_CHECKED(FChaosClothAssetConnectableIStringValue, StringValue));
-			RegisterInputArrayConnection(GetSimConnectionReference(Index), GET_MEMBER_NAME_CHECKED(FChaosClothAssetConnectableIStringValue, StringValue));
+			FindOrRegisterInputArrayConnection(GetRenderConnectionReference(Index), GET_MEMBER_NAME_CHECKED(FChaosClothAssetConnectableIStringValue, StringValue));
+			FindOrRegisterInputArrayConnection(GetSimConnectionReference(Index), GET_MEMBER_NAME_CHECKED(FChaosClothAssetConnectableIStringValue, StringValue));
+		}
+
+		if (Ar.IsTransacting())
+		{
+			const int32 OrigNumRegisteredInputs = GetNumInputs();
+			check(OrigNumRegisteredInputs >= NumRequiredInputs + NumInitialSelectionFilterSets * 2);
+			const int32 OrigNumSelectionFilterSets = SelectionFilterSets.Num();
+			const int32 OrigNumRegisteredSelectionFilterSets = (OrigNumRegisteredInputs - NumRequiredInputs) / 2;
+
+			if (OrigNumRegisteredSelectionFilterSets > OrigNumSelectionFilterSets)
+			{
+				ensure(Ar.IsTransacting());
+				// Temporarily expand SelectionFilterSets so we can get connection references.
+				SelectionFilterSets.SetNum(OrigNumRegisteredSelectionFilterSets);
+				for (int32 Index = OrigNumSelectionFilterSets; Index < SelectionFilterSets.Num(); ++Index)
+				{
+					UnregisterInputConnection(GetSimConnectionReference(Index));
+					UnregisterInputConnection(GetRenderConnectionReference(Index));
+				}
+				SelectionFilterSets.SetNum(OrigNumSelectionFilterSets);
+			}
+		}
+		else
+		{
+			ensureAlways(SelectionFilterSets.Num() * 2 + NumRequiredInputs == GetNumInputs());
 		}
 	}
 }
@@ -574,12 +607,15 @@ FChaosClothAssetProxyDeformerNode::FChaosClothAssetProxyDeformerNode(const Dataf
 	SimVertexSelection.StringValue = FString();  // An empty selection is an accepted input, but a non existing one isn't
 	SkinningBlendName = ClothCollectionAttribute::RenderDeformerSkinningBlend.ToString();
 
+	// If you change the number of InputConnections registered here, you must change the NumRequiredInputs in Serialize()
 	RegisterInputConnection(&Collection);
 	RegisterInputConnection(&SimVertexSelection.StringValue, GET_MEMBER_NAME_CHECKED(FChaosClothAssetConnectableIStringValue, StringValue));
 	RegisterInputConnection(&SelectionFilterSet0.StringValue, GET_MEMBER_NAME_CHECKED(FChaosClothAssetConnectableIStringValue, StringValue));
 	RegisterOutputConnection(&Collection)
 		.SetPassthroughInput(&Collection);
 	RegisterOutputConnection(&SkinningBlendName);
+
+	check(GetNumInputs() == NumRequiredInputs + NumInitialOptionalInputs); // Update NumRequiredInputs if you add more Inputs. This is used by Serialize.
 }
 
 void FChaosClothAssetProxyDeformerNode::Evaluate(Dataflow::FContext& Context, const FDataflowOutput* Out) const
@@ -655,7 +691,7 @@ void FChaosClothAssetProxyDeformerNode::Evaluate(Dataflow::FContext& Context, co
 
 TArray<Dataflow::FPin> FChaosClothAssetProxyDeformerNode::AddPins()
 {
-	check(NumFilterSets > 0);
+	check(NumFilterSets >= NumInitialOptionalInputs);
 	const FChaosClothAssetConnectableStringValue* const SelectionFilterSet = Get1To9SelectionFilterSets()[NumFilterSets - 1];
 
 	RegisterInputConnection(&SelectionFilterSet->StringValue, GET_MEMBER_NAME_CHECKED(FChaosClothAssetConnectableIStringValue, StringValue));
@@ -667,7 +703,7 @@ TArray<Dataflow::FPin> FChaosClothAssetProxyDeformerNode::AddPins()
 
 TArray<Dataflow::FPin> FChaosClothAssetProxyDeformerNode::GetPinsToRemove() const
 {
-	check(NumFilterSets > 1);
+	check(NumFilterSets > NumInitialOptionalInputs);
 	const FChaosClothAssetConnectableStringValue* const SelectionFilterSet = Get1To9SelectionFilterSets()[NumFilterSets - 2];
 	const FDataflowInput* const Input = FindInput(SelectionFilterSet);
 	check(Input);
@@ -676,7 +712,7 @@ TArray<Dataflow::FPin> FChaosClothAssetProxyDeformerNode::GetPinsToRemove() cons
 
 void FChaosClothAssetProxyDeformerNode::OnPinRemoved(const Dataflow::FPin& Pin)
 {
-	check(NumFilterSets > 1);
+	check(NumFilterSets > NumInitialOptionalInputs);
 	const FChaosClothAssetConnectableStringValue* const SelectionFilterSet = Get1To9SelectionFilterSets()[NumFilterSets - 2];
 	check(Pin.Direction == Dataflow::FPin::EDirection::INPUT);
 #if DO_CHECK
@@ -694,13 +730,30 @@ void FChaosClothAssetProxyDeformerNode::Serialize(FArchive& Ar)
 	// Restore the pins when re-loading so they can get properly reconnected
 	if (Ar.IsLoading())
 	{
-		const int32 NumFilterSetsToAdd = (NumFilterSets - 1);
-		NumFilterSets = 1;  // Reset to default, add pin will increment it again 
-		for (int32 Index = 0; Index < NumFilterSetsToAdd; ++Index)
+		const int32 OrigNumRegisteredInputs = GetNumInputs();
+		check(OrigNumRegisteredInputs >= NumRequiredInputs + NumInitialOptionalInputs);
+		const int32 OrigNumSelectionFilterSets = NumFilterSets;
+		const int32 OrigNumRegisteredSelectionFilterSets = OrigNumRegisteredInputs - NumRequiredInputs;
+		const int32 NumFilterSetsToAdd = (OrigNumSelectionFilterSets - OrigNumRegisteredSelectionFilterSets);
+		check(Ar.IsTransacting() || OrigNumRegisteredSelectionFilterSets == NumInitialOptionalInputs);
+		if (NumFilterSetsToAdd > 0)
 		{
-			AddPins();
+			NumFilterSets = OrigNumRegisteredSelectionFilterSets;  // Reset to default, add pin will increment it again 
+			for (int32 Index = 0; Index < NumFilterSetsToAdd; ++Index)
+			{
+				AddPins();
+			}
+		}
+		else if (NumFilterSetsToAdd < 0)
+		{
+			check(Ar.IsTransacting());
+			TArray<const FChaosClothAssetConnectableStringValue*> Non0SelectionFilterSets = Get1To9SelectionFilterSets();
+			for (int32 Index = NumFilterSets; Index < OrigNumRegisteredSelectionFilterSets; ++Index)
+			{
+				UnregisterInputConnection(Non0SelectionFilterSets[Index-1]);
+			}
 		}
-		ensure(NumFilterSetsToAdd == (NumFilterSets - 1));
+		check(NumFilterSets + NumRequiredInputs == GetNumInputs());
 	}
 }
 
diff --git a/Engine/Plugins/ChaosClothAssetEditor/Source/ChaosClothAssetDataflowNodes/Private/ChaosClothAsset/TerminalNode.cpp b/Engine/Plugins/ChaosClothAssetEditor/Source/ChaosClothAssetDataflowNodes/Private/ChaosClothAsset/TerminalNode.cpp
--- a/Engine/Plugins/ChaosClothAssetEditor/Source/ChaosClothAssetDataflowNodes/Private/ChaosClothAsset/TerminalNode.cpp
+++ b/Engine/Plugins/ChaosClothAssetEditor/Source/ChaosClothAssetDataflowNodes/Private/ChaosClothAsset/TerminalNode.cpp
@@ -91,7 +91,11 @@ FChaosClothAssetTerminalNode_v2::FChaosClothAssetTerminalNode_v2(const Dataflow:
 	: FDataflowTerminalNode(InParam, InGuid)
 {
 	// Start with Lod0
-	AddPins();
+	for (int32 Index = 0; Index < NumInitialCollectionLods; ++Index)
+	{
+		AddPins();
+	}
+	check(GetNumInputs() == NumRequiredInputs + NumInitialCollectionLods);  // Update NumRequiredInputs if you add more Inputs. This is used by Serialize.
 }
 
 void FChaosClothAssetTerminalNode_v2::SetAssetValue(TObjectPtr<UObject> Asset, Dataflow::FContext& Context) const
@@ -331,12 +335,37 @@ void FChaosClothAssetTerminalNode_v2::Serialize(FArchive& Ar)
 
 	if (Ar.IsLoading())
 	{
-		check(CollectionLods.Num() > 0);
-		check(FindInput(GetConnectionReference(0)));
+		check(CollectionLods.Num() >= NumInitialCollectionLods);
+		for (int32 Index = 0; Index < NumInitialCollectionLods; ++Index)
+		{
+			check(FindInput(GetConnectionReference(Index)));
+		}
 
-		for (int32 Index = 1; Index < CollectionLods.Num(); ++Index)
+		for (int32 Index = NumInitialCollectionLods; Index < CollectionLods.Num(); ++Index)
+		{
+			FindOrRegisterInputArrayConnection(GetConnectionReference(Index));
+		}
+		if (Ar.IsTransacting())
+		{
+			const int32 OrigNumRegisteredInputs = GetNumInputs();
+			check(OrigNumRegisteredInputs >= NumRequiredInputs + NumInitialCollectionLods);
+			const int32 OrigNumCollections = CollectionLods.Num();
+			const int32 OrigNumRegisteredCollections = (OrigNumRegisteredInputs - NumRequiredInputs);
+			if (OrigNumRegisteredCollections > OrigNumCollections)
+			{
+				// Inputs have been removed.
+				// Temporarily expand Collections so we can get connection references.
+				CollectionLods.SetNum(OrigNumRegisteredCollections);
+				for (int32 Index = OrigNumCollections; Index < CollectionLods.Num(); ++Index)
+				{
+					UnregisterInputConnection(GetConnectionReference(Index));
+				}
+				CollectionLods.SetNum(OrigNumCollections);
+			}
+		}
+		else
 		{
-			RegisterInputArrayConnection(GetConnectionReference(Index));
+			ensureAlways(CollectionLods.Num() + NumRequiredInputs == GetNumInputs());
 		}
 	}
 }
@@ -348,6 +377,7 @@ FChaosClothAssetTerminalNode::FChaosClothAssetTerminalNode(const Dataflow::FNode
 	: FDataflowTerminalNode(InParam, InGuid)
 {
 	RegisterInputConnection(&CollectionLod0);
+	check(NumInitialCollectionLods + NumRequiredInputs == GetNumInputs()); // Update NumRequiredInputs if you add more Inputs. This is used by Serialize.
 }
 
 void FChaosClothAssetTerminalNode::SetAssetValue(TObjectPtr<UObject> Asset, Dataflow::FContext& Context) const
@@ -602,6 +632,20 @@ TArray<const FManagedArrayCollection*> FChaosClothAssetTerminalNode::GetCollecti
 	return CollectionLods;
 }
 
+const FManagedArrayCollection* FChaosClothAssetTerminalNode::GetCollectionLod(int32 LodIndex) const
+{
+	switch (LodIndex)
+	{
+	case 0: return &CollectionLod0;
+	case 1: return &CollectionLod1; 
+	case 2: return &CollectionLod2; 
+	case 3: return &CollectionLod3; 
+	case 4: return &CollectionLod4; 
+	case 5: return &CollectionLod5; 
+	default: check(false); return nullptr;
+	}
+}
+
 TArray<TSharedRef<FManagedArrayCollection>> FChaosClothAssetTerminalNode::GetCleanedCollectionLodValues(Dataflow::FContext& Context) const
 {
 	using namespace UE::Chaos::ClothAsset;
@@ -657,13 +701,30 @@ void FChaosClothAssetTerminalNode::Serialize(FArchive& Ar)
 
 	if (Ar.IsLoading())
 	{
-		const int32 NumLodToAdd = (NumLods - 1);
-		NumLods = 1; // need to reset back to default because add pin will increment it again 
-		for (int32 LodIndex = 0; LodIndex < NumLodToAdd; ++LodIndex)
+		const int32 OrigNumRegisteredInputs = GetNumInputs();
+		check(OrigNumRegisteredInputs >= NumRequiredInputs + NumInitialCollectionLods);
+		const int32 OrigNumLods = NumLods;
+		const int32 OrigNumRegisteredLods = OrigNumRegisteredInputs - NumRequiredInputs;
+		const int32 NumLodToAdd = (OrigNumLods - OrigNumRegisteredLods);
+		check(Ar.IsTransacting() || OrigNumRegisteredLods == NumInitialCollectionLods);
+		if (NumLodToAdd > 0)
+		{
+			NumLods = OrigNumRegisteredLods;  // AddPin will increment it again
+			for (int32 LodIndex = 0; LodIndex < NumLodToAdd; ++LodIndex)
+			{
+				AddPins();
+			}
+		}
+		else if (NumLodToAdd < 0)
 		{
-			AddPins();
+			check(Ar.IsTransacting());
+			for (int32 Index = NumLods; Index < OrigNumRegisteredLods; ++Index)
+			{
+				UnregisterInputConnection(GetCollectionLod(Index));
+			}
+
 		}
-		ensure(NumLodToAdd == (NumLods - 1));
+		check(NumLods + NumRequiredInputs == GetNumInputs());
 	}
 }
 #undef LOCTEXT_NAMESPACE
diff --git a/Engine/Plugins/ChaosClothAssetEditor/Source/ChaosClothAssetDataflowNodes/Private/ChaosClothAsset/WeightedValueCustomization.cpp b/Engine/Plugins/ChaosClothAssetEditor/Source/ChaosClothAssetDataflowNodes/Private/ChaosClothAsset/WeightedValueCustomization.cpp
--- a/Engine/Plugins/ChaosClothAssetEditor/Source/ChaosClothAssetDataflowNodes/Private/ChaosClothAsset/WeightedValueCustomization.cpp
+++ b/Engine/Plugins/ChaosClothAssetEditor/Source/ChaosClothAssetDataflowNodes/Private/ChaosClothAsset/WeightedValueCustomization.cpp
@@ -211,7 +211,7 @@ PRAGMA_ENABLE_DEPRECATION_WARNINGS
 				{
 					if (bIsUsingSlider)
 					{
-						WeakHandlePtr.Pin()->SetValue(Value, EPropertyValueSetFlags::InteractiveChange);
+						WeakHandlePtr.Pin()->SetValue(Value, EPropertyValueSetFlags::InteractiveChange | EPropertyValueSetFlags::NotTransactable);
 					}
 				})
 			.OnBeginSliderMovement_Lambda([this]()
diff --git a/Engine/Plugins/ChaosClothAssetEditor/Source/ChaosClothAssetDataflowNodes/Public/ChaosClothAsset/MergeClothCollectionsNode.h b/Engine/Plugins/ChaosClothAssetEditor/Source/ChaosClothAssetDataflowNodes/Public/ChaosClothAsset/MergeClothCollectionsNode.h
--- a/Engine/Plugins/ChaosClothAssetEditor/Source/ChaosClothAssetDataflowNodes/Public/ChaosClothAsset/MergeClothCollectionsNode.h
+++ b/Engine/Plugins/ChaosClothAssetEditor/Source/ChaosClothAssetDataflowNodes/Public/ChaosClothAsset/MergeClothCollectionsNode.h
@@ -27,12 +27,15 @@ private:
 	virtual void Evaluate(Dataflow::FContext& Context, const FDataflowOutput* Out) const override;
 	virtual TArray<Dataflow::FPin> AddPins() override;
 	virtual bool CanAddPin() const override { return true; }
-	virtual bool CanRemovePin() const override { return Collections.Num() > 2; }
+	virtual bool CanRemovePin() const override { return Collections.Num() > NumInitialOptionalInputs; }
 	virtual TArray<Dataflow::FPin> GetPinsToRemove() const override;
 	virtual void OnPinRemoved(const Dataflow::FPin& Pin) override;
 	virtual void Serialize(FArchive& Ar) override;
 	//~ End FDataflowNode interface
 
+
+	static constexpr int32 NumRequiredInputs = 0;
+	static constexpr int32 NumInitialOptionalInputs = 2;
 	Dataflow::TConnectionReference<FManagedArrayCollection> GetConnectionReference(int32 Index) const;
 };
 
@@ -68,7 +71,7 @@ public:
 	FManagedArrayCollection Collection5;
 	/** The number of inputs currently exposed to the node UI. */
 	UPROPERTY()
-	int32 NumInputs = 1;
+	int32 NumInputs = NumInitialOptionalInputs;
 
 	FChaosClothAssetMergeClothCollectionsNode(const Dataflow::FNodeParameters& InParam, FGuid InGuid = FGuid::NewGuid());
 
@@ -77,11 +80,15 @@ private:
 	virtual void Evaluate(Dataflow::FContext& Context, const FDataflowOutput* Out) const override;
 	virtual TArray<Dataflow::FPin> AddPins() override;
 	virtual bool CanAddPin() const override { return NumInputs < MaxInputs; }
-	virtual bool CanRemovePin() const override { return NumInputs > 1; }
+	virtual bool CanRemovePin() const override { return NumInputs > NumInitialOptionalInputs; }
 	virtual TArray<Dataflow::FPin> GetPinsToRemove() const override;
 	virtual void OnPinRemoved(const Dataflow::FPin& Pin) override;
 	virtual void Serialize(FArchive& Ar) override;
 	//~ End FDataflowNode interface
 
 	TArray<const FManagedArrayCollection*> GetCollections() const;
+	const FManagedArrayCollection* GetCollection(int32 Index) const;
+
+	static constexpr int32 NumRequiredInputs = 0;
+	static constexpr int32 NumInitialOptionalInputs = 1;
 };
diff --git a/Engine/Plugins/ChaosClothAssetEditor/Source/ChaosClothAssetDataflowNodes/Public/ChaosClothAsset/ProxyDeformerNode.h b/Engine/Plugins/ChaosClothAssetEditor/Source/ChaosClothAssetDataflowNodes/Public/ChaosClothAsset/ProxyDeformerNode.h
--- a/Engine/Plugins/ChaosClothAssetEditor/Source/ChaosClothAssetDataflowNodes/Public/ChaosClothAsset/ProxyDeformerNode.h
+++ b/Engine/Plugins/ChaosClothAssetEditor/Source/ChaosClothAssetDataflowNodes/Public/ChaosClothAsset/ProxyDeformerNode.h
@@ -78,7 +78,7 @@ private:
 	virtual void Evaluate(Dataflow::FContext& Context, const FDataflowOutput* Out) const override;
 	virtual TArray<Dataflow::FPin> AddPins() override;
 	virtual bool CanAddPin() const override { return true; }
-	virtual bool CanRemovePin() const override { return SelectionFilterSets.Num() > 1; }
+	virtual bool CanRemovePin() const override { return SelectionFilterSets.Num() > NumInitialSelectionFilterSets; }
 	virtual TArray<Dataflow::FPin> GetPinsToRemove() const override;
 	virtual void OnPinRemoved(const Dataflow::FPin& Pin) override;
 	virtual void Serialize(FArchive& Ar) override;
@@ -87,6 +87,9 @@ private:
 	TArray<TPair<FName, FName>> GetSelectionFilterNames(Dataflow::FContext& Context) const;
 	Dataflow::TConnectionReference<FString> GetRenderConnectionReference(int32 Index) const;
 	Dataflow::TConnectionReference<FString> GetSimConnectionReference(int32 Index) const;
+
+	static constexpr int32 NumRequiredInputs = 2; // non-filter set inputs
+	static constexpr int32 NumInitialSelectionFilterSets = 1;
 };
 
 /** Add the proxy deformer information to this cloth collection's render data. */
@@ -182,13 +185,13 @@ private:
 
 	/** The number of filter sets currently exposed to the node UI. */
 	UPROPERTY()
-	int32 NumFilterSets = 1;
+	int32 NumFilterSets = NumInitialOptionalInputs;
 
 	//~ Begin FDataflowNode interface
 	virtual void Evaluate(Dataflow::FContext& Context, const FDataflowOutput* Out) const override;
 	virtual TArray<Dataflow::FPin> AddPins() override;
 	virtual bool CanAddPin() const override { return NumFilterSets < MaxNumFilterSets; }
-	virtual bool CanRemovePin() const override { return NumFilterSets > 1; }
+	virtual bool CanRemovePin() const override { return NumFilterSets > NumInitialOptionalInputs; }
 	virtual TArray<Dataflow::FPin> GetPinsToRemove() const override;
 	virtual void OnPinRemoved(const Dataflow::FPin& Pin) override;
 	virtual void Serialize(FArchive& Ar) override;
@@ -196,4 +199,7 @@ private:
 
 	TArray<FName> GetSelectionFilterNames(Dataflow::FContext& Context) const;
 	TArray<const FChaosClothAssetConnectableStringValue*> Get1To9SelectionFilterSets() const;
+
+	static constexpr int32 NumRequiredInputs = 2; // non-filter set inputs
+	static constexpr int32 NumInitialOptionalInputs = 1; // filter set inputs that are created in the constructor.
 };
diff --git a/Engine/Plugins/ChaosClothAssetEditor/Source/ChaosClothAssetDataflowNodes/Public/ChaosClothAsset/TerminalNode.h b/Engine/Plugins/ChaosClothAssetEditor/Source/ChaosClothAssetDataflowNodes/Public/ChaosClothAsset/TerminalNode.h
--- a/Engine/Plugins/ChaosClothAssetEditor/Source/ChaosClothAssetDataflowNodes/Public/ChaosClothAsset/TerminalNode.h
+++ b/Engine/Plugins/ChaosClothAssetEditor/Source/ChaosClothAssetDataflowNodes/Public/ChaosClothAsset/TerminalNode.h
@@ -44,7 +44,7 @@ private:
 	virtual void Evaluate(Dataflow::FContext& Context, const FDataflowOutput* Out) const override {}
 	virtual TArray<Dataflow::FPin> AddPins() override;
 	virtual bool CanAddPin() const override { return true; }
-	virtual bool CanRemovePin() const override { return CollectionLods.Num() > 1; }
+	virtual bool CanRemovePin() const override { return CollectionLods.Num() > NumInitialCollectionLods; }
 	virtual TArray<Dataflow::FPin> GetPinsToRemove() const override;
 	virtual void OnPinRemoved(const Dataflow::FPin& Pin) override;
 	virtual void Serialize(FArchive& Ar) override;
@@ -59,6 +59,8 @@ private:
 	// This is for runtime only--used to determine if only properties need to be updated.
 	mutable bool bClothCollectionChecksumValid = false;
 	mutable uint32 ClothColllectionChecksum = 0;
+	static constexpr int32 NumRequiredInputs = 0;
+	static constexpr int32 NumInitialCollectionLods = 1;
 };
 
 
@@ -93,7 +95,7 @@ public:
 	FManagedArrayCollection CollectionLod5;
 	/** The number of LODs currently exposed to the node UI. */
 	UPROPERTY()
-	int32 NumLods = 1;
+	int32 NumLods = NumInitialCollectionLods;
 	/**
 	 * Refresh the asset even if the ClothCollection hasn't changed.
 	 * Note that it is not required to manually refresh the cloth asset, this is done automatically when there is a change in the Dataflow.
@@ -118,6 +120,7 @@ private:
 
 	TArray<const FManagedArrayCollection*> GetCollectionLods() const;
 	TArray<TSharedRef<FManagedArrayCollection>> GetCleanedCollectionLodValues(Dataflow::FContext& Context) const;
+	const FManagedArrayCollection* GetCollectionLod(int32 LodIndex) const;
 
 	UPROPERTY()
 	mutable TArray<FChaosClothAssetLodTransitionDataCache> LODTransitionDataCache;
@@ -125,4 +128,7 @@ private:
 	// This is for runtime only--used to determine if only properties need to be updated.
 	mutable bool bClothCollectionChecksumValid = false;
 	mutable uint32 ClothColllectionChecksum = 0;
+
+	static constexpr int32 NumRequiredInputs = 0;
+	static constexpr int32 NumInitialCollectionLods = 1;
 };
diff --git a/Engine/Plugins/ChaosClothAssetEditor/Source/ChaosClothAssetEditor/Private/ChaosClothAsset/ClothEditorToolkit.cpp b/Engine/Plugins/ChaosClothAssetEditor/Source/ChaosClothAssetEditor/Private/ChaosClothAsset/ClothEditorToolkit.cpp
--- a/Engine/Plugins/ChaosClothAssetEditor/Source/ChaosClothAssetEditor/Private/ChaosClothAsset/ClothEditorToolkit.cpp
+++ b/Engine/Plugins/ChaosClothAssetEditor/Source/ChaosClothAssetEditor/Private/ChaosClothAsset/ClothEditorToolkit.cpp
@@ -209,12 +209,13 @@ FChaosClothAssetEditorToolkit::FChaosClothAssetEditorToolkit(UAssetEditor* InOwn
 FChaosClothAssetEditorToolkit::~FChaosClothAssetEditorToolkit()
 {
 	// This code is also called in OnRequestClose. Leaving this here as well in case the toolkit gets destroyed without having OnRequestClose called
+	TSharedPtr<FDataflowNode> SelectedDataflowNode = GetSelectedDataflowNode();
 	if (SelectedDataflowNode && OnNodeInvalidatedDelegateHandle.IsValid())
 	{
 		SelectedDataflowNode->GetOnNodeInvalidatedDelegate().Remove(OnNodeInvalidatedDelegateHandle);
 		SelectedDataflowNode->OnDeselected();
 	}
-	SelectedDataflowNode.Reset();
+	SelectedDataflowNodeGuid.Invalidate();
 
 	if (ClothPreviewViewportClient)
 	{
@@ -392,12 +393,13 @@ bool FChaosClothAssetEditorToolkit::OnRequestClose(EAssetEditorCloseReason InClo
 		return true;
 	}
 
+	TSharedPtr<FDataflowNode> SelectedDataflowNode = GetSelectedDataflowNode();
 	if (SelectedDataflowNode && OnNodeInvalidatedDelegateHandle.IsValid())
 	{
 		SelectedDataflowNode->GetOnNodeInvalidatedDelegate().Remove(OnNodeInvalidatedDelegateHandle);
 		SelectedDataflowNode->OnDeselected();
 	}
-	SelectedDataflowNode.Reset();
+	SelectedDataflowNodeGuid.Invalidate();
 
 	FCoreUObjectDelegates::OnPackageReloaded.Remove(OnPackageReloadedDelegateHandle);
 
@@ -708,6 +710,13 @@ void FChaosClothAssetEditorToolkit::UnregisterTabSpawners(const TSharedRef<FTabM
 
 //~ End IToolkit overrides
 
+void FChaosClothAssetEditorToolkit::NotifyPreChange(FEditPropertyChain* PropertyAboutToChange)
+{
+	if (UDataflow* const DataflowAsset = GetDataflow())
+	{
+		FDataflowEditorCommands::OnNotifyPropertyPreChange(NodeDetailsEditor, DataflowAsset, PropertyAboutToChange);
+	}
+}
 
 UChaosClothAsset* FChaosClothAssetEditorToolkit::GetAsset() const
 {
@@ -969,7 +978,7 @@ void FChaosClothAssetEditorToolkit::EvaluateNode(FDataflowNode* Node, bool bForc
 			OnClothAssetChanged();
 
 			// Refresh the construction viewport
-			if (SelectedDataflowNode)
+			if (TSharedPtr<FDataflowNode> SelectedDataflowNode = GetSelectedDataflowNode())
 			{
 				TSharedPtr<FManagedArrayCollection> InputCollection = GetInputClothCollectionIfPossible(SelectedDataflowNode, DataflowContext);
 				TSharedPtr<FManagedArrayCollection> Collection = GetClothCollectionIfPossible(SelectedDataflowNode, DataflowContext);
@@ -1068,7 +1077,7 @@ TSharedPtr<IStructureDetailsView> FChaosClothAssetEditorToolkit::CreateNodeDetai
 		DetailsViewArgs.bLockable = false;
 		DetailsViewArgs.bSearchInitialKeyFocus = true;
 		DetailsViewArgs.bUpdatesFromSelection = false;
-		DetailsViewArgs.NotifyHook = nullptr;
+		DetailsViewArgs.NotifyHook = this;
 		DetailsViewArgs.bShowOptions = true;
 		DetailsViewArgs.bShowModifiedPropertiesOption = false;
 		DetailsViewArgs.bShowScrollBar = false;
@@ -1176,6 +1185,35 @@ TSharedPtr<FManagedArrayCollection> FChaosClothAssetEditorToolkit::GetInputCloth
 	return TSharedPtr<FManagedArrayCollection>();
 }
 
+TSharedPtr<FDataflowNode> FChaosClothAssetEditorToolkit::GetSelectedDataflowNode()
+{
+	if (SelectedDataflowNodeGuid.IsValid())
+	{
+		if (UDataflow* const Dataflow = GetDataflow())
+		{
+			if (TSharedPtr<Dataflow::FGraph> Graph = Dataflow->GetDataflow())
+			{
+				return Graph->FindBaseNode(SelectedDataflowNodeGuid);
+			}
+		}
+	}
+	return TSharedPtr<FDataflowNode>(nullptr);
+}
+
+TSharedPtr<const FDataflowNode> FChaosClothAssetEditorToolkit::GetSelectedDataflowNode() const
+{
+	if (SelectedDataflowNodeGuid.IsValid())
+	{
+		if (const UDataflow* const Dataflow = GetDataflow())
+		{
+			if (TSharedPtr<const Dataflow::FGraph> Graph = Dataflow->GetDataflow())
+			{
+				return Graph->FindBaseNode(SelectedDataflowNodeGuid);
+			}
+		}
+	}
+	return TSharedPtr<const FDataflowNode>(nullptr);
+}
 
 void FChaosClothAssetEditorToolkit::OnNodeSelectionChanged(const TSet<UObject*>& NewSelection)
 {
@@ -1199,7 +1237,7 @@ void FChaosClothAssetEditorToolkit::OnNodeSelectionChanged(const TSet<UObject*>&
 		{
 			// No new node selected
 
-			if (SelectedDataflowNode)
+			if (TSharedPtr<FDataflowNode> SelectedDataflowNode = GetSelectedDataflowNode())
 			{
 				bNodeSelectionChanged = true;		// current node was deselected
 
@@ -1208,7 +1246,7 @@ void FChaosClothAssetEditorToolkit::OnNodeSelectionChanged(const TSet<UObject*>&
 					SelectedDataflowNode->GetOnNodeInvalidatedDelegate().Remove(OnNodeInvalidatedDelegateHandle);
 				}
 				SelectedDataflowNode->OnDeselected();
-				SelectedDataflowNode.Reset();
+				SelectedDataflowNodeGuid.Invalidate();
 			}
 		}
 		else
@@ -1219,6 +1257,7 @@ void FChaosClothAssetEditorToolkit::OnNodeSelectionChanged(const TSet<UObject*>&
 				{
 					Dataflow->RenderTargets.Add(Node);
 
+					TSharedPtr<FDataflowNode> SelectedDataflowNode = GetSelectedDataflowNode();
 					if (SelectedDataflowNode != Node->GetDataflowNode())
 					{
 						bNodeSelectionChanged = true;	// changing from one selected node to a different one, or from unselected to selected
@@ -1232,6 +1271,7 @@ void FChaosClothAssetEditorToolkit::OnNodeSelectionChanged(const TSet<UObject*>&
 							SelectedDataflowNode->OnDeselected();
 						}
 						SelectedDataflowNode = Node->GetDataflowNode();
+						SelectedDataflowNodeGuid = SelectedDataflowNode->GetGuid();
 
 						if (SelectedDataflowNode)
 						{
@@ -1248,6 +1288,7 @@ void FChaosClothAssetEditorToolkit::OnNodeSelectionChanged(const TSet<UObject*>&
 								{
 									// Warning: Do not execute code that rebuilds the UI in this lambda as it is called by the UI!
 
+									TSharedPtr<FDataflowNode> SelectedDataflowNode = GetSelectedDataflowNode();
 									if (SelectedDataflowNode.Get() == InDataflowNode)
 									{
 										GetClothCollectionIfPossible(SelectedDataflowNode, DataflowContext);
@@ -1431,7 +1472,7 @@ void FChaosClothAssetEditorToolkit::HandlePackageReloaded(const EPackageReloadPh
 				}
 
 				// Eliminate anything could be holding a reference to the Dataflow object that will be reloaded (including references to nodes)
-				SelectedDataflowNode.Reset();
+				SelectedDataflowNodeGuid.Invalidate();
 				GraphEditor.Reset();
 				GraphEditorTab.Get()->SetContent(SNew(SSpacer));
 			}
diff --git a/Engine/Plugins/ChaosClothAssetEditor/Source/ChaosClothAssetEditor/Private/ChaosClothAsset/ClothEditorToolkit.h b/Engine/Plugins/ChaosClothAssetEditor/Source/ChaosClothAssetEditor/Private/ChaosClothAsset/ClothEditorToolkit.h
--- a/Engine/Plugins/ChaosClothAssetEditor/Source/ChaosClothAssetEditor/Private/ChaosClothAsset/ClothEditorToolkit.h
+++ b/Engine/Plugins/ChaosClothAssetEditor/Source/ChaosClothAssetEditor/Private/ChaosClothAsset/ClothEditorToolkit.h
@@ -47,7 +47,7 @@ namespace UE::Chaos::ClothAsset
  * initializing the Cloth mode.
  * Thus, the FChaosClothAssetEditorToolkit ends up being the central place for the Cloth Asset Editor setup.
  */
-class CHAOSCLOTHASSETEDITOR_API FChaosClothAssetEditorToolkit final : public FBaseCharacterFXEditorToolkit, public FTickableEditorObject
+class CHAOSCLOTHASSETEDITOR_API FChaosClothAssetEditorToolkit final : public FBaseCharacterFXEditorToolkit, public FTickableEditorObject, public FNotifyHook
 {
 public:
 
@@ -103,6 +103,9 @@ private:
 	virtual void RegisterTabSpawners(const TSharedRef<FTabManager>& InTabManager) override;
 	virtual void UnregisterTabSpawners(const TSharedRef<FTabManager>& TabManager) override;
 
+	// FNotifyHook
+	virtual void NotifyPreChange(class FEditPropertyChain* PropertyAboutToChange) override;
+
 	// Return the cloth asset held by the Cloth Editor
 	UChaosClothAsset* GetAsset() const;
 	UDataflow* GetDataflow();
@@ -129,8 +132,9 @@ private:
 	TSharedPtr<IStructureDetailsView> CreateNodeDetailsEditorWidget(UObject* ObjectToEdit);
 
 	TSharedPtr<FManagedArrayCollection> GetClothCollectionIfPossible(const TSharedPtr<FDataflowNode> InDataflowNode, const TSharedPtr<Dataflow::FEngineContext> Context);
-
 	TSharedPtr<FManagedArrayCollection> GetInputClothCollectionIfPossible(const TSharedPtr<FDataflowNode> InDataflowNode, const TSharedPtr<Dataflow::FEngineContext> Context);
+	TSharedPtr<FDataflowNode> GetSelectedDataflowNode();
+	TSharedPtr<const FDataflowNode> GetSelectedDataflowNode() const;
 
 	// DataflowEditorActions
 	void OnPropertyValueChanged(const FPropertyChangedEvent& PropertyChangedEvent);
@@ -166,7 +170,7 @@ private:
 	TSharedPtr<Dataflow::FEngineContext> DataflowContext;
 	Dataflow::FTimestamp LastDataflowNodeTimestamp = Dataflow::FTimestamp::Invalid;
 	FDelegateHandle OnNodeInvalidatedDelegateHandle;
-	TSharedPtr<FDataflowNode> SelectedDataflowNode;
+	FGuid SelectedDataflowNodeGuid;
 
 	static const FName GraphCanvasTabId;
 	TSharedPtr<SDockTab> GraphEditorTab;
diff --git a/Engine/Plugins/Experimental/Dataflow/Source/DataflowEditor/Private/Dataflow/DataflowEditorCommands.cpp b/Engine/Plugins/Experimental/Dataflow/Source/DataflowEditor/Private/Dataflow/DataflowEditorCommands.cpp
--- a/Engine/Plugins/Experimental/Dataflow/Source/DataflowEditor/Private/Dataflow/DataflowEditorCommands.cpp
+++ b/Engine/Plugins/Experimental/Dataflow/Source/DataflowEditor/Private/Dataflow/DataflowEditorCommands.cpp
@@ -315,6 +315,36 @@ void FDataflowEditorCommands::OnNodeTitleCommitted(const FText& InNewText, EText
 	}
 }
 
+void FDataflowEditorCommands::OnNotifyPropertyPreChange(TSharedPtr<IStructureDetailsView> PropertiesEditor, UDataflow* Graph, class FEditPropertyChain* PropertyAboutToChange)
+{
+	// Find the associated UDataflowEdNode(s) and call Modify on them for Undo/Redo.
+	if (PropertiesEditor && Graph)
+	{
+		if (TSharedPtr<const IStructureDataProvider> StructProvider = PropertiesEditor->GetStructureProvider())
+		{
+			const UStruct* const BaseStruct = StructProvider->GetBaseStructure();
+			if (BaseStruct && BaseStruct->IsChildOf(FDataflowNode::StaticStruct()))
+			{
+				TArray<TSharedPtr<FStructOnScope>> StructData;
+				StructProvider->GetInstances(StructData, FDataflowNode::StaticStruct());
+				for (const TSharedPtr<FStructOnScope>& StructOnScope : StructData)
+				{
+					if (StructOnScope.IsValid() && StructOnScope->IsValid())
+					{
+						const FDataflowNode* const Node = reinterpret_cast<FDataflowNode*>(StructOnScope->GetStructMemory());
+						check(Node);
+						if (const TObjectPtr<UDataflowEdNode> EdNode = Graph->FindEdNodeByDataflowNodeGuid(Node->GetGuid()))
+						{
+							EdNode->Modify();
+						}
+					}
+				}
+			}
+		}
+	}
+}
+
+
 void FDataflowEditorCommands::OnAssetPropertyValueChanged(TObjectPtr<UDataflowBaseContent> Content, const FPropertyChangedEvent& InPropertyChangedEvent)
 {
 	if (Content)
diff --git a/Engine/Plugins/Experimental/Dataflow/Source/DataflowEditor/Private/Dataflow/DataflowEditorToolkit.cpp b/Engine/Plugins/Experimental/Dataflow/Source/DataflowEditor/Private/Dataflow/DataflowEditorToolkit.cpp
--- a/Engine/Plugins/Experimental/Dataflow/Source/DataflowEditor/Private/Dataflow/DataflowEditorToolkit.cpp
+++ b/Engine/Plugins/Experimental/Dataflow/Source/DataflowEditor/Private/Dataflow/DataflowEditorToolkit.cpp
@@ -315,6 +315,18 @@ void FDataflowEditorToolkit::CreateEditorModeManager()
 		= StaticCastSharedPtr<FAssetEditorModeManager>(SimulationModeManager);
 }
 
+void FDataflowEditorToolkit::NotifyPreChange(FEditPropertyChain* PropertyAboutToChange)
+{
+	if (const TObjectPtr<UDataflowBaseContent>& EditorContent = GetEditorContent())
+	{
+		ensure(EditorContent);
+		if (UDataflow* DataflowAsset = EditorContent->GetDataflowAsset())
+		{
+			FDataflowEditorCommands::OnNotifyPropertyPreChange(NodeDetailsEditor, DataflowAsset, PropertyAboutToChange);
+		}
+	}
+}
+
 bool FDataflowEditorToolkit::CanOpenDataflowEditor(UObject* ObjectToEdit)
 {
 	if (const UClass* Class = ObjectToEdit->GetClass())
@@ -1118,7 +1130,7 @@ TSharedPtr<IStructureDetailsView> FDataflowEditorToolkit::CreateNodeDetailsEdito
 		DetailsViewArgs.bLockable = false;
 		DetailsViewArgs.bSearchInitialKeyFocus = true;
 		DetailsViewArgs.bUpdatesFromSelection = false;
-		DetailsViewArgs.NotifyHook = nullptr;
+		DetailsViewArgs.NotifyHook = this;
 		DetailsViewArgs.bShowOptions = true;
 		DetailsViewArgs.bShowModifiedPropertiesOption = false;
 		DetailsViewArgs.bShowScrollBar = false;
diff --git a/Engine/Plugins/Experimental/Dataflow/Source/DataflowEditor/Public/Dataflow/DataflowEditorCommands.h b/Engine/Plugins/Experimental/Dataflow/Source/DataflowEditor/Public/Dataflow/DataflowEditorCommands.h
--- a/Engine/Plugins/Experimental/Dataflow/Source/DataflowEditor/Public/Dataflow/DataflowEditorCommands.h
+++ b/Engine/Plugins/Experimental/Dataflow/Source/DataflowEditor/Public/Dataflow/DataflowEditorCommands.h
@@ -126,6 +126,11 @@ public:
 	*/
 	static void OnNodeTitleCommitted(const FText& InNewText, ETextCommit::Type InCommitType, UEdGraphNode* GraphNode);
 
+	/*
+	* OnNotifyPropertyPreChange
+	*/
+	static void OnNotifyPropertyPreChange(TSharedPtr<IStructureDetailsView> PropertiesEditor, UDataflow* Graph, class FEditPropertyChain* PropertyAboutToChange);
+
 	/*
 	*  OnPropertyValueChanged
 	*/
diff --git a/Engine/Plugins/Experimental/Dataflow/Source/DataflowEditor/Public/Dataflow/DataflowEditorToolkit.h b/Engine/Plugins/Experimental/Dataflow/Source/DataflowEditor/Public/Dataflow/DataflowEditorToolkit.h
--- a/Engine/Plugins/Experimental/Dataflow/Source/DataflowEditor/Public/Dataflow/DataflowEditorToolkit.h
+++ b/Engine/Plugins/Experimental/Dataflow/Source/DataflowEditor/Public/Dataflow/DataflowEditorToolkit.h
@@ -125,6 +125,9 @@ private:
 	virtual TSharedPtr<FEditorViewportClient> CreateEditorViewportClient() const override;
 	virtual void CreateEditorModeManager() override;
 
+	// FNotifyHook
+	virtual void NotifyPreChange(class FEditPropertyChain* PropertyAboutToChange) override;
+
 	// List of all the tab names ids that will be used to identify the editor widgets
 	static const FName GraphCanvasTabId;
 	static const FName NodeDetailsTabId;
diff --git a/Engine/Source/Editor/PropertyEditor/Private/SStructureDetailsView.h b/Engine/Source/Editor/PropertyEditor/Private/SStructureDetailsView.h
--- a/Engine/Source/Editor/PropertyEditor/Private/SStructureDetailsView.h
+++ b/Engine/Source/Editor/PropertyEditor/Private/SStructureDetailsView.h
@@ -73,6 +73,8 @@ public:
 		return OnFinishedChangingProperties();
 	}
 
+	TSharedPtr<const IStructureDataProvider> GetStructureProvider() const override { return StructProvider; }
+
 public:
 	/** Causes the details view to be refreshed (new widgets generated) with the current set of objects */
 	virtual void ForceRefresh() override;
diff --git a/Engine/Source/Editor/PropertyEditor/Public/IStructureDetailsView.h b/Engine/Source/Editor/PropertyEditor/Public/IStructureDetailsView.h
--- a/Engine/Source/Editor/PropertyEditor/Public/IStructureDetailsView.h
+++ b/Engine/Source/Editor/PropertyEditor/Public/IStructureDetailsView.h
@@ -51,6 +51,15 @@ public:
 	 */
 	virtual void SetCustomName(const FText& Text) = 0;
 
+	/*
+	* Get Const access to the StructureProvider.
+	* Not pure virtual for backwards compatibility reasons.
+	*/
+	virtual TSharedPtr<const class IStructureDataProvider> GetStructureProvider() const
+	{
+		return nullptr;
+	}
+
 public:
 
 	/**
diff --git a/Engine/Source/Runtime/Experimental/Dataflow/Core/Private/Dataflow/DataflowCoreNodes.cpp b/Engine/Source/Runtime/Experimental/Dataflow/Core/Private/Dataflow/DataflowCoreNodes.cpp
--- a/Engine/Source/Runtime/Experimental/Dataflow/Core/Private/Dataflow/DataflowCoreNodes.cpp
+++ b/Engine/Source/Runtime/Experimental/Dataflow/Core/Private/Dataflow/DataflowCoreNodes.cpp
@@ -90,11 +90,14 @@ FDataflowSelectNode::FDataflowSelectNode(const Dataflow::FNodeParameters& Param,
 	: Super(Param, InGuid)
 {
 	// Add two sets of pins to start.
-	AddPins();
-	AddPins();
 	RegisterInputConnection(&SelectedIndex);
+	for (int32 Index = 0; Index < NumInitialInputs; ++Index)
+	{
+		AddPins();
+	}
 	RegisterOutputConnection(&Result)
 		.SetPassthroughInput(GetConnectionReference(0));
+	check(NumRequiredDataflowInputs + NumInitialInputs == GetNumInputs()); // Update NumRequiredDataflowInputs when adding more inputs. This is used by Serialize
 }
 
 void FDataflowSelectNode::Evaluate(Dataflow::FContext& Context, const FDataflowOutput* Out) const
@@ -182,13 +185,37 @@ void FDataflowSelectNode::Serialize(FArchive& Ar)
 {
 	if (Ar.IsLoading())
 	{
-		check(Inputs.Num() > 1);
-		check(FindInput(GetConnectionReference(0)));
-		check(FindInput(GetConnectionReference(1)));
+		check(Inputs.Num() >= NumInitialInputs);
+		for (int32 Index = 0; Index < NumInitialInputs; ++Index)
+		{
+			check(FindInput(GetConnectionReference(Index)));
+		}
+
+		for (int32 Index = NumInitialInputs; Index < Inputs.Num(); ++Index)
+		{
+			FindOrRegisterInputArrayConnection(GetConnectionReference(Index));
+		}
 
-		for (int32 Index = 2; Index < Inputs.Num(); ++Index)
+		if (Ar.IsTransacting())
+		{
+			const int32 OrigNumRegisteredInputs = GetNumInputs() - NumRequiredDataflowInputs;
+			const int32 OrigNumInputs = Inputs.Num();
+			if (OrigNumRegisteredInputs > OrigNumInputs)
+			{
+				// Inputs have been removed.
+				// Temporarily expand Inputs so we can get connection references.
+				Inputs.SetNum(OrigNumRegisteredInputs);
+				for (int32 Index = OrigNumInputs; Index < Inputs.Num(); ++Index)
+				{
+					UnregisterInputConnection(GetConnectionReference(Index));
+				}
+				Inputs.SetNum(OrigNumInputs);
+			}
+		}
+		else
 		{
-			RegisterInputArrayConnection(GetConnectionReference(Index));
+			// Index + all Inputs
+			ensureAlways(Inputs.Num() + NumRequiredDataflowInputs == GetNumInputs());
 		}
 	}
 }
diff --git a/Engine/Source/Runtime/Experimental/Dataflow/Core/Private/Dataflow/DataflowNode.cpp b/Engine/Source/Runtime/Experimental/Dataflow/Core/Private/Dataflow/DataflowNode.cpp
--- a/Engine/Source/Runtime/Experimental/Dataflow/Core/Private/Dataflow/DataflowNode.cpp
+++ b/Engine/Source/Runtime/Experimental/Dataflow/Core/Private/Dataflow/DataflowNode.cpp
@@ -147,6 +147,11 @@ void FDataflowNode::AddInput(FDataflowInput* InPtr)
 	}
 }
 
+int32 FDataflowNode::GetNumInputs() const
+{
+	return ExpandedInputs.Num();
+}
+
 FDataflowInput* FDataflowNode::FindInput(FName InName)
 {
 	for (TPair<Dataflow::FConnectionKey, FDataflowInput*>& Elem : ExpandedInputs)
diff --git a/Engine/Source/Runtime/Experimental/Dataflow/Core/Public/Dataflow/DataflowCoreNodes.h b/Engine/Source/Runtime/Experimental/Dataflow/Core/Public/Dataflow/DataflowCoreNodes.h
--- a/Engine/Source/Runtime/Experimental/Dataflow/Core/Public/Dataflow/DataflowCoreNodes.h
+++ b/Engine/Source/Runtime/Experimental/Dataflow/Core/Public/Dataflow/DataflowCoreNodes.h
@@ -90,12 +90,15 @@ private:
 	virtual bool OnOutputTypeChanged(const FDataflowOutput* Input) override;
 	virtual TArray<Dataflow::FPin> AddPins() override;
 	virtual bool CanAddPin() const override { return true; }
-	virtual bool CanRemovePin() const override { return Inputs.Num() > 2; }
+	virtual bool CanRemovePin() const override { return Inputs.Num() > NumInitialInputs; }
 	virtual TArray<Dataflow::FPin> GetPinsToRemove() const override;
 	virtual void OnPinRemoved(const Dataflow::FPin& Pin) override;
 	virtual void Serialize(FArchive& Ar) override;
 
 	Dataflow::TConnectionReference<FDataflowAnyType> GetConnectionReference(int32 Index) const;
+
+	static constexpr int32 NumRequiredDataflowInputs = 1;
+	static constexpr int32 NumInitialInputs = 2;
 };
 
 /** 
diff --git a/Engine/Source/Runtime/Experimental/Dataflow/Core/Public/Dataflow/DataflowNode.h b/Engine/Source/Runtime/Experimental/Dataflow/Core/Public/Dataflow/DataflowNode.h
--- a/Engine/Source/Runtime/Experimental/Dataflow/Core/Public/Dataflow/DataflowNode.h
+++ b/Engine/Source/Runtime/Experimental/Dataflow/Core/Public/Dataflow/DataflowNode.h
@@ -156,6 +156,7 @@ PRAGMA_ENABLE_DEPRECATION_WARNINGS
 	DATAFLOWCORE_API bool OutputSupportsType(FName Name, FName Type) const;
 
 	DATAFLOWCORE_API virtual void AddInput(FDataflowInput* InPtr);
+	DATAFLOWCORE_API int32 GetNumInputs() const;
 	DATAFLOWCORE_API TArray< FDataflowInput* > GetInputs() const;
 	DATAFLOWCORE_API void ClearInputs();
 	DATAFLOWCORE_API bool HasHideableInputs() const;
@@ -274,6 +275,23 @@ PRAGMA_ENABLE_DEPRECATION_WARNINGS
 		return Input;
 	}
 
+	template<typename T>
+	FDataflowInput& FindOrRegisterInputArrayConnection(const Dataflow::TConnectionReference<T>& Reference, const FName& ElementPropertyName = NAME_None,
+		const FName& ArrayPropertyName = NAME_None)
+	{
+		if (FDataflowInput* const FoundInput = FindInput(Reference))
+		{
+			return *FoundInput;
+		}
+
+		FDataflowInput& Input = RegisterInputArrayConnectionInternal(Reference, ElementPropertyName, ArrayPropertyName);
+		if constexpr (std::is_base_of_v<FDataflowAnyType, T>)
+		{
+			Input.SetTypePolicy(T::FPolicyType::GetInterface());
+		}
+		return Input;
+	}
+
 	/** Unregister the input connection if one exists matching this property, and then invalidate the graph. */
 	void UnregisterInputConnection(const Dataflow::FConnectionReference& Reference)
 	{
@@ -552,7 +570,7 @@ private:
 	static FString StripContainerIndexFromPropertyFullName(const FString& PropertyFullName);
 	static uint32 GetPropertyOffset(const TArray<const FProperty*>& PropertyChain);
 	uint32 GetConnectionOffsetFromReference(const void* Reference) const;
-	Dataflow::FConnectionKey GetKeyFromReference(const Dataflow::FConnectionReference& Reference) const;
+	DATAFLOWCORE_API Dataflow::FConnectionKey GetKeyFromReference(const Dataflow::FConnectionReference& Reference) const;
 
 	/**
 	* Find a property using the property address and name (not including its parent struct property names).
diff --git a/Engine/Source/Runtime/Experimental/Dataflow/Engine/Private/Dataflow/DataflowEdNode.cpp b/Engine/Source/Runtime/Experimental/Dataflow/Engine/Private/Dataflow/DataflowEdNode.cpp
--- a/Engine/Source/Runtime/Experimental/Dataflow/Engine/Private/Dataflow/DataflowEdNode.cpp
+++ b/Engine/Source/Runtime/Experimental/Dataflow/Engine/Private/Dataflow/DataflowEdNode.cpp
@@ -408,6 +408,20 @@ void UDataflowEdNode::Serialize(FArchive& Ar)
 {
 	Super::Serialize(Ar);
 	Ar << DataflowNodeGuid;
+#if WITH_EDITOR
+	if (Ar.IsTransacting())
+	{
+		const UDataflow* const DataflowObject = Cast<UDataflow>(GetGraph());
+		bool bCanSerializeNode = !DataflowObject || DataflowObject->IsPerNodeTransactionSerializationEnabled();
+		if (bCanSerializeNode)
+		{
+			if (TSharedPtr<FDataflowNode> DataflowNode = GetDataflowNode())
+			{
+				DataflowNode->SerializeInternal(Ar);
+			}
+		}
+	}
+#endif
 }
 
 #if WITH_EDITOR
diff --git a/Engine/Source/Runtime/Experimental/Dataflow/Engine/Private/Dataflow/DataflowObject.cpp b/Engine/Source/Runtime/Experimental/Dataflow/Engine/Private/Dataflow/DataflowObject.cpp
--- a/Engine/Source/Runtime/Experimental/Dataflow/Engine/Private/Dataflow/DataflowObject.cpp
+++ b/Engine/Source/Runtime/Experimental/Dataflow/Engine/Private/Dataflow/DataflowObject.cpp
@@ -136,8 +136,47 @@ void UDataflow::RemoveWireframeRenderTarget(TObjectPtr<const UDataflowEdNode> In
 
 void UDataflow::Serialize(FArchive& Ar)
 {
+#if WITH_EDITOR
+	// Disable per-node serialization (used for transactions, i.e., undo/redo) when serializing the whole graph.
+	bEnablePerNodeTransactionSerialization = false;
+#endif
+
 	Super::Serialize(Ar);
 	Dataflow->Serialize(Ar, this);
+
+#if WITH_EDITOR
+	bEnablePerNodeTransactionSerialization = true;
+#endif
+}
+
+TObjectPtr<const UDataflowEdNode> UDataflow::FindEdNodeByDataflowNodeGuid(const FGuid& Guid) const
+{
+	for (const UEdGraphNode* const EdNode : Nodes)
+	{
+		if (const UDataflowEdNode* const DataflowEdNode = Cast<UDataflowEdNode>(EdNode))
+		{
+			if (DataflowEdNode->GetDataflowNodeGuid() == Guid)
+			{
+				return TObjectPtr<const UDataflowEdNode>(DataflowEdNode);
+			}
+		}
+	}
+	return TObjectPtr<const UDataflowEdNode>(nullptr);
+}
+
+TObjectPtr<UDataflowEdNode> UDataflow::FindEdNodeByDataflowNodeGuid(const FGuid& Guid)
+{
+	for (UEdGraphNode* const EdNode : Nodes)
+	{
+		if (UDataflowEdNode* const DataflowEdNode = Cast<UDataflowEdNode>(EdNode))
+		{
+			if (DataflowEdNode->GetDataflowNodeGuid() == Guid)
+			{
+				return TObjectPtr<UDataflowEdNode>(DataflowEdNode);
+			}
+		}
+	}
+	return TObjectPtr<UDataflowEdNode>(nullptr);
 }
 
 #if WITH_EDITOR
diff --git a/Engine/Source/Runtime/Experimental/Dataflow/Engine/Public/Dataflow/DataflowObject.h b/Engine/Source/Runtime/Experimental/Dataflow/Engine/Public/Dataflow/DataflowObject.h
--- a/Engine/Source/Runtime/Experimental/Dataflow/Engine/Public/Dataflow/DataflowObject.h
+++ b/Engine/Source/Runtime/Experimental/Dataflow/Engine/Public/Dataflow/DataflowObject.h
@@ -128,6 +128,9 @@ public:
 		return FDataflowAssetEdit(ThisNC, [ThisNC]() {ThisNC->PostEditCallback(); });
 	}
 
+	DATAFLOWENGINE_API TObjectPtr<const UDataflowEdNode> FindEdNodeByDataflowNodeGuid(const FGuid& Guid) const;
+	DATAFLOWENGINE_API TObjectPtr<UDataflowEdNode> FindEdNodeByDataflowNodeGuid(const FGuid& Guid);
+
 	//
 	// Render Targets
 	//
@@ -165,7 +168,13 @@ public:
 
 #endif
 
+#if WITH_EDITOR
+	/** Used to disable per-node serialization when serializing a transaction */
+	bool IsPerNodeTransactionSerializationEnabled() const { return bEnablePerNodeTransactionSerialization; }
+private:
 
-
+	/** Used to disable per-node serialization when serializing a transaction */
+	bool bEnablePerNodeTransactionSerialization = true;
+#endif
 };
 
