Dataflow Editor: add the cloth preview blueprint class when opening a cloth asset via the context menu.
Also move the blueprint asset from ClothAsset to ClothAssetEditor

#virtualized

diff --git a/Engine/Plugins/ChaosClothAssetEditor/Source/ChaosClothAssetEditor/Private/ChaosClothAsset/AssetDefinition_ClothAsset.cpp b/Engine/Plugins/ChaosClothAssetEditor/Source/ChaosClothAssetEditor/Private/ChaosClothAsset/AssetDefinition_ClothAsset.cpp
--- a/Engine/Plugins/ChaosClothAssetEditor/Source/ChaosClothAssetEditor/Private/ChaosClothAsset/AssetDefinition_ClothAsset.cpp
+++ b/Engine/Plugins/ChaosClothAssetEditor/Source/ChaosClothAssetEditor/Private/ChaosClothAsset/AssetDefinition_ClothAsset.cpp
@@ -201,7 +201,7 @@ EAssetCommandResult UAssetDefinition_ClothAsset::OpenAssets(const FAssetOpenArgs
 		{
 			UDataflowEditor* const AssetEditor = NewObject<UDataflowEditor>(AssetEditorSubsystem, NAME_None, RF_Transient);
             const TSubclassOf<AActor> ActorClass = StaticLoadClass(AActor::StaticClass(), nullptr,
-            	TEXT("/ChaosClothAsset/BP_ClothPreview.BP_ClothPreview_C"), nullptr, LOAD_None, nullptr);
+            	TEXT("/ChaosClothAssetEditor/BP_ClothPreview.BP_ClothPreview_C"), nullptr, LOAD_None, nullptr);
 			AssetEditor->Initialize({ ClothObjects[0] }, ActorClass);
 		}
 
diff --git a/Engine/Plugins/ChaosClothAssetEditor/Source/ChaosClothAssetEditor/Private/ChaosClothAsset/ChaosClothAssetEditorModule.cpp b/Engine/Plugins/ChaosClothAssetEditor/Source/ChaosClothAssetEditor/Private/ChaosClothAsset/ChaosClothAssetEditorModule.cpp
--- a/Engine/Plugins/ChaosClothAssetEditor/Source/ChaosClothAssetEditor/Private/ChaosClothAsset/ChaosClothAssetEditorModule.cpp
+++ b/Engine/Plugins/ChaosClothAssetEditor/Source/ChaosClothAssetEditor/Private/ChaosClothAsset/ChaosClothAssetEditorModule.cpp
@@ -55,7 +55,11 @@ namespace UE::Chaos::ClothAsset
 							AssetsToEdit.Append(Context->LoadSelectedObjects<UObject>());
 
 							UDataflowEditor* const AssetEditor = NewObject<UDataflowEditor>(DataflowEditorSubsystem, NAME_None, RF_Transient);
-							AssetEditor->Initialize({AssetsToEdit[0]});
+
+							const TSubclassOf<AActor> PreviewActorClass = StaticLoadClass(AActor::StaticClass(), nullptr,
+								TEXT("/ChaosClothAssetEditor/BP_ClothPreview.BP_ClothPreview_C"), nullptr, LOAD_None, nullptr);
+
+							AssetEditor->Initialize({AssetsToEdit[0]}, PreviewActorClass);
 						}),
 						FCanExecuteAction::CreateWeakLambda(Context, [Context]() { return Context->bCanBeModified; }));
 
