Fix crash when opening brand new Cloth Asset in the Dataflow Editor

diff --git a/Engine/Plugins/ChaosClothAssetEditor/Source/ChaosClothAssetEditor/Private/ChaosClothAsset/AssetDefinition_ClothAsset.cpp b/Engine/Plugins/ChaosClothAssetEditor/Source/ChaosClothAssetEditor/Private/ChaosClothAsset/AssetDefinition_ClothAsset.cpp
--- a/Engine/Plugins/ChaosClothAssetEditor/Source/ChaosClothAssetEditor/Private/ChaosClothAsset/AssetDefinition_ClothAsset.cpp
+++ b/Engine/Plugins/ChaosClothAssetEditor/Source/ChaosClothAssetEditor/Private/ChaosClothAsset/AssetDefinition_ClothAsset.cpp
@@ -130,21 +130,21 @@ namespace ClothAssetDefinitionHelpers
 
 		return true;
 	}
+}
 
-	// Create a new UDataflow if one doesn't already exist for the Cloth Asset
-	UObject* NewOrOpenDataflowAsset(const UChaosClothAsset* ClothAsset)
-	{
-		UObject* DataflowAsset = nullptr;
-		bool bDialogDone = false;
-		while (!bDialogDone)
-		{
-			bDialogDone = CreateNewDataflowAsset(ClothAsset, DataflowAsset);
-		}
 
-		return DataflowAsset;
+// Create a new UDataflow if one doesn't already exist for the Cloth Asset
+UObject* UAssetDefinition_ClothAsset::NewOrOpenDataflowAsset(const UChaosClothAsset* ClothAsset)
+{
+	UObject* DataflowAsset = nullptr;
+	bool bDialogDone = false;
+	while (!bDialogDone)
+	{
+		bDialogDone = ClothAssetDefinitionHelpers::CreateNewDataflowAsset(ClothAsset, DataflowAsset);
 	}
-}
 
+	return DataflowAsset;
+}
 
 FText UAssetDefinition_ClothAsset::GetAssetDisplayName() const
 {
@@ -185,7 +185,7 @@ EAssetCommandResult UAssetDefinition_ClothAsset::OpenAssets(const FAssetOpenArgs
 		UChaosClothAsset* const ClothAsset = CastChecked<UChaosClothAsset>(ClothObjects[0]);
 		if (!ClothAsset->GetDataflow())
 		{
-			if (UDataflow* const NewDataflowAsset = Cast<UDataflow>(ClothAssetDefinitionHelpers::NewOrOpenDataflowAsset(ClothAsset)))
+			if (UDataflow* const NewDataflowAsset = Cast<UDataflow>(NewOrOpenDataflowAsset(ClothAsset)))
 			{
 				ClothAsset->SetDataflow(NewDataflowAsset);
 			}
diff --git a/Engine/Plugins/ChaosClothAssetEditor/Source/ChaosClothAssetEditor/Private/ChaosClothAsset/AssetDefinition_ClothAsset.h b/Engine/Plugins/ChaosClothAssetEditor/Source/ChaosClothAssetEditor/Private/ChaosClothAsset/AssetDefinition_ClothAsset.h
--- a/Engine/Plugins/ChaosClothAssetEditor/Source/ChaosClothAssetEditor/Private/ChaosClothAsset/AssetDefinition_ClothAsset.h
+++ b/Engine/Plugins/ChaosClothAssetEditor/Source/ChaosClothAssetEditor/Private/ChaosClothAsset/AssetDefinition_ClothAsset.h
@@ -5,11 +5,17 @@
 #include "AssetDefinitionDefault.h"
 #include "AssetDefinition_ClothAsset.generated.h"
 
+class UChaosClothAsset;
+
 UCLASS()
 class UAssetDefinition_ClothAsset : public UAssetDefinitionDefault
 {
 	GENERATED_BODY()
 
+public:
+
+	static UObject* NewOrOpenDataflowAsset(const UChaosClothAsset* ClothAsset);
+
 private:
 
 	virtual FText GetAssetDisplayName() const override;
diff --git a/Engine/Plugins/ChaosClothAssetEditor/Source/ChaosClothAssetEditor/Private/ChaosClothAsset/ChaosClothAssetEditorModule.cpp b/Engine/Plugins/ChaosClothAssetEditor/Source/ChaosClothAssetEditor/Private/ChaosClothAsset/ChaosClothAssetEditorModule.cpp
--- a/Engine/Plugins/ChaosClothAssetEditor/Source/ChaosClothAssetEditor/Private/ChaosClothAsset/ChaosClothAssetEditorModule.cpp
+++ b/Engine/Plugins/ChaosClothAssetEditor/Source/ChaosClothAssetEditor/Private/ChaosClothAsset/ChaosClothAssetEditorModule.cpp
@@ -1,6 +1,7 @@
 // Copyright Epic Games, Inc. All Rights Reserved.
 
 #include "ChaosClothAsset/ChaosClothAssetEditorModule.h"
+#include "ChaosClothAsset/AssetDefinition_ClothAsset.h"
 #include "ChaosClothAsset/ClothAsset.h"
 #include "ChaosClothAsset/ClothEditorCommands.h"
 #include "ChaosClothAsset/ClothEditorMode.h"
@@ -56,6 +57,16 @@ namespace UE::Chaos::ClothAsset
 
 							UDataflowEditor* const AssetEditor = NewObject<UDataflowEditor>(DataflowEditorSubsystem, NAME_None, RF_Transient);
 
+							// Macke sure the cloth asset has a Dataflow asset
+							UChaosClothAsset* const ClothAsset = CastChecked<UChaosClothAsset>(AssetsToEdit[0]);
+							if (!ClothAsset->GetDataflow())
+							{
+								if (UDataflow* const NewDataflowAsset = Cast<UDataflow>(UAssetDefinition_ClothAsset::NewOrOpenDataflowAsset(ClothAsset)))
+								{
+									ClothAsset->SetDataflow(NewDataflowAsset);
+								}
+							}
+
 							const TSubclassOf<AActor> PreviewActorClass = StaticLoadClass(AActor::StaticClass(), nullptr,
 								TEXT("/ChaosClothAssetEditor/BP_ClothPreview.BP_ClothPreview_C"), nullptr, LOAD_None, nullptr);
 
