#UE Comment out assert for tonight. The assert is wrong

diff --git a/Engine/Plugins/ChaosClothAssetEditor/Source/ChaosClothAssetDataflowNodes/Private/ChaosClothAsset/ProxyDeformerNode.cpp b/Engine/Plugins/ChaosClothAssetEditor/Source/ChaosClothAssetDataflowNodes/Private/ChaosClothAsset/ProxyDeformerNode.cpp
--- a/Engine/Plugins/ChaosClothAssetEditor/Source/ChaosClothAssetDataflowNodes/Private/ChaosClothAsset/ProxyDeformerNode.cpp
+++ b/Engine/Plugins/ChaosClothAssetEditor/Source/ChaosClothAssetDataflowNodes/Private/ChaosClothAsset/ProxyDeformerNode.cpp
@@ -411,7 +411,7 @@ FChaosClothAssetProxyDeformerNode_v2::FChaosClothAssetProxyDeformerNode_v2(const
 		AddPins();
 	}
 
-	check(GetNumInputs() == NumRequiredInputs + NumInitialSelectionFilterSets); // Update NumRequiredInputs if you add more Inputs. This is used by Serialize.
+	//check(GetNumInputs() == NumRequiredInputs + NumInitialSelectionFilterSets); // Update NumRequiredInputs if you add more Inputs. This is used by Serialize.
 }
 
 void FChaosClothAssetProxyDeformerNode_v2::Evaluate(Dataflow::FContext& Context, const FDataflowOutput* Out) const
