Cloth Asset:
- Hides bCouldUseFabrics from the UI.
- Added collection name verification to all string based connectable value customisations.
- Added new output only connectable value struct to allow for future collection name verification on outputs.

#rnx

diff --git a/Engine/Plugins/ChaosClothAssetEditor/Source/ChaosClothAssetDataflowNodes/Private/ChaosClothAsset/ChaosClothAssetDataflowNodesModule.cpp b/Engine/Plugins/ChaosClothAssetEditor/Source/ChaosClothAssetDataflowNodes/Private/ChaosClothAsset/ChaosClothAssetDataflowNodesModule.cpp
--- a/Engine/Plugins/ChaosClothAssetEditor/Source/ChaosClothAssetDataflowNodes/Private/ChaosClothAsset/ChaosClothAssetDataflowNodesModule.cpp
+++ b/Engine/Plugins/ChaosClothAssetEditor/Source/ChaosClothAssetDataflowNodes/Private/ChaosClothAsset/ChaosClothAssetDataflowNodesModule.cpp
@@ -297,9 +297,10 @@ namespace UE::Chaos::ClothAsset
 				PropertyModule->RegisterCustomPropertyTypeLayout(FChaosClothAssetWeightedValue::StaticStruct()->GetFName(), FOnGetPropertyTypeCustomizationInstance::CreateStatic(&FWeightedValueCustomization::MakeInstance));
 				PropertyModule->RegisterCustomPropertyTypeLayout(FChaosClothAssetWeightedValueNonAnimatable::StaticStruct()->GetFName(), FOnGetPropertyTypeCustomizationInstance::CreateStatic(&FWeightedValueCustomization::MakeInstance));
 				PropertyModule->RegisterCustomPropertyTypeLayout(FChaosClothAssetWeightedValueNonAnimatableNoLowHighRange::StaticStruct()->GetFName(), FOnGetPropertyTypeCustomizationInstance::CreateStatic(&FWeightedValueCustomization::MakeInstance));
-				PropertyModule->RegisterCustomPropertyTypeLayout(FChaosClothAssetConnectableStringValue::StaticStruct()->GetFName(), FOnGetPropertyTypeCustomizationInstance::CreateStatic(&FWeightedValueCustomization::MakeInstance));
 				PropertyModule->RegisterCustomPropertyTypeLayout(FChaosClothAssetWeightedValueOverride::StaticStruct()->GetFName(), FOnGetPropertyTypeCustomizationInstance::CreateStatic(&FWeightedValueCustomization::MakeInstance));
+				PropertyModule->RegisterCustomPropertyTypeLayout(FChaosClothAssetConnectableStringValue::StaticStruct()->GetFName(), FOnGetPropertyTypeCustomizationInstance::CreateStatic(&FWeightedValueCustomization::MakeInstance));
 				PropertyModule->RegisterCustomPropertyTypeLayout(FChaosClothAssetConnectableIStringValue::StaticStruct()->GetFName(), FOnGetPropertyTypeCustomizationInstance::CreateStatic(&FConnectableValueCustomization::MakeInstance));
+				PropertyModule->RegisterCustomPropertyTypeLayout(FChaosClothAssetConnectableOStringValue::StaticStruct()->GetFName(), FOnGetPropertyTypeCustomizationInstance::CreateStatic(&FWeightedValueCustomization::MakeInstance));
 				PropertyModule->RegisterCustomPropertyTypeLayout(FChaosClothAssetConnectableIOStringValue::StaticStruct()->GetFName(), FOnGetPropertyTypeCustomizationInstance::CreateStatic(&FWeightedValueCustomization::MakeInstance));
 				PropertyModule->RegisterCustomPropertyTypeLayout(FChaosClothAssetNodeSelectionGroup::StaticStruct()->GetFName(), FOnGetPropertyTypeCustomizationInstance::CreateStatic(&FSelectionGroupCustomization::MakeInstance));
 				PropertyModule->RegisterCustomPropertyTypeLayout(FChaosClothAssetImportFilePath::StaticStruct()->GetFName(), FOnGetPropertyTypeCustomizationInstance::CreateStatic(&FImportFilePathCustomization::MakeInstance));
@@ -328,9 +329,10 @@ namespace UE::Chaos::ClothAsset
 					PropertyModule->UnregisterCustomPropertyTypeLayout(FChaosClothAssetWeightedValue::StaticStruct()->GetFName());
 					PropertyModule->UnregisterCustomPropertyTypeLayout(FChaosClothAssetWeightedValueNonAnimatable::StaticStruct()->GetFName());
 					PropertyModule->UnregisterCustomPropertyTypeLayout(FChaosClothAssetWeightedValueNonAnimatableNoLowHighRange::StaticStruct()->GetFName());
-					PropertyModule->UnregisterCustomPropertyTypeLayout(FChaosClothAssetConnectableStringValue::StaticStruct()->GetFName());
 					PropertyModule->UnregisterCustomPropertyTypeLayout(FChaosClothAssetWeightedValueOverride::StaticStruct()->GetFName());
+					PropertyModule->UnregisterCustomPropertyTypeLayout(FChaosClothAssetConnectableStringValue::StaticStruct()->GetFName());
 					PropertyModule->UnregisterCustomPropertyTypeLayout(FChaosClothAssetConnectableIStringValue::StaticStruct()->GetFName());
+					PropertyModule->UnregisterCustomPropertyTypeLayout(FChaosClothAssetConnectableOStringValue::StaticStruct()->GetFName());
 					PropertyModule->UnregisterCustomPropertyTypeLayout(FChaosClothAssetConnectableIOStringValue::StaticStruct()->GetFName());
 					PropertyModule->UnregisterCustomPropertyTypeLayout(FChaosClothAssetNodeSelectionGroup::StaticStruct()->GetFName());
 					PropertyModule->UnregisterCustomPropertyTypeLayout(FChaosClothAssetImportFilePath::StaticStruct()->GetFName());
diff --git a/Engine/Plugins/ChaosClothAssetEditor/Source/ChaosClothAssetDataflowNodes/Private/ChaosClothAsset/ClothDataflowTools.cpp b/Engine/Plugins/ChaosClothAssetEditor/Source/ChaosClothAssetDataflowNodes/Private/ChaosClothAsset/ClothDataflowTools.cpp
--- a/Engine/Plugins/ChaosClothAssetEditor/Source/ChaosClothAssetDataflowNodes/Private/ChaosClothAsset/ClothDataflowTools.cpp
+++ b/Engine/Plugins/ChaosClothAssetEditor/Source/ChaosClothAssetDataflowNodes/Private/ChaosClothAsset/ClothDataflowTools.cpp
@@ -467,11 +467,13 @@ namespace UE::Chaos::ClothAsset
 		UE_LOG(LogChaosClothAssetDataflowNodes, Warning, TEXT("%s"), *Text.ToString());
 	}
 
-	void FClothDataflowTools::MakeCollectionName(FString& InOutString)
+	bool FClothDataflowTools::MakeCollectionName(FString& InOutString)
 	{
+		const FString SourceString = InOutString;
 		InOutString = SlugStringForValidName(InOutString, TEXT("_")).Replace(TEXT("\\"), TEXT("_"));
 		bool bCharsWereRemoved;
 		do { InOutString.TrimCharInline(TEXT('_'), &bCharsWereRemoved); } while (bCharsWereRemoved);
+		return InOutString.Equals(SourceString);
 	}
 
 	static void CopyBuildSettings(const FMeshBuildSettings& InStaticMeshBuildSettings, FSkeletalMeshBuildSettings& OutSkeletalMeshBuildSettings)
diff --git a/Engine/Plugins/ChaosClothAssetEditor/Source/ChaosClothAssetDataflowNodes/Private/ChaosClothAsset/ClothDataflowTools.h b/Engine/Plugins/ChaosClothAssetEditor/Source/ChaosClothAssetDataflowNodes/Private/ChaosClothAsset/ClothDataflowTools.h
--- a/Engine/Plugins/ChaosClothAssetEditor/Source/ChaosClothAssetDataflowNodes/Private/ChaosClothAsset/ClothDataflowTools.h
+++ b/Engine/Plugins/ChaosClothAssetEditor/Source/ChaosClothAssetDataflowNodes/Private/ChaosClothAsset/ClothDataflowTools.h
@@ -31,8 +31,10 @@ namespace UE::Chaos::ClothAsset
 		 * The resulting name won't contains spaces and any other special characters as listed in
 		 * INVALID_OBJECTNAME_CHARACTERS (currently "',/.:|&!~\n\r\t@#(){}[]=;^%$`).
 		 * It will also have all leading underscore removed, as these names are reserved for internal use.
+		 * @param InOutString The string to turn into a valid collection name.
+		 * @return Whether the InOutString was already a valid collection name.
 		 */
-		static void MakeCollectionName(FString& InOutString);
+		static bool MakeCollectionName(FString& InOutString);
 
 		static bool BuildSkeletalMeshModelFromMeshDescription(const FMeshDescription* const InMeshDescription, const FMeshBuildSettings& InBuildSettings, FSkeletalMeshLODModel& SkeletalMeshModel);
 
diff --git a/Engine/Plugins/ChaosClothAssetEditor/Source/ChaosClothAssetDataflowNodes/Private/ChaosClothAsset/ConnectableValueCustomization.cpp b/Engine/Plugins/ChaosClothAssetEditor/Source/ChaosClothAssetDataflowNodes/Private/ChaosClothAsset/ConnectableValueCustomization.cpp
--- a/Engine/Plugins/ChaosClothAssetEditor/Source/ChaosClothAssetDataflowNodes/Private/ChaosClothAsset/ConnectableValueCustomization.cpp
+++ b/Engine/Plugins/ChaosClothAssetEditor/Source/ChaosClothAssetDataflowNodes/Private/ChaosClothAsset/ConnectableValueCustomization.cpp
@@ -175,9 +175,25 @@ PRAGMA_ENABLE_DEPRECATION_WARNINGS
 					{
 						if (const TSharedPtr<IPropertyHandle> HandlePtr = HandleWeakPtr.Pin())
 						{
-							HandlePtr->SetValue(Text.ToString(), EPropertyValueSetFlags::DefaultFlags);
+							FString TextString = Text.ToString();
+							FClothDataflowTools::MakeCollectionName(TextString);
+							HandlePtr->SetValue(TextString, EPropertyValueSetFlags::DefaultFlags);
 						}
 					})
+				.OnVerifyTextChanged_Lambda([](const FText& Text, FText& OutErrorMessage) -> bool
+					{
+						bool bIsValidCollectionName = false;
+						FString TextString = Text.ToString();
+						bIsValidCollectionName = FClothDataflowTools::MakeCollectionName(TextString);
+						if (!bIsValidCollectionName)
+						{
+							OutErrorMessage =
+								LOCTEXT("NotValidCollectioName",
+									"To be a valid collection name, this text string musn't start by an underscore,\n"
+									"contain whitespaces, or any of the following character: \"',/.:|&!~@#(){}[]=;^%$`");
+						}
+						return bIsValidCollectionName;
+					})
 				.IsEnabled_Lambda([HandleWeakPtr, StructurePropertyHandle]() -> bool
 					{
 						if (const TSharedPtr<IPropertyHandle> HandlePtr = HandleWeakPtr.Pin())
diff --git a/Engine/Plugins/ChaosClothAssetEditor/Source/ChaosClothAssetDataflowNodes/Public/ChaosClothAsset/ConnectableValue.h b/Engine/Plugins/ChaosClothAssetEditor/Source/ChaosClothAssetDataflowNodes/Public/ChaosClothAsset/ConnectableValue.h
--- a/Engine/Plugins/ChaosClothAssetEditor/Source/ChaosClothAssetDataflowNodes/Public/ChaosClothAsset/ConnectableValue.h
+++ b/Engine/Plugins/ChaosClothAssetEditor/Source/ChaosClothAssetDataflowNodes/Public/ChaosClothAsset/ConnectableValue.h
@@ -35,7 +35,7 @@ PRAGMA_ENABLE_DEPRECATION_WARNINGS
 	/**
 	  * Whether the property could import fabrics datas or not
 	  */
-	UPROPERTY(VisibleAnywhere, Transient, Category = "Value")
+	UPROPERTY(VisibleAnywhere, Transient, Category = "Value", Meta = (EditCondition = "false", EditConditionHides))
 	bool bCouldUseFabrics = false;
 
 	UE_DEPRECATED(5.5, "Override properties are no longer used.")
@@ -49,6 +49,16 @@ PRAGMA_ENABLE_DEPRECATION_WARNINGS
 	bool bBuildFabricMaps = false;
 };
 
+USTRUCT()
+struct FChaosClothAssetConnectableOStringValue
+{
+	GENERATED_BODY()
+
+	/** The value for this property. */
+	UPROPERTY(EditAnywhere, Category = "Value", Meta = (DataflowOutput))
+	FString StringValue;
+};
+
 PRAGMA_DISABLE_DEPRECATION_WARNINGS
 USTRUCT()
 struct FChaosClothAssetConnectableIOStringValue
diff --git a/Engine/Plugins/ChaosClothAssetEditor/Source/ChaosClothAssetDataflowNodes/Public/ChaosClothAsset/WeightedValue.h b/Engine/Plugins/ChaosClothAssetEditor/Source/ChaosClothAssetDataflowNodes/Public/ChaosClothAsset/WeightedValue.h
--- a/Engine/Plugins/ChaosClothAssetEditor/Source/ChaosClothAssetDataflowNodes/Public/ChaosClothAsset/WeightedValue.h
+++ b/Engine/Plugins/ChaosClothAssetEditor/Source/ChaosClothAssetDataflowNodes/Public/ChaosClothAsset/WeightedValue.h
@@ -71,7 +71,7 @@ PRAGMA_ENABLE_DEPRECATION_WARNINGS
 	/**
 	 * Whether the property could import fabrics datas or not
 	 */
-	UPROPERTY(VisibleAnywhere, Transient, Category = "Weighted Value")
+	UPROPERTY(VisibleAnywhere, Transient, Category = "Weighted Value", Meta = (EditCondition = "false", EditConditionHides))
 	bool bCouldUseFabrics = false;
 
 	UE_DEPRECATED(5.5, "Override properties are no longer used.")
@@ -127,7 +127,7 @@ PRAGMA_ENABLE_DEPRECATION_WARNINGS
 	/**
 	 * Whether the property could import fabrics datas or not
 	 */
-	UPROPERTY(VisibleAnywhere, Transient, Category = "Weighted Value")
+	UPROPERTY(VisibleAnywhere, Transient, Category = "Weighted Value", Meta = (EditCondition = "false", EditConditionHides))
 	bool bCouldUseFabrics = false;
 	
 	UE_DEPRECATED(5.5, "Override properties are no longer used.")
