Cloth Asset:
- Split the ProxyDeformer node in two nodes (ProxyDeformer and SkinningBlend) and made the deformer skinned by default unless selection filters are provided.
- Fixed check in ProxyDeformer constructor.

diff --git a/Engine/Plugins/ChaosClothAssetEditor/Source/ChaosClothAssetDataflowNodes/Private/ChaosClothAsset/ChaosClothAssetDataflowNodesModule.cpp b/Engine/Plugins/ChaosClothAssetEditor/Source/ChaosClothAssetDataflowNodes/Private/ChaosClothAsset/ChaosClothAssetDataflowNodesModule.cpp
--- a/Engine/Plugins/ChaosClothAssetEditor/Source/ChaosClothAssetDataflowNodes/Private/ChaosClothAsset/ChaosClothAssetDataflowNodesModule.cpp
+++ b/Engine/Plugins/ChaosClothAssetEditor/Source/ChaosClothAssetDataflowNodes/Private/ChaosClothAsset/ChaosClothAssetDataflowNodesModule.cpp
@@ -52,6 +52,7 @@
 #include "ChaosClothAsset/SimulationXPBDBendingSpringConfigNode.h"
 #include "ChaosClothAsset/SimulationXPBDEdgeSpringConfigNode.h"
 #include "ChaosClothAsset/SkeletalMeshImportNode.h"
+#include "ChaosClothAsset/SkinningBlendNode.h"
 #include "ChaosClothAsset/StaticMeshImportNode.h"
 #include "ChaosClothAsset/TerminalNode.h"
 #include "ChaosClothAsset/TerminalNodeRefreshAssetCustomization.h"
@@ -133,6 +134,7 @@ namespace UE::Chaos::ClothAsset
 			DATAFLOW_NODE_REGISTER_CREATION_FACTORY(FChaosClothAssetSimulationXPBDEdgeSpringConfigNode);
 			DATAFLOW_NODE_REGISTER_CREATION_FACTORY(FChaosClothAssetSkeletalMeshImportNode);
 			DATAFLOW_NODE_REGISTER_CREATION_FACTORY(FChaosClothAssetSkeletalMeshImportNode_v2);
+			DATAFLOW_NODE_REGISTER_CREATION_FACTORY(FChaosClothAssetSkinningBlendNode);
 			DATAFLOW_NODE_REGISTER_CREATION_FACTORY(FChaosClothAssetStaticMeshImportNode);
 			DATAFLOW_NODE_REGISTER_CREATION_FACTORY(FChaosClothAssetStaticMeshImportNode_v2);
 			DATAFLOW_NODE_REGISTER_CREATION_FACTORY(FChaosClothAssetTerminalNode);
diff --git a/Engine/Plugins/ChaosClothAssetEditor/Source/ChaosClothAssetDataflowNodes/Private/ChaosClothAsset/ProxyDeformerNode.cpp b/Engine/Plugins/ChaosClothAssetEditor/Source/ChaosClothAssetDataflowNodes/Private/ChaosClothAsset/ProxyDeformerNode.cpp
--- a/Engine/Plugins/ChaosClothAssetEditor/Source/ChaosClothAssetDataflowNodes/Private/ChaosClothAsset/ProxyDeformerNode.cpp
+++ b/Engine/Plugins/ChaosClothAssetEditor/Source/ChaosClothAssetDataflowNodes/Private/ChaosClothAsset/ProxyDeformerNode.cpp
@@ -33,7 +33,7 @@ namespace UE::Chaos::ClothAsset::Private
 		TArrayView<TArray<float>> RenderDeformerWeight;
 		TArrayView<float> RenderDeformerSkinningBlend;
 
-		int32 Generate(bool bUseSmoothTransition, bool bUseMultipleInfluences, float InfluenceRadius)
+		int32 Generate(bool bUseSmoothTransition, bool bUseMultipleInfluences, float InfluenceRadius, bool bDoSkinningBlend)
 		{
 			check(RenderPositions.Num() == RenderNormals.Num());
 			check(RenderPositions.Num() == RenderDeformerPositionBaryCoordsAndDist.Num());
@@ -84,6 +84,18 @@ namespace UE::Chaos::ClothAsset::Private
 			check(MeshToMeshVertData.Num() == RenderPositions.Num() * NumInfluences);  // Check modulo
 			check((!bUseMultipleInfluences && NumInfluences == 1) || (bUseMultipleInfluences && NumInfluences > 1));
 
+			auto IsRenderVertexInFilterSets = [this](int32 Index) -> bool
+				{
+					for (const ClothingMeshUtils::FMeshToMeshFilterSet& Set : MeshToMeshFilterSet)
+					{
+						if (Set.TargetVertices.Contains(Index))
+						{
+							return true;
+						}
+					}
+					return false;
+				};
+
 			for (int32 Index = 0; Index < RenderPositions.Num(); ++Index)
 			{
 				RenderDeformerPositionBaryCoordsAndDist[Index].SetNum(NumInfluences);
@@ -92,7 +104,7 @@ namespace UE::Chaos::ClothAsset::Private
 				RenderDeformerSimIndices3D[Index].SetNum(NumInfluences);
 				RenderDeformerWeight[Index].SetNum(NumInfluences);
 
-				RenderDeformerSkinningBlend[Index] = 0.f;
+				RenderDeformerSkinningBlend[Index] = bDoSkinningBlend || IsRenderVertexInFilterSets(Index) ? 0.f : 1.f;
 
 				for (int32 Influence = 0; Influence < NumInfluences; ++Influence)
 				{
@@ -107,7 +119,10 @@ namespace UE::Chaos::ClothAsset::Private
 						MeshToMeshVertDatum.SourceMeshVertIndices[2]);
 					RenderDeformerWeight[Index][Influence] = MeshToMeshVertDatum.Weight;
 
-					RenderDeformerSkinningBlend[Index] += MeshToMeshVertDatum.Weight * (float)MeshToMeshVertDatum.SourceMeshVertIndices[3] / (float)TNumericLimits<uint16>::Max();
+					if (bDoSkinningBlend)
+					{
+						RenderDeformerSkinningBlend[Index] += MeshToMeshVertDatum.Weight * (float)MeshToMeshVertDatum.SourceMeshVertIndices[3] / (float)TNumericLimits<uint16>::Max();
+					}
 				}
 			}
 			return NumInfluences;
@@ -239,23 +254,20 @@ namespace UE::Chaos::ClothAsset::Private
 				{
 					if (const TSet<int32>* SimSelectionSet = SelectionFacade.FindSelectionSet(SelectionName.Get<1>()))
 					{
-						if (RenderSelectionSet->Num() && SimSelectionSet->Num())
+						if (RenderSelectionSet->Num() || SimSelectionSet->Num())
 						{
 							const FName RenderSelectionGroup = SelectionFacade.GetSelectionGroup(SelectionName.Get<0>());
 							const FName SimSelectionGroup = SelectionFacade.GetSelectionGroup(SelectionName.Get<1>());
 
 							// Retrieve the sim face selection
 							TSet<int32> SimFaceSelection = GetSimFaceSelection(SimSelectionGroup, *SimSelectionSet);
-							if (!SimFaceSelection.Num())
-							{
-								continue;  // Nothing selected on the simulation side
-							}
 
 							// Retrieve the render vertex selection
 							TSet<int32> RenderVertexSelection = GetRenderVertexSelection(RenderSelectionGroup, *RenderSelectionSet);
-							if (!RenderVertexSelection.Num())
+
+							if (!SimFaceSelection.Num() && !RenderVertexSelection.Num())
 							{
-								continue;  // Nothing selected on the render side
+								continue;  // Nothing selected
 							}
 
 							ClothingMeshUtils::FMeshToMeshFilterSet& MeshToMeshFilterSet = MeshToMeshFilterSets.AddDefaulted_GetRef();
@@ -395,15 +407,9 @@ namespace UE::Chaos::ClothAsset::Private
 FChaosClothAssetProxyDeformerNode_v2::FChaosClothAssetProxyDeformerNode_v2(const Dataflow::FNodeParameters& InParam, FGuid InGuid)
 	: FDataflowNode(InParam, InGuid)
 {
-	using namespace UE::Chaos::ClothAsset;
-	SimVertexSelection.StringValue = FString();  // An empty selection is an accepted input, but a non existing one isn't
-	SkinningBlendName = ClothCollectionAttribute::RenderDeformerSkinningBlend.ToString();
-
 	RegisterInputConnection(&Collection);
-	RegisterInputConnection(&SimVertexSelection.StringValue, GET_MEMBER_NAME_CHECKED(FChaosClothAssetConnectableIStringValue, StringValue));
 	RegisterOutputConnection(&Collection)
 		.SetPassthroughInput(&Collection);
-	RegisterOutputConnection(&SkinningBlendName);
 
 	// Start with one set of option pins.
 	for (int32 Index = 0; Index < NumInitialSelectionFilterSets; ++Index)
@@ -411,7 +417,7 @@ FChaosClothAssetProxyDeformerNode_v2::FChaosClothAssetProxyDeformerNode_v2(const
 		AddPins();
 	}
 
-	//check(GetNumInputs() == NumRequiredInputs + NumInitialSelectionFilterSets); // Update NumRequiredInputs if you add more Inputs. This is used by Serialize.
+	check(GetNumInputs() == NumRequiredInputs + NumInitialSelectionFilterSets * 2); // Update NumRequiredInputs if you add more Inputs. This is used by Serialize.
 }
 
 void FChaosClothAssetProxyDeformerNode_v2::Evaluate(Dataflow::FContext& Context, const FDataflowOutput* Out) const
@@ -420,7 +426,6 @@ void FChaosClothAssetProxyDeformerNode_v2::Evaluate(Dataflow::FContext& Context,
 	{
 		using namespace UE::Chaos::ClothAsset;
 
-
 		// Evaluate in collection
 		FManagedArrayCollection InCollection = GetValue<FManagedArrayCollection>(Context, &Collection);
 		const TSharedRef<FManagedArrayCollection> ClothCollection = MakeShared<FManagedArrayCollection>(MoveTemp(InCollection));
@@ -431,16 +436,6 @@ void FChaosClothAssetProxyDeformerNode_v2::Evaluate(Dataflow::FContext& Context,
 		{
 			FCollectionClothSelectionFacade SelectionFacade(ClothCollection);
 
-			// Retrieve the SimVertexSelection name
-			FName SimVertexSelectionName = FName(*GetValue<FString>(Context, &SimVertexSelection.StringValue));
-			if (SimVertexSelectionName != NAME_None && (!SelectionFacade.IsValid() || !SelectionFacade.FindSelectionSet(SimVertexSelectionName)))
-			{
-				FClothDataflowTools::LogAndToastWarning(*this,
-					LOCTEXT("HasSimVertexSelectionHeadline", "Unknown SimVertexSelection."),
-					LOCTEXT("HasSimVertexSelectionDetails", "The specified SimVertexSelection does't exist within the input Cloth Collection."));
-				SimVertexSelectionName = NAME_None;
-			}
-
 			// Add the optional render deformer schema
 			if (!ClothFacade.IsValid(EClothCollectionOptionalSchemas::RenderDeformer))
 			{
@@ -454,7 +449,7 @@ void FChaosClothAssetProxyDeformerNode_v2::Evaluate(Dataflow::FContext& Context,
 			DeformerMappingDataGenerator.RenderPositions = ClothFacade.GetRenderPosition();
 			DeformerMappingDataGenerator.RenderNormals = ClothFacade.GetRenderNormal();
 			DeformerMappingDataGenerator.RenderIndices = ClothFacade.GetRenderIndices();
-			DeformerMappingDataGenerator.PointWeightMap = Private::SelectionToPointWeightMap(ClothFacade, SelectionFacade, SimVertexSelectionName);
+			DeformerMappingDataGenerator.PointWeightMap = FPointWeightMap(); // _v2 no longer compute SkinningBlend transitions
 			DeformerMappingDataGenerator.MeshToMeshFilterSet = Private::SelectionsToMeshToMeshFilterSets_v2(ClothFacade, SelectionFacade, GetSelectionFilterNames(Context));
 			DeformerMappingDataGenerator.RenderDeformerPositionBaryCoordsAndDist = ClothFacade.GetRenderDeformerPositionBaryCoordsAndDist();
 			DeformerMappingDataGenerator.RenderDeformerNormalBaryCoordsAndDist = ClothFacade.GetRenderDeformerNormalBaryCoordsAndDist();
@@ -463,7 +458,9 @@ void FChaosClothAssetProxyDeformerNode_v2::Evaluate(Dataflow::FContext& Context,
 			DeformerMappingDataGenerator.RenderDeformerWeight = ClothFacade.GetRenderDeformerWeight();
 			DeformerMappingDataGenerator.RenderDeformerSkinningBlend = ClothFacade.GetRenderDeformerSkinningBlend();
 
-			const int32 NumInfluences = DeformerMappingDataGenerator.Generate(bUseSmoothTransition, bUseMultipleInfluences, InfluenceRadius);
+			constexpr bool bUseSmoothTransition = false;  // _v2 no longer compute SkinningBlend transitions
+			constexpr bool bDoSkinningBlend = false;  // _v2 no longer compute SkinningBlend transitions
+			const int32 NumInfluences = DeformerMappingDataGenerator.Generate(bUseSmoothTransition, bUseMultipleInfluences, InfluenceRadius, bDoSkinningBlend);
 
 			for (int32 RenderPatternIndex = 0; RenderPatternIndex < ClothFacade.GetNumRenderPatterns(); ++RenderPatternIndex)
 			{
@@ -474,10 +471,6 @@ void FChaosClothAssetProxyDeformerNode_v2::Evaluate(Dataflow::FContext& Context,
 
 		SetValue(Context, MoveTemp(*ClothCollection), &Collection);
 	}
-	else if (Out->IsA<FString>(&SkinningBlendName))
-	{
-		SetValue(Context, SkinningBlendName, &SkinningBlendName);
-	}
 }
 
 TArray<Dataflow::FPin> FChaosClothAssetProxyDeformerNode_v2::AddPins()
@@ -672,7 +665,8 @@ void FChaosClothAssetProxyDeformerNode::Evaluate(Dataflow::FContext& Context, co
 			DeformerMappingDataGenerator.RenderDeformerWeight = ClothFacade.GetRenderDeformerWeight();
 			DeformerMappingDataGenerator.RenderDeformerSkinningBlend = ClothFacade.GetRenderDeformerSkinningBlend();
 
-			const int32 NumInfluences = DeformerMappingDataGenerator.Generate(bUseSmoothTransition, bUseMultipleInfluences, InfluenceRadius);
+			constexpr bool bDoSkinningBlend = true;  // Compute SkinningBlend (legacy behavior)
+			const int32 NumInfluences = DeformerMappingDataGenerator.Generate(bUseSmoothTransition, bUseMultipleInfluences, InfluenceRadius, bDoSkinningBlend);
 
 			for (int32 RenderPatternIndex = 0; RenderPatternIndex < ClothFacade.GetNumRenderPatterns(); ++RenderPatternIndex)
 			{
diff --git a/Engine/Plugins/ChaosClothAssetEditor/Source/ChaosClothAssetDataflowNodes/Public/ChaosClothAsset/ProxyDeformerNode.h b/Engine/Plugins/ChaosClothAssetEditor/Source/ChaosClothAssetDataflowNodes/Public/ChaosClothAsset/ProxyDeformerNode.h
--- a/Engine/Plugins/ChaosClothAssetEditor/Source/ChaosClothAssetDataflowNodes/Public/ChaosClothAsset/ProxyDeformerNode.h
+++ b/Engine/Plugins/ChaosClothAssetEditor/Source/ChaosClothAssetDataflowNodes/Public/ChaosClothAsset/ProxyDeformerNode.h
@@ -22,7 +22,7 @@ struct FChaosClothAssetSelectionFilterSet
 };
 
 /** Add the proxy deformer information to this cloth collection's render data. */
-USTRUCT(Meta = (DataflowCloth, Experimental))
+USTRUCT(Meta = (DataflowCloth, Experimental))  // TODO: Remove experimental for 5.5 release
 struct FChaosClothAssetProxyDeformerNode_v2 : public FDataflowNode
 {
 	GENERATED_USTRUCT_BODY()
@@ -34,15 +34,6 @@ public:
 	UPROPERTY(Meta = (DataflowInput, DataflowOutput, DataflowPassthrough = "Collection"))
 	FManagedArrayCollection Collection;
 
-	/**
-	 * The name of a selection containing all the dynamic points. Must be of group type SimVertices2D, SimVertices3D, or SimFaces.
-	 * Using an empty (or invalid) selection will make the proxy deformer consider all simulation points as dynamic points,
-	 * and will fully contribute to the render mesh animations (as opposed to using the render mesh skinning for the non dynamic points).
-	 * This selection is usually built from the same weight map set to the MaxDistance config using a WeightMapToSelection node and a very low threshold.
-	 */
-	UPROPERTY(EditAnywhere, Category = "ProxyDeformer")
-	FChaosClothAssetConnectableIStringValue SimVertexSelection;
-
 	/** Selection filter set used to restrict a selection of render vertices to a selection of simulation mesh triangles. Right click and do AddOptionPin to add more selection sets .*/
 	UPROPERTY(EditAnywhere, EditFixedSize, Category = "ProxyDeformer")
 	TArray<FChaosClothAssetSelectionFilterSet> SelectionFilterSets;
@@ -55,22 +46,6 @@ public:
 	UPROPERTY(EditAnywhere, Category = "ProxyDeformer", Meta = (EditCondition = "bUseMultipleInfluences"))
 	float InfluenceRadius = 5.f;
 
-	/**
-	 * Whether to create a smoothed _SkinningBlendWeight render weight map to ease the transition between the deformed part and the skinned part of the render mesh.
-	 * When no transition is created there will be a visible step in the rendered triangles around the edge of the kinematic/dynamic transition of the proxy simulation mesh.
-	 * The _SkinningBlendWeight render weight map is created regardless of the transition being created smooth or not, and can be later adjusted using the weight map tool.
-	 */
-	UPROPERTY(EditAnywhere, Category = "ProxyDeformer")
-	bool bUseSmoothTransition = true;
-
-	/**
-	 * The name of the render mesh weight map generated by this node detailing the contribution of the proxy deformer.
-	 * Value ranges between 0 (fully deformed) and 1 (fully skinned).
-	 * The name of this render mesh weight map cannot be changed and is only provided for further tweaking.
-	 */
-	UPROPERTY(VisibleAnywhere, Category = "ProxyDeformer", Meta = (DataflowOutput))
-	FString SkinningBlendName;
-
 	FChaosClothAssetProxyDeformerNode_v2(const Dataflow::FNodeParameters& InParam, FGuid InGuid = FGuid::NewGuid());
 
 private:
@@ -89,7 +64,7 @@ private:
 	Dataflow::TConnectionReference<FString> GetRenderConnectionReference(int32 Index) const;
 	Dataflow::TConnectionReference<FString> GetSimConnectionReference(int32 Index) const;
 
-	static constexpr int32 NumRequiredInputs = 2; // non-filter set inputs
+	static constexpr int32 NumRequiredInputs = 1; // non-filter set inputs
 	static constexpr int32 NumInitialSelectionFilterSets = 1;
 };
 
diff --git a/Engine/Plugins/ChaosClothAssetEditor/Source/ChaosClothAssetDataflowNodes/Public/ChaosClothAsset/SkinningBlendNode.h b/Engine/Plugins/ChaosClothAssetEditor/Source/ChaosClothAssetDataflowNodes/Public/ChaosClothAsset/SkinningBlendNode.h
new file mode 100644
--- /dev/null
+++ b/Engine/Plugins/ChaosClothAssetEditor/Source/ChaosClothAssetDataflowNodes/Public/ChaosClothAsset/SkinningBlendNode.h
@@ -0,0 +1,64 @@
+// Copyright Epic Games, Inc. All Rights Reserved.
+
+#pragma once
+
+#include "Dataflow/DataflowNode.h"
+#include "ChaosClothAsset/ConnectableValue.h"
+#include "GeometryCollection/ManagedArrayCollection.h"
+#include "SkinningBlendNode.generated.h"
+
+/**
+ * Initialize the RenderDeformerSkinningBlend weight map from the ProxyDeformer mapping data.
+ * 
+ * The weight map is used by the cloth render shader to decide how much to blend between skinned and simulated points.
+ * Value ranges between 0 for fully deformed, to 1 for fully skinned.
+ * 
+ * Editing the RenderDeformerSkinningBlend is possible after this node (via a WeightMap node),
+ * as the RenderDeformerSkinningBlend weight map is an intrinsic attribute of the ClothCollection.
+ * 
+ * Note: ProxyDeformer mapping data must exist on the input ClothCollection prior to using this node.
+ */
+USTRUCT(Meta = (DataflowCloth))
+struct FChaosClothAssetSkinningBlendNode : public FDataflowNode
+{
+	GENERATED_USTRUCT_BODY()
+	DATAFLOW_NODE_DEFINE_INTERNAL(FChaosClothAssetSkinningBlendNode, "SkinningBlend", "Cloth", "Cloth Simulation Skinning Blend")
+
+public:
+
+	UPROPERTY(Meta = (DataflowInput, DataflowOutput, DataflowPassthrough = "Collection"))
+	FManagedArrayCollection Collection;
+
+	/**
+	 * The name of a selection containing all the kinematic points. Must be of group type SimVertices2D, SimVertices3D, or SimFaces.
+	 * Using an empty (or invalid) selection will make this node consider all points as dynamic.
+	 * This selection is usually obtained from the MaxDistanceConfig node or built from the same weight map set to the MaxDistanceConfig node
+	 * using a WeightMapToSelection node and a very low threshold.
+	 */
+	UPROPERTY(EditAnywhere, Category = "SkinningBlend")
+	FChaosClothAssetConnectableIStringValue KinematicVertices3D;
+
+	/**
+	 * Whether to create a smoothed RenderDeformerSkinningBlend weight map to ease the transition between the deformed part and the skinned part of the render mesh.
+	 * When no transition is created there will be a visible step in the rendered triangles around the edge of the kinematic/dynamic transition of the proxy simulation mesh.
+	 * The RenderDeformerSkinningBlend weight map is created regardless of the transition being created smooth or not, and can be later adjusted by using the WeightMap node.
+	 */
+	UPROPERTY(EditAnywhere, Category = "SkinningBlend")
+	bool bUseSmoothTransition = true;
+
+	/**
+	 * The name of the render mesh weight map generated by this node detailing the contribution of the proxy deformer.
+	 * Value ranges between 0 (fully deformed) and 1 (fully skinned).
+	 * The name of this render mesh weight map cannot be changed and is only provided for further tweaking.
+	 */
+	UPROPERTY(VisibleAnywhere, Category = "SkinningBlend", Meta = (DataflowOutput))
+	FString SkinningBlendName;
+
+	FChaosClothAssetSkinningBlendNode(const Dataflow::FNodeParameters& InParam, FGuid InGuid = FGuid::NewGuid());
+
+private:
+
+	//~ Begin FDataflowNode interface
+	virtual void Evaluate(Dataflow::FContext& Context, const FDataflowOutput* Out) const override;
+	//~ End FDataflowNode interface
+};
