Cloth Asset - Remove Terminal Node inheritance from the WeightMap node.

#rnx

diff --git a/Engine/Plugins/ChaosClothAssetEditor/Source/ChaosClothAssetDataflowNodes/Private/ChaosClothAsset/WeightMapNode.cpp b/Engine/Plugins/ChaosClothAssetEditor/Source/ChaosClothAssetDataflowNodes/Private/ChaosClothAsset/WeightMapNode.cpp
--- a/Engine/Plugins/ChaosClothAssetEditor/Source/ChaosClothAssetDataflowNodes/Private/ChaosClothAsset/WeightMapNode.cpp
+++ b/Engine/Plugins/ChaosClothAssetEditor/Source/ChaosClothAssetDataflowNodes/Private/ChaosClothAsset/WeightMapNode.cpp
@@ -34,9 +34,9 @@ namespace UE::Chaos::ClothAsset::Private
 	void CalculateFinalVertexWeightValues(const TConstArrayView<float> InputMap, TArrayView<float> FinalOutputMap, EChaosClothAssetWeightMapOverrideType OverrideType, const TArray<float>& SourceVertexWeights);
 }
 
-
 FChaosClothAssetWeightMapNode::FChaosClothAssetWeightMapNode(const Dataflow::FNodeParameters& InParam, FGuid InGuid)
-	: FDataflowTerminalNode(InParam, InGuid)
+	: FDataflowNode(InParam, InGuid)
+	, Transfer(FSimpleDelegate::CreateRaw(this, &FChaosClothAssetWeightMapNode::OnTransfer))
 {
 	RegisterInputConnection(&Collection);
 	RegisterInputConnection(&InputName.StringValue, GET_MEMBER_NAME_CHECKED(FChaosClothAssetConnectableIStringValue, StringValue))
@@ -46,87 +46,61 @@ FChaosClothAssetWeightMapNode::FChaosClothAssetWeightMapNode(const Dataflow::FNo
 		.SetCanHidePin(true)
 		.SetPinIsHidden(true);
 	RegisterOutputConnection(&Collection, &Collection);
-	RegisterOutputConnection(&Name);
+	RegisterOutputConnection(&OutputName.StringValue, (FString*)nullptr, GET_MEMBER_NAME_CHECKED(FChaosClothAssetConnectableOStringValue, StringValue));
 }
 
-void FChaosClothAssetWeightMapNode::SetAssetValue(TObjectPtr<UObject> Asset, Dataflow::FContext& Context) const
+void FChaosClothAssetWeightMapNode::OnTransfer()
 {
 	using namespace UE::Chaos::ClothAsset;
 
-	if (UChaosClothAsset* const ClothAsset = Cast<UChaosClothAsset>(Asset.Get()))
+	// Create a temporary context (until we find a way to re-use the one from the tool calling this function)
+	Dataflow::FContextThreaded Context(FPlatformTime::Cycles64());
+
+	// Transfer weight map if the transfer collection input has changed and is valid
+	FManagedArrayCollection InCollection = GetValue<FManagedArrayCollection>(Context, &Collection);
+	const TSharedRef<FManagedArrayCollection> ClothCollection = MakeShared<FManagedArrayCollection>(MoveTemp(InCollection));
+	FCollectionClothConstFacade ClothFacade(ClothCollection);
+	if (ClothFacade.IsValid())  // Can only act on the collection if it is a valid cloth collection
 	{
-		if (UDataflow* const DataflowAsset = ClothAsset->GetDataflow())
-		{
-			const TSharedPtr<Dataflow::FGraph, ESPMode::ThreadSafe> Dataflow = DataflowAsset->GetDataflow();
-			if (const TSharedPtr<FDataflowNode> BaseNode = Dataflow->FindBaseNode(this->GetGuid()))  // This is basically a safe const_cast
-			{
-				FChaosClothAssetWeightMapNode* const MutableThis = static_cast<FChaosClothAssetWeightMapNode*>(BaseNode.Get());
-				check(MutableThis == this);
+		FManagedArrayCollection InTransferCollection = GetValue<FManagedArrayCollection>(Context, &TransferCollection);
+		const TSharedRef<const FManagedArrayCollection> TransferClothCollection = MakeShared<const FManagedArrayCollection>(MoveTemp(InTransferCollection));
+		FCollectionClothConstFacade TransferClothFacade(TransferClothCollection);
+		const FName InInputName = GetInputName(Context);
 
-				// Make the name a valid attribute name, and replace the value in the UI
-				FWeightMapTools::MakeWeightMapName(MutableThis->Name);
+		if (TransferClothFacade.HasWeightMap(InInputName))
+		{
+			// Remap the weights
+			TArray<float> RemappedWeights;
+			RemappedWeights.SetNumZeroed(ClothFacade.GetNumSimVertices3D());
 
-				// Transfer weight map if the transfer collection input has changed and is valid
-				FManagedArrayCollection InCollection = GetValue<FManagedArrayCollection>(Context, &Collection);
-				const TSharedRef<FManagedArrayCollection> ClothCollection = MakeShared<FManagedArrayCollection>(MoveTemp(InCollection));
-				FCollectionClothConstFacade ClothFacade(ClothCollection);
-				if (ClothFacade.IsValid())  // Can only act on the collection if it is a valid cloth collection
-				{
-					FManagedArrayCollection InTransferCollection = GetValue<FManagedArrayCollection>(Context, &TransferCollection);
-					const TSharedRef<const FManagedArrayCollection> TransferClothCollection = MakeShared<const FManagedArrayCollection>(MoveTemp(InTransferCollection));
-					FCollectionClothConstFacade TransferClothFacade(TransferClothCollection);
-					const FName InInputName = GetInputName(Context);
-
-					const uint32 NameTypeHash = HashCombineFast(GetTypeHash(InInputName), (uint32)TransferType);
-					const uint32 InTransferCollectionHash = (TransferClothFacade.HasValidSimulationData() && InInputName != NAME_None) ?
-						HashCombineFast(TransferClothFacade.CalculateWeightMapTypeHash(), NameTypeHash) : 0;  // TODO: Remove after adding the function (currently shelved!)  
-
-					if (TransferCollectionHash != InTransferCollectionHash)
-					{
-						MutableThis->TransferCollectionHash = InTransferCollectionHash;
-
-						if (TransferCollectionHash)
-						{
-							if (TransferClothFacade.HasWeightMap(InInputName))
-							{
-								// Remap the weights
-								TArray<float> RemappedWeights;
-								RemappedWeights.SetNumZeroed(ClothFacade.GetNumSimVertices3D());
-
-								switch (TransferType)
-								{
-								case EChaosClothAssetWeightMapTransferType::Use2DSimMesh:
-									Private::TransferWeightMap(
-										TransferClothFacade.GetSimPosition2D(),
-										TransferClothFacade.GetSimIndices2D(),
-										TransferClothFacade.GetSimVertex3DLookup(),
-										TransferClothFacade.GetWeightMap(InInputName),
-										ClothFacade.GetSimPosition2D(),
-										ClothFacade.GetSimIndices2D(),
-										ClothFacade.GetSimVertex3DLookup(),
-										RemappedWeights);
-									break;
-								case EChaosClothAssetWeightMapTransferType::Use3DSimMesh:
-									FClothGeometryTools::TransferWeightMap(
-										TransferClothFacade.GetSimPosition3D(),
-										TransferClothFacade.GetSimIndices3D(),
-										TransferClothFacade.GetWeightMap(InInputName),
-										ClothFacade.GetSimPosition3D(),
-										ClothFacade.GetSimNormal(),
-										ClothFacade.GetSimIndices3D(),
-										TArrayView<float>(RemappedWeights));
-									break;
-								default: unimplemented();
-								}
-
-								// TODO: Allow for transferring render weights from InTransferCollection?
-
-								MutableThis->SetVertexWeights(ClothFacade.GetWeightMap(InInputName), RemappedWeights);
-							}
-						}
-					}
-				}
+			switch (TransferType)
+			{
+			case EChaosClothAssetWeightMapTransferType::Use2DSimMesh:
+				Private::TransferWeightMap(
+					TransferClothFacade.GetSimPosition2D(),
+					TransferClothFacade.GetSimIndices2D(),
+					TransferClothFacade.GetSimVertex3DLookup(),
+					TransferClothFacade.GetWeightMap(InInputName),
+					ClothFacade.GetSimPosition2D(),
+					ClothFacade.GetSimIndices2D(),
+					ClothFacade.GetSimVertex3DLookup(),
+					RemappedWeights);
+				break;
+			case EChaosClothAssetWeightMapTransferType::Use3DSimMesh:
+				FClothGeometryTools::TransferWeightMap(
+					TransferClothFacade.GetSimPosition3D(),
+					TransferClothFacade.GetSimIndices3D(),
+					TransferClothFacade.GetWeightMap(InInputName),
+					ClothFacade.GetSimPosition3D(),
+					ClothFacade.GetSimNormal(),
+					ClothFacade.GetSimIndices3D(),
+					TArrayView<float>(RemappedWeights));
+				break;
+			default: unimplemented();
 			}
+
+			// TODO: Allow for transferring render weights from InTransferCollection?
+			SetVertexWeights(ClothFacade.GetWeightMap(InInputName), RemappedWeights);
 		}
 	}
 }
@@ -159,7 +133,7 @@ void FChaosClothAssetWeightMapNode::Evaluate(Dataflow::FContext& Context, const
 		FCollectionClothFacade ClothFacade(ClothCollection);
 		if (ClothFacade.IsValid())  // Can only act on the collection if it is a valid cloth collection
 		{
-			const FName InName(Name.IsEmpty() ? InInputName : FName(Name));
+			const FName InName(OutputName.StringValue.IsEmpty() ? InInputName : FName(OutputName.StringValue));
 
 			// Copy simulation weights into cloth collection
 
@@ -209,12 +183,35 @@ void FChaosClothAssetWeightMapNode::Evaluate(Dataflow::FContext& Context, const
 
 		SetValue(Context, MoveTemp(*ClothCollection), &Collection);
 	}
-	else if (Out->IsA<FString>(&Name))
+	else if (Out->IsA<FString>(&OutputName.StringValue))
 	{
 		FString InputNameString = GetValue<FString>(Context, &InputName.StringValue);
 		UE::Chaos::ClothAsset::FWeightMapTools::MakeWeightMapName(InputNameString);
-		SetValue(Context, Name.IsEmpty() ? InputNameString : Name, &Name);
+		SetValue(Context, OutputName.StringValue.IsEmpty() ? InputNameString : OutputName.StringValue, &OutputName.StringValue);
+	}
+}
+
+void FChaosClothAssetWeightMapNode::Serialize(FArchive& Ar)
+{
+	if (Ar.IsLoading())
+	{
+PRAGMA_DISABLE_DEPRECATION_WARNINGS
+		if (!Name.IsEmpty() && OutputName.StringValue.IsEmpty())  // TODO: Discard for v2
+		{
+			OutputName.StringValue = MoveTemp(Name);
+			Name.Empty();
+		}
+PRAGMA_ENABLE_DEPRECATION_WARNINGS
+	}
+}
+
+FDataflowOutput* FChaosClothAssetWeightMapNode::RedirectSerializedOutput(const FName& MissingOutputName)
+{
+	if (MissingOutputName == TEXT("Name"))
+	{
+		return FindOutput(FName(TEXT("OutputName.StringValue")));
 	}
+	return nullptr;
 }
 
 FName FChaosClothAssetWeightMapNode::GetInputName(Dataflow::FContext& Context) const
@@ -222,7 +219,7 @@ FName FChaosClothAssetWeightMapNode::GetInputName(Dataflow::FContext& Context) c
 	FString InputNameString = GetValue<FString>(Context, &InputName.StringValue);
 	UE::Chaos::ClothAsset::FWeightMapTools::MakeWeightMapName(InputNameString);
 	const FName InInputName(*InputNameString);
-	return InInputName != NAME_None ? InInputName : FName(Name);
+	return InInputName != NAME_None ? InInputName : FName(OutputName.StringValue);
 }
 
 void FChaosClothAssetWeightMapNode::SetVertexWeights(const TConstArrayView<float> InputMap, const TArray<float>& FinalValues)
@@ -245,7 +242,7 @@ public:
 		NodeGuid(Node.GetGuid()),
 		SavedWeights(Node.GetVertexWeights()),
 		SavedMapOverrideType(Node.MapOverrideType),
-		SavedWeightMapName(Node.Name)
+		SavedWeightMapName(Node.OutputName.StringValue)
 	{}
 
 private:
@@ -281,7 +278,7 @@ private:
 				{
 					Swap(Node->GetVertexWeights(), SavedWeights);
 					Swap(Node->MapOverrideType, SavedMapOverrideType);
-					Swap(Node->Name, SavedWeightMapName);
+					Swap(Node->OutputName.StringValue, SavedWeightMapName);
 
 					Node->Invalidate();
 				}
diff --git a/Engine/Plugins/ChaosClothAssetEditor/Source/ChaosClothAssetDataflowNodes/Public/ChaosClothAsset/WeightMapNode.h b/Engine/Plugins/ChaosClothAssetEditor/Source/ChaosClothAssetDataflowNodes/Public/ChaosClothAsset/WeightMapNode.h
--- a/Engine/Plugins/ChaosClothAssetEditor/Source/ChaosClothAssetDataflowNodes/Public/ChaosClothAsset/WeightMapNode.h
+++ b/Engine/Plugins/ChaosClothAssetEditor/Source/ChaosClothAssetDataflowNodes/Public/ChaosClothAsset/WeightMapNode.h
@@ -3,6 +3,7 @@
 #pragma once 
 
 #include "ChaosClothAsset/AddWeightMapNode.h"
+#include "Dataflow/DataflowFunctionProperty.h"
 #include "GeometryCollection/ManagedArrayCollection.h"
 #include "ChaosClothAsset/ConnectableValue.h"
 #include "WeightMapNode.generated.h"
@@ -17,7 +18,7 @@ enum class EChaosClothAssetWeightMapMeshTarget : uint8
 
 /** Painted weight map attributes node. */
 USTRUCT(Meta = (DataflowCloth))
-struct FChaosClothAssetWeightMapNode : public FDataflowTerminalNode
+struct FChaosClothAssetWeightMapNode : public FDataflowNode
 {
 	GENERATED_USTRUCT_BODY()
 	DATAFLOW_NODE_DEFINE_INTERNAL(FChaosClothAssetWeightMapNode, "WeightMap", "Cloth", "Cloth Weight Map")
@@ -29,24 +30,21 @@ public:
 	UPROPERTY(Meta = (DataflowInput, DataflowOutput, DataflowPassthrough = "Collection"))
 	FManagedArrayCollection Collection;
 
-	/**
-	 * The collection used to transfer weight map from.
-	 * Connecting a collection containing a weight map with Input Name (or Name if Input Name is empty)
-	 * will transfer the weights to the input collection vertices.
-	 * Note this operation only happens once when the TransferCollection is first connected, or updated.
-	 * Changing the InputName or the TransferType will also redo the transfer operation.
-	 */
+	/** The collection used to transfer weight map from. */
 	UPROPERTY(Meta = (DataflowInput))
 	FManagedArrayCollection TransferCollection;
 
 	/** The name to be set as a weight map attribute. */
-	UPROPERTY(EditAnywhere, Category = "Weight Map", Meta = (DataflowOutput))
-	FString Name;
+	UPROPERTY(EditAnywhere, Category = "Weight Map")
+	FChaosClothAssetConnectableOStringValue OutputName;
 
-	/** The name to populate this map from and override based on Map Override Type. Name will be used if Input Name is empty.*/
+	/** The name to populate this map from and override based on Map Override Type. Output Name will be used if Input Name is empty. */
 	UPROPERTY(EditAnywhere, Category = "Weight Map")
 	FChaosClothAssetConnectableIStringValue InputName = { TEXT("") };
 
+	UPROPERTY(EditAnywhere, Category = "Weight Map")
+	EChaosClothAssetWeightMapMeshTarget MeshTarget = EChaosClothAssetWeightMapMeshTarget::Simulation;
+
 	/** How to apply this node's weight values onto existing maps. Changing this value will change the output map.
 	 *  To change how the node's stored weights are calculated, change the equivalent value on the Weight Map Paint Tool context.*/
 	UPROPERTY(EditAnywhere, Category = "Weight Map")
@@ -56,11 +54,12 @@ public:
 	 * The type of transfer used to transfer the weight map when a TransferCollection is connected.
 	 * This property is disabled when no TransferCollection input has been conencted.
 	 */
-	UPROPERTY(EditAnywhere, Category = "Weight Map", Meta = (EditCondition = "TransferCollectionHash != 0"))
+	UPROPERTY(EditAnywhere, Category = "Weight Map")
 	EChaosClothAssetWeightMapTransferType TransferType = EChaosClothAssetWeightMapTransferType::Use2DSimMesh;
 
-	UPROPERTY(EditAnywhere, Category = "Weight Map")
-	EChaosClothAssetWeightMapMeshTarget MeshTarget = EChaosClothAssetWeightMapMeshTarget::Simulation;
+	/** Transfer the weight map from the connected Transfer Collection containing a weight map with Input Name (or Output Name if Input Name is empty). */
+	UPROPERTY(EditAnywhere, Category = "Weight Map", Meta = (ButtonImage = "Icons.Convert"))
+	FDataflowFunctionProperty Transfer;
 
 	FChaosClothAssetWeightMapNode(const Dataflow::FNodeParameters& InParam, FGuid InGuid = FGuid::NewGuid());
 
@@ -73,9 +72,6 @@ private:
 	const TArray<float>& GetVertexWeights() const { return VertexWeights; }
 	TArray<float>& GetVertexWeights() { return VertexWeights; }
 
-	UPROPERTY()
-	TArray<float> VertexWeights;
-
 	// These methods are exported for UClothEditorWeightMapPaintTool which lives in a different module.
 	FName CHAOSCLOTHASSETDATAFLOWNODES_API GetInputName(Dataflow::FContext& Context) const;
 
@@ -85,10 +81,17 @@ private:
 	void CHAOSCLOTHASSETDATAFLOWNODES_API CalculateFinalVertexWeightValues(const TConstArrayView<float> InputMap, TArrayView<float> FinalOutputMap) const;
 
 	//~ Begin FDataflowNode interface
-	virtual void SetAssetValue(TObjectPtr<UObject> Asset, Dataflow::FContext& Context) const override;
 	virtual void Evaluate(Dataflow::FContext& Context, const FDataflowOutput* Out) const override;
+	virtual void Serialize(FArchive& Ar) override;
+	virtual FDataflowOutput* RedirectSerializedOutput(const FName& MissingOutputName) override;
 	//~ End FDataflowNode interface
 
+	void OnTransfer();
+
+	UPROPERTY()
+	TArray<float> VertexWeights;
+
+	UE_DEPRECATED(5.5, "Use OutputName instead.")
 	UPROPERTY()
-	uint32 TransferCollectionHash = 0;
+	FString Name;  // TODO: Discard in future v2. but keep for backward compatibility here as some weight maps have been created with it
 };
diff --git a/Engine/Plugins/ChaosClothAssetEditor/Source/ChaosClothAssetEditorTools/Private/ChaosClothAsset/ClothWeightMapPaintTool.cpp b/Engine/Plugins/ChaosClothAssetEditor/Source/ChaosClothAssetEditorTools/Private/ChaosClothAsset/ClothWeightMapPaintTool.cpp
--- a/Engine/Plugins/ChaosClothAssetEditor/Source/ChaosClothAssetEditorTools/Private/ChaosClothAsset/ClothWeightMapPaintTool.cpp
+++ b/Engine/Plugins/ChaosClothAssetEditor/Source/ChaosClothAssetEditorTools/Private/ChaosClothAsset/ClothWeightMapPaintTool.cpp
@@ -195,10 +195,10 @@ void UClothEditorWeightMapPaintTool::Setup()
 	PolygonSelectionMechanic->Initialize(DynamicMeshComponent, GradientSelectionTopology.Get(), [this]() { return MeshSpatial.Get(); });
 
 	UpdateWeightMapProperties = NewObject<UClothEditorUpdateWeightMapProperties>(this);
-	UpdateWeightMapProperties->Name = WeightMapNodeToUpdate->Name;
+	UpdateWeightMapProperties->Name = WeightMapNodeToUpdate->OutputName.StringValue;
 	UpdateWeightMapProperties->MapOverrideType = WeightMapNodeToUpdate->MapOverrideType;
 
-	UpdateWeightMapProperties->WatchProperty(WeightMapNodeToUpdate->Name, [this](const FString& NewName)
+	UpdateWeightMapProperties->WatchProperty(WeightMapNodeToUpdate->OutputName.StringValue, [this](const FString& NewName)
 	{
 		UpdateWeightMapProperties->Name = NewName;
 	});
@@ -490,7 +490,7 @@ void UClothEditorWeightMapPaintTool::Setup()
 		}
 	}
 
-	UpdateWeightMapProperties->Name = WeightMapNodeToUpdate->Name;
+	UpdateWeightMapProperties->Name = WeightMapNodeToUpdate->OutputName.StringValue;
 	UpdateWeightMapProperties->MapOverrideType = WeightMapNodeToUpdate->MapOverrideType;
 	SetToolPropertySourceEnabled(UpdateWeightMapProperties, true);
 
@@ -1777,7 +1777,7 @@ void UClothEditorWeightMapPaintTool::OnTick(float DeltaTime)
 bool UClothEditorWeightMapPaintTool::CanAccept() const
 {
 	return bAnyChangeMade || 
-		UpdateWeightMapProperties->Name != WeightMapNodeToUpdate->Name || 
+		UpdateWeightMapProperties->Name != WeightMapNodeToUpdate->OutputName.StringValue ||
 		UpdateWeightMapProperties->MapOverrideType != WeightMapNodeToUpdate->MapOverrideType;
 }
 
@@ -1985,7 +1985,7 @@ void UClothEditorWeightMapPaintTool::UpdateSelectedNode()
 
 
 	WeightMapNodeToUpdate->MapOverrideType = UpdateWeightMapProperties->MapOverrideType;
-	WeightMapNodeToUpdate->Name = UpdateWeightMapProperties->Name;
+	WeightMapNodeToUpdate->OutputName.StringValue = UpdateWeightMapProperties->Name;
 
 	const bool bIsRenderMode = (ClothEditorContextObject->GetConstructionViewMode() == UE::Chaos::ClothAsset::EClothPatternVertexType::Render);
 
