CIS - Fixed issue 540556: WeightMapNode.h(21,8): warning: 'Name' is deprecated: Use OutputName instead.

#rnx

diff --git a/Engine/Plugins/ChaosClothAssetEditor/Source/ChaosClothAssetDataflowNodes/Public/ChaosClothAsset/WeightMapNode.h b/Engine/Plugins/ChaosClothAssetEditor/Source/ChaosClothAssetDataflowNodes/Public/ChaosClothAsset/WeightMapNode.h
--- a/Engine/Plugins/ChaosClothAssetEditor/Source/ChaosClothAssetDataflowNodes/Public/ChaosClothAsset/WeightMapNode.h
+++ b/Engine/Plugins/ChaosClothAssetEditor/Source/ChaosClothAssetDataflowNodes/Public/ChaosClothAsset/WeightMapNode.h
@@ -17,9 +17,11 @@ enum class EChaosClothAssetWeightMapMeshTarget : uint8
 };
 
 /** Painted weight map attributes node. */
+PRAGMA_DISABLE_DEPRECATION_WARNINGS  // For Name implicit operators.
 USTRUCT(Meta = (DataflowCloth))
 struct FChaosClothAssetWeightMapNode : public FDataflowNode
 {
+PRAGMA_ENABLE_DEPRECATION_WARNINGS
 	GENERATED_USTRUCT_BODY()
 	DATAFLOW_NODE_DEFINE_INTERNAL(FChaosClothAssetWeightMapNode, "WeightMap", "Cloth", "Cloth Weight Map")
 	DATAFLOW_NODE_RENDER_TYPE("SurfaceRender", FName("FClothCollection"), "Collection")
