Dataflow Editor Tool Registry: a map from Node types to Tools that the Dataflow Editor can use.
Also:
- Replace UClothEditorContextObject with UDataflowContextObject in the cloth editor so that Cloth tools can be used in the Dataflow Editor
- Dataflow Editor: get the input collection of the selected node and add it to the context object when a tool starts
- Rename PrimarySelectedNode to SelectedNode and PrimaryRenderCollection to RenderCollection in UDataflowContextObject

diff --git a/Engine/Plugins/ChaosClothAssetEditor/Source/ChaosClothAssetEditor/Private/ChaosClothAsset/ClothEditorMode.cpp b/Engine/Plugins/ChaosClothAssetEditor/Source/ChaosClothAssetEditor/Private/ChaosClothAsset/ClothEditorMode.cpp
--- a/Engine/Plugins/ChaosClothAssetEditor/Source/ChaosClothAssetEditor/Private/ChaosClothAsset/ClothEditorMode.cpp
+++ b/Engine/Plugins/ChaosClothAssetEditor/Source/ChaosClothAssetEditor/Private/ChaosClothAsset/ClothEditorMode.cpp
@@ -64,6 +64,7 @@
 #include "ChaosClothAsset/ClothGeometryTools.h"
 #include "DynamicMesh/NonManifoldMappingSupport.h"
 #include "DynamicMesh/MeshNormals.h"
+#include "Dataflow/DataflowRenderingFactory.h"	// For Dataflow View Modes
 
 #include UE_INLINE_GENERATED_CPP_BY_NAME(ClothEditorMode)
 
@@ -234,13 +235,13 @@ void UChaosClothAssetEditorMode::RegisterClothTool(TSharedPtr<FUICommandInfo> UI
 			// Check if we need to switch view modes before starting the tool
 			TArray<UE::Chaos::ClothAsset::EClothPatternVertexType> SupportedModes;
 
-			const UClothEditorContextObject* EditorContextObject = ToolsContext->ContextObjectStore->FindContext<UClothEditorContextObject>();
+			const UDataflowContextObject* EditorContextObject = ToolsContext->ContextObjectStore->FindContext<UDataflowContextObject>();
 			if (!EditorContextObject)
 			{
 				InitializeContextObject();
-				EditorContextObject = ToolsContext->ContextObjectStore->FindContext<UClothEditorContextObject>();
+				EditorContextObject = ToolsContext->ContextObjectStore->FindContext<UDataflowContextObject>();
 			}
-			checkf(EditorContextObject, TEXT("Failed to find or create UClothEditorContextObject"));
+			checkf(EditorContextObject, TEXT("Failed to find or create ContextObject"));
 			ClothToolBuilder->GetSupportedViewModes(*EditorContextObject, SupportedModes);
 
 			bDynamicMeshUseInputCollection = true;
@@ -1026,12 +1027,7 @@ void UChaosClothAssetEditorMode::ReinitializeDynamicMeshComponents()
 
 	// Update the context object with the ConstructionViewMode and Collection used to build the DynamicMeshComponents, so 
 	// tools know how to use the components.
-	UEditorInteractiveToolsContext* const RestSpaceToolsContext = GetInteractiveToolsContext();
-	UClothEditorContextObject* EditorContextObject = RestSpaceToolsContext->ContextObjectStore->FindContext<UClothEditorContextObject>();
-	if (ensure(EditorContextObject))
-	{
-		EditorContextObject->SetClothCollection(ConstructionViewMode, Collection, bDynamicMeshUseInputCollection);
-	}
+	UpdateContextObject(Collection);
 }
 
 void UChaosClothAssetEditorMode::RefocusRestSpaceViewportClient()
@@ -1437,8 +1433,8 @@ bool UChaosClothAssetEditorMode::CanChangeConstructionViewModeTo(UE::Chaos::Clot
 	const UEditorInteractiveToolsContext* const RestSpaceToolsContext = GetInteractiveToolsContext();
 	checkf(RestSpaceToolsContext, TEXT("Cloth Editor Mode doesn't have a valid InteractiveToolsContext"));
 
-	const UClothEditorContextObject* const EditorContextObject = RestSpaceToolsContext->ContextObjectStore->FindContext<UClothEditorContextObject>();
-	checkf(EditorContextObject, TEXT("UClothEditorContextObject not found in ContextObjectStore despite having an active tool. This should have been created by the time a tool is activated"));
+	const UDataflowContextObject* const EditorContextObject = RestSpaceToolsContext->ContextObjectStore->FindContext<UDataflowContextObject>();
+	checkf(EditorContextObject, TEXT("ContextObject not found in ContextObjectStore despite having an active tool. This should have been created by the time a tool is activated"));
 
 	TArray<UE::Chaos::ClothAsset::EClothPatternVertexType> SupportedViewModes;
 	ClothToolBuilder->GetSupportedViewModes(*EditorContextObject, SupportedViewModes);
@@ -1588,25 +1584,48 @@ void UChaosClothAssetEditorMode::InitializeContextObject()
 {
 	UEditorInteractiveToolsContext* const RestSpaceToolsContext = GetInteractiveToolsContext();
 
-	UClothEditorContextObject* EditorContextObject = RestSpaceToolsContext->ContextObjectStore->FindContext<UClothEditorContextObject>();
-	if (!EditorContextObject)
+	// Dataflow context object
+	UDataflowContextObject* DataflowContextObject = RestSpaceToolsContext->ContextObjectStore->FindContext<UDataflowContextObject>();
+	if (!DataflowContextObject)
 	{
-		EditorContextObject = NewObject<UClothEditorContextObject>();
-		RestSpaceToolsContext->ContextObjectStore->AddContextObject(EditorContextObject);
+		DataflowContextObject = NewObject<UDataflowContextObject>();
+		RestSpaceToolsContext->ContextObjectStore->AddContextObject(DataflowContextObject);
+	}
+
+	DataflowContextObject->SetConstructionViewMode(Dataflow::FRenderingViewModeFactory::GetInstance().GetViewMode(ClothViewModeToDataflowViewModeName(ConstructionViewMode)));
+
+	if (const TSharedPtr<SDataflowGraphEditor> GraphEditor = DataflowGraphEditor.Pin())
+	{
+		if (UEdGraphNode* const SingleSelectedNode = GraphEditor->GetSingleSelectedNode())
+		{
+			if (UDataflowEdNode* const SelectedDataflowEdNode = Cast<UDataflowEdNode>(SingleSelectedNode))
+			{
+				DataflowContextObject->SetSelectedNode(SelectedDataflowEdNode);
+			}
+		}
 	}
-	constexpr bool bUsingInputCollection = false;
-	EditorContextObject->Init(DataflowGraphEditor, DataflowContext, ConstructionViewMode, SelectedClothCollection, bUsingInputCollection);
+	DataflowContextObject->SetDataflowContext(DataflowContext.Pin());
+	DataflowContextObject->SetDataflowAsset(DataflowGraph.Get());
+}
 
-	check(EditorContextObject);
+void UChaosClothAssetEditorMode::UpdateContextObject(const TSharedPtr<FManagedArrayCollection>& Collection)
+{
+	UEditorInteractiveToolsContext* const RestSpaceToolsContext = GetInteractiveToolsContext();
+	
+	if (UDataflowContextObject* DataflowContextObject = RestSpaceToolsContext->ContextObjectStore->FindContext<UDataflowContextObject>())
+	{
+		DataflowContextObject->SetConstructionViewMode(Dataflow::FRenderingViewModeFactory::GetInstance().GetViewMode(ClothViewModeToDataflowViewModeName(ConstructionViewMode)));
 
+		DataflowContextObject->SetSelectedCollection(Collection, bDynamicMeshUseInputCollection);
+	}
 }
 
 void UChaosClothAssetEditorMode::DeleteContextObject()
 {
 	UEditorInteractiveToolsContext* const RestSpaceToolsContext = GetInteractiveToolsContext();
-	if (UClothEditorContextObject* ClothEditorContextObject = RestSpaceToolsContext->ContextObjectStore->FindContext<UClothEditorContextObject>())
+	if (UDataflowContextObject* DataflowContextObject = RestSpaceToolsContext->ContextObjectStore->FindContext<UDataflowContextObject>())
 	{
-		RestSpaceToolsContext->ContextObjectStore->RemoveContextObject(ClothEditorContextObject);
+		RestSpaceToolsContext->ContextObjectStore->RemoveContextObject(DataflowContextObject);
 	}
 }
 
@@ -1614,10 +1633,10 @@ void UChaosClothAssetEditorMode::SetDataflowContext(TWeakPtr<Dataflow::FEngineCo
 {
 	DataflowContext = InDataflowContext;
 	UEditorInteractiveToolsContext* const RestSpaceToolsContext = GetInteractiveToolsContext();
-	if (UClothEditorContextObject* ClothEditorContextObject = RestSpaceToolsContext->ContextObjectStore->FindContext<UClothEditorContextObject>())
+	if (UDataflowContextObject* DataflowContextObject = RestSpaceToolsContext->ContextObjectStore->FindContext<UDataflowContextObject>())
 	{
-		ClothEditorContextObject->SetDataflowContext(DataflowContext);
-	}	
+		DataflowContextObject->SetDataflowContext(DataflowContext.Pin());
+	}
 }
 
 void UChaosClothAssetEditorMode::SetDataflowGraphEditor(TSharedPtr<SDataflowGraphEditor> InGraphEditor)
diff --git a/Engine/Plugins/ChaosClothAssetEditor/Source/ChaosClothAssetEditor/Private/ChaosClothAsset/ClothEditorMode.h b/Engine/Plugins/ChaosClothAssetEditor/Source/ChaosClothAssetEditor/Private/ChaosClothAsset/ClothEditorMode.h
--- a/Engine/Plugins/ChaosClothAssetEditor/Source/ChaosClothAssetEditor/Private/ChaosClothAsset/ClothEditorMode.h
+++ b/Engine/Plugins/ChaosClothAssetEditor/Source/ChaosClothAssetEditor/Private/ChaosClothAsset/ClothEditorMode.h
@@ -228,6 +228,7 @@ private:
 
 
 	void InitializeContextObject();
+	void UpdateContextObject(const TSharedPtr<FManagedArrayCollection>& Collection);
 	void DeleteContextObject();
 
 	bool IsComponentSelected(const UPrimitiveComponent* InComponent);
diff --git a/Engine/Plugins/ChaosClothAssetEditor/Source/ChaosClothAssetEditor/Private/ChaosClothAsset/ClothEditorToolkit.cpp b/Engine/Plugins/ChaosClothAssetEditor/Source/ChaosClothAssetEditor/Private/ChaosClothAsset/ClothEditorToolkit.cpp
--- a/Engine/Plugins/ChaosClothAssetEditor/Source/ChaosClothAssetEditor/Private/ChaosClothAsset/ClothEditorToolkit.cpp
+++ b/Engine/Plugins/ChaosClothAssetEditor/Source/ChaosClothAssetEditor/Private/ChaosClothAsset/ClothEditorToolkit.cpp
@@ -43,6 +43,7 @@
 #include "Toolkits/AssetEditorToolkitMenuContext.h"
 #include "FileHelpers.h"
 #include "UObject/PackageReload.h"
+#include "ContextObjectStore.h"
 
 #define LOCTEXT_NAMESPACE "ChaosClothAssetEditorToolkit"
 
@@ -1347,6 +1348,29 @@ void FChaosClothAssetEditorToolkit::OnNodeSelectionChanged(const TSet<UObject*>&
 			//   to be built from the InputCollection rather than the (Output)Collection.
 			const bool bDeferDynamicMeshInitForTool = NewSelection.Num() == 1;
 			ClothMode->SetSelectedClothCollection(Collection, InputCollection, bDeferDynamicMeshInitForTool);
+
+			// Update selected node in Context Object
+			if (UDataflowContextObject* DataflowContextObject = ToolsContext->ContextObjectStore->FindContext<UDataflowContextObject>())
+			{
+				DataflowContextObject->SetSelectedNode(nullptr);
+
+				if (TSharedPtr<FDataflowNode> SelectedDataflowNode = GetSelectedDataflowNode())
+				{
+					for (UObject* const Selected : NewSelection)
+					{
+						if (UDataflowEdNode* const Node = Cast<UDataflowEdNode>(Selected))
+						{
+							if (const TSharedPtr<const FDataflowNode> DataflowNode = Node->GetDataflowNode())
+							{
+								if (DataflowNode->GetGuid() == SelectedDataflowNode->GetGuid())
+								{
+									DataflowContextObject->SetSelectedNode(Node);
+								}
+							}
+						}
+					}
+				}
+			}
 		}
 
 		if (Outliner)
diff --git a/Engine/Plugins/ChaosClothAssetEditor/Source/ChaosClothAssetEditorTools/Private/ChaosClothAsset/ChaosClothAssetEditorToolsModule.cpp b/Engine/Plugins/ChaosClothAssetEditor/Source/ChaosClothAssetEditorTools/Private/ChaosClothAsset/ChaosClothAssetEditorToolsModule.cpp
--- a/Engine/Plugins/ChaosClothAssetEditor/Source/ChaosClothAssetEditorTools/Private/ChaosClothAsset/ChaosClothAssetEditorToolsModule.cpp
+++ b/Engine/Plugins/ChaosClothAssetEditor/Source/ChaosClothAssetEditorTools/Private/ChaosClothAsset/ChaosClothAssetEditorToolsModule.cpp
@@ -2,6 +2,12 @@
 
 #include "Modules/ModuleManager.h"
 
+#include "Dataflow/DataflowToolRegistry.h"
+#include "ChaosClothAsset/ClothEditorToolBuilders.h"
+#include "ChaosClothAsset/SelectionNode.h"
+#include "ChaosClothAsset/TransferSkinWeightsNode.h"
+#include "ChaosClothAsset/WeightMapNode.h"
+
 namespace UE::Chaos::ClothAsset
 {
 	class FChaosClothAssetEditorToolsModule : public IModuleInterface
@@ -9,10 +15,18 @@ namespace UE::Chaos::ClothAsset
 	public:
 		virtual void StartupModule() override
 		{
+			Dataflow::FDataflowToolRegistry& ToolRegistry = Dataflow::FDataflowToolRegistry::Get();
+			ToolRegistry.AddNodeToToolMapping(FChaosClothAssetWeightMapNode::StaticType(), NewObject<UClothEditorWeightMapPaintToolBuilder>());
+			ToolRegistry.AddNodeToToolMapping(FChaosClothAssetSelectionNode::StaticType(), NewObject<UClothMeshSelectionToolBuilder>());
+			ToolRegistry.AddNodeToToolMapping(FChaosClothAssetTransferSkinWeightsNode::StaticType(), NewObject<UClothTransferSkinWeightsToolBuilder>());
 		}
 
 		virtual void ShutdownModule() override
 		{
+			Dataflow::FDataflowToolRegistry& ToolRegistry = Dataflow::FDataflowToolRegistry::Get();
+			ToolRegistry.RemoveNodeToToolMapping(FChaosClothAssetWeightMapNode::StaticType());
+			ToolRegistry.RemoveNodeToToolMapping(FChaosClothAssetSelectionNode::StaticType());
+			ToolRegistry.RemoveNodeToToolMapping(FChaosClothAssetTransferSkinWeightsNode::StaticType());
 		}
 	};
 }
diff --git a/Engine/Plugins/ChaosClothAssetEditor/Source/ChaosClothAssetEditorTools/Private/ChaosClothAsset/ClothEditorToolBuilders.cpp b/Engine/Plugins/ChaosClothAssetEditor/Source/ChaosClothAssetEditorTools/Private/ChaosClothAsset/ClothEditorToolBuilders.cpp
--- a/Engine/Plugins/ChaosClothAssetEditor/Source/ChaosClothAssetEditorTools/Private/ChaosClothAsset/ClothEditorToolBuilders.cpp
+++ b/Engine/Plugins/ChaosClothAssetEditor/Source/ChaosClothAssetEditorTools/Private/ChaosClothAsset/ClothEditorToolBuilders.cpp
@@ -4,12 +4,13 @@
 #include "ChaosClothAsset/ClothCollectionGroup.h"
 #include "ChaosClothAsset/ClothPatternVertexType.h"
 #include "TargetInterfaces/PrimitiveComponentBackedTarget.h"
-#include "ChaosClothAsset/ClothEditorContextObject.h"
 #include "ChaosClothAsset/SelectionNode.h"
 #include "ChaosClothAsset/WeightMapNode.h"
 #include "ToolContextInterfaces.h"
 #include "ToolTargetManager.h"
 #include "ContextObjectStore.h"
+#include "Dataflow/DataflowContextObject.h"
+#include "Dataflow/DataflowRenderingViewMode.h"
 
 // Tools
 #include "ClothMeshSelectionTool.h"
@@ -19,10 +20,10 @@
 
 // ------------------- Weight Map Paint Tool -------------------
 
-void UClothEditorWeightMapPaintToolBuilder::GetSupportedViewModes(const UClothEditorContextObject& ContextObject, TArray<UE::Chaos::ClothAsset::EClothPatternVertexType>& Modes) const
+void UClothEditorWeightMapPaintToolBuilder::GetSupportedViewModes(const UDataflowContextObject& ContextObject, TArray<UE::Chaos::ClothAsset::EClothPatternVertexType>& Modes) const
 {
 	using namespace UE::Chaos::ClothAsset;
-	const FChaosClothAssetWeightMapNode* const WeightMapNode = ContextObject.GetSingleSelectedNodeOfType<FChaosClothAssetWeightMapNode>();
+	const FChaosClothAssetWeightMapNode* const WeightMapNode = ContextObject.GetSelectedNodeOfType<FChaosClothAssetWeightMapNode>();
 	if (WeightMapNode)
 	{
 		if (WeightMapNode->MeshTarget == EChaosClothAssetWeightMapMeshTarget::Simulation)
@@ -41,8 +42,11 @@ void UClothEditorWeightMapPaintToolBuilder::GetSupportedViewModes(const UClothEd
 		// No node selected. This happens if we start the tool due to pushing the button in the toolbar -- the tool starts before the node selection can change.
 		// In this case lock to either sim or render mode, whatever is current.
 		// TODO: See if we can have the button action select the node before attempting to start the tool.
-		const EClothPatternVertexType CurrentViewMode = ContextObject.GetConstructionViewMode();
-		if (CurrentViewMode == EClothPatternVertexType::Render)
+		
+		const EClothPatternVertexType ViewMode = DataflowViewModeToClothViewMode(ContextObject.GetConstructionViewMode());
+		const bool bViewModeIsRender = (ViewMode == EClothPatternVertexType::Render);
+		
+		if (bViewModeIsRender)
 		{
 			Modes.Add(EClothPatternVertexType::Render);
 		}
@@ -59,9 +63,9 @@ UMeshSurfacePointTool* UClothEditorWeightMapPaintToolBuilder::CreateNewTool(cons
 	UClothEditorWeightMapPaintTool* PaintTool = NewObject<UClothEditorWeightMapPaintTool>(SceneState.ToolManager);
 	PaintTool->SetWorld(SceneState.World);
 
-	if (UClothEditorContextObject* ContextObject = SceneState.ToolManager->GetContextObjectStore()->FindContext<UClothEditorContextObject>())
+	if (UDataflowContextObject* const DataflowContextObject = SceneState.ToolManager->GetContextObjectStore()->FindContext<UDataflowContextObject>())
 	{
-		PaintTool->SetClothEditorContextObject(ContextObject);
+		PaintTool->SetDataflowContextObject(DataflowContextObject);
 	}
 
 	return PaintTool;
@@ -69,7 +73,7 @@ UMeshSurfacePointTool* UClothEditorWeightMapPaintToolBuilder::CreateNewTool(cons
 
 // ------------------- Selection Tool -------------------
 
-void UClothMeshSelectionToolBuilder::GetSupportedViewModes(const UClothEditorContextObject& ContextObject, TArray<UE::Chaos::ClothAsset::EClothPatternVertexType>& Modes) const
+void UClothMeshSelectionToolBuilder::GetSupportedViewModes(const UDataflowContextObject& ContextObject, TArray<UE::Chaos::ClothAsset::EClothPatternVertexType>& Modes) const
 {
 	// TODO: When the Secondary Selection set is removed, update this function to be similar to UClothEditorWeightMapPaintToolBuilder::GetSupportedViewModes above
 	Modes.Add(UE::Chaos::ClothAsset::EClothPatternVertexType::Sim3D);
@@ -85,9 +89,9 @@ const FToolTargetTypeRequirements& UClothMeshSelectionToolBuilder::GetTargetRequ
 
 bool UClothMeshSelectionToolBuilder::CanBuildTool(const FToolBuilderState& SceneState) const
 {
-	if (UClothEditorContextObject* const ContextObject = SceneState.ToolManager->GetContextObjectStore()->FindContext<UClothEditorContextObject>())
+	if (UDataflowContextObject* const DataflowContextObject = SceneState.ToolManager->GetContextObjectStore()->FindContext<UDataflowContextObject>())
 	{
-		return ContextObject->GetSingleSelectedNodeOfType<FChaosClothAssetSelectionNode>() != nullptr && (SceneState.TargetManager->CountSelectedAndTargetable(SceneState, GetTargetRequirements()) == 1);
+		return DataflowContextObject->GetSelectedNodeOfType<FChaosClothAssetSelectionNode>() != nullptr && (SceneState.TargetManager->CountSelectedAndTargetable(SceneState, GetTargetRequirements()) == 1);
 	}
 
 	return false;
@@ -101,9 +105,9 @@ UInteractiveTool* UClothMeshSelectionToolBuilder::BuildTool(const FToolBuilderSt
 	NewTool->SetTarget(Target);
 	NewTool->SetWorld(SceneState.World);
 
-	if (UClothEditorContextObject* const ContextObject = SceneState.ToolManager->GetContextObjectStore()->FindContext<UClothEditorContextObject>())
+	if (UDataflowContextObject* const DataflowContextObject = SceneState.ToolManager->GetContextObjectStore()->FindContext<UDataflowContextObject>())
 	{
-		NewTool->SetClothEditorContextObject(ContextObject);
+		NewTool->SetDataflowContextObject(DataflowContextObject);
 	}
 
 	return NewTool;
@@ -112,7 +116,7 @@ UInteractiveTool* UClothMeshSelectionToolBuilder::BuildTool(const FToolBuilderSt
 
 // ------------------- Skin Weight Transfer Tool -------------------
 
-void UClothTransferSkinWeightsToolBuilder::GetSupportedViewModes(const UClothEditorContextObject& ContextObject, TArray<UE::Chaos::ClothAsset::EClothPatternVertexType>& Modes) const
+void UClothTransferSkinWeightsToolBuilder::GetSupportedViewModes(const UDataflowContextObject& ContextObject, TArray<UE::Chaos::ClothAsset::EClothPatternVertexType>& Modes) const
 {
 	Modes.Add(UE::Chaos::ClothAsset::EClothPatternVertexType::Sim3D);
 }
@@ -121,9 +125,9 @@ USingleSelectionMeshEditingTool* UClothTransferSkinWeightsToolBuilder::CreateNew
 {
 	UClothTransferSkinWeightsTool* NewTool = NewObject<UClothTransferSkinWeightsTool>(SceneState.ToolManager);
 
-	if (UClothEditorContextObject* ContextObject = SceneState.ToolManager->GetContextObjectStore()->FindContext<UClothEditorContextObject>())
+	if (UDataflowContextObject* const DataflowContextObject = SceneState.ToolManager->GetContextObjectStore()->FindContext<UDataflowContextObject>())
 	{
-		NewTool->SetClothEditorContextObject(ContextObject);
+		NewTool->SetDataflowEditorContextObject(DataflowContextObject);
 	}
 
 	return NewTool;
@@ -139,4 +143,41 @@ namespace UE::Chaos::ClothAsset
 		ToolCDOs.Add(GetMutableDefault<UClothTransferSkinWeightsTool>());
 		ToolCDOs.Add(GetMutableDefault<UClothMeshSelectionTool>());
 	}
+
+	EClothPatternVertexType DataflowViewModeToClothViewMode(const Dataflow::IDataflowConstructionViewMode* DataflowViewMode)
+	{
+		const FName ViewModeName = DataflowViewMode->GetName();
+		if (ViewModeName == FName("Cloth2DSimView"))
+		{
+			return EClothPatternVertexType::Sim2D;
+		}
+		else if (ViewModeName == FName("Cloth3DSimView"))
+		{
+			return EClothPatternVertexType::Sim3D;
+		}
+		else
+		{
+			check(ViewModeName == FName("ClothRenderView"));
+			return EClothPatternVertexType::Render;
+		}
+	}
+
+	FName ClothViewModeToDataflowViewModeName(EClothPatternVertexType ClothViewMode)
+	{
+		switch (ClothViewMode)
+		{
+		case EClothPatternVertexType::Sim2D:
+			return FName("Cloth2DSimView");
+			break;
+		case EClothPatternVertexType::Sim3D:
+			return FName("Cloth3DSimView");
+			break;
+		case EClothPatternVertexType::Render:
+			return FName("ClothRenderView");
+			break;
+		default:
+			checkNoEntry();
+			return NAME_None;
+		};
+	}
 }
diff --git a/Engine/Plugins/ChaosClothAssetEditor/Source/ChaosClothAssetEditorTools/Private/ChaosClothAsset/ClothMeshSelectionTool.cpp b/Engine/Plugins/ChaosClothAssetEditor/Source/ChaosClothAssetEditorTools/Private/ChaosClothAsset/ClothMeshSelectionTool.cpp
--- a/Engine/Plugins/ChaosClothAssetEditor/Source/ChaosClothAssetEditorTools/Private/ChaosClothAsset/ClothMeshSelectionTool.cpp
+++ b/Engine/Plugins/ChaosClothAssetEditor/Source/ChaosClothAssetEditorTools/Private/ChaosClothAsset/ClothMeshSelectionTool.cpp
@@ -1,8 +1,8 @@
 // Copyright Epic Games, Inc. All Rights Reserved.
 
 #include "ChaosClothAsset/ClothMeshSelectionTool.h"
+#include "ChaosClothAsset/ClothEditorToolBuilders.h"
 #include "ChaosClothAsset/ClothCollectionGroup.h"
-#include "ChaosClothAsset/ClothEditorContextObject.h"
 #include "ChaosClothAsset/ClothGeometryTools.h"
 #include "ChaosClothAsset/ClothPatternVertexType.h"
 #include "ChaosClothAsset/SelectionNode.h"
@@ -17,8 +17,10 @@
 #include "ToolSetupUtil.h"
 #include "Selections/GeometrySelection.h"
 #include "ContextObjectStore.h"
+#include "Dataflow/DataflowContextObject.h"
 #include "Dataflow/DataflowGraphEditor.h"
 #include "Dataflow/DataflowEdNode.h"
+#include "Dataflow/DataflowObject.h"
 #include "DynamicMesh/NonManifoldMappingSupport.h"
 #include "Selections/GeometrySelectionUtil.h"
 #include "Materials/Material.h"
@@ -254,10 +256,10 @@ void UClothMeshSelectionTool::Setup()
 	UE::ToolTarget::HideSourceObject(Target);
 
 	// Setup non-manifold mapping if necessary
-	if (ClothEditorContextObject)
+	if (DataflowContextObject)
 	{
-		ensure(ClothEditorContextObject->IsUsingInputCollection());
-		if (const TSharedPtr<const FManagedArrayCollection> ClothCollection = ClothEditorContextObject->GetSelectedClothCollection().Pin())
+		ensure(DataflowContextObject->IsUsingInputCollection());
+		if (const TSharedPtr<const FManagedArrayCollection> ClothCollection = DataflowContextObject->GetSelectedCollection())
 		{
 			PreviewMesh->ProcessMesh([this, ClothCollection](const FDynamicMesh3& Mesh)
 			{
@@ -421,17 +423,16 @@ FBox UClothMeshSelectionTool::GetWorldSpaceFocusBox()
 	return FBox(SelectionMechanic->GetSelectionBounds(bWorld));
 }
 
-
-void UClothMeshSelectionTool::SetClothEditorContextObject(TObjectPtr<UClothEditorContextObject> InClothEditorContextObject)
+void UClothMeshSelectionTool::SetDataflowContextObject(TObjectPtr<UDataflowContextObject> InDataflowContextObject)
 {
-	ClothEditorContextObject = InClothEditorContextObject;
+	DataflowContextObject = InDataflowContextObject;
 }
 
 bool UClothMeshSelectionTool::GetSelectedNodeInfo(FString& OutSelectionName, UE::Geometry::FGroupTopologySelection& OutSelection, EChaosClothAssetSelectionOverrideType& OutOverrideType)
 {
 	using namespace UE::Chaos::ClothAsset;
 
-	SelectionNodeToUpdate = ClothEditorContextObject->GetSingleSelectedNodeOfType<FChaosClothAssetSelectionNode>();
+	SelectionNodeToUpdate = DataflowContextObject->GetSelectedNodeOfType<FChaosClothAssetSelectionNode>();
 	check(SelectionNodeToUpdate);
 
 	// We need to sanitize the incoming indices, as the user can manually set them to anything on the node
@@ -496,12 +497,12 @@ bool UClothMeshSelectionTool::GetSelectedNodeInfo(FString& OutSelectionName, UE:
 
 	// Get the input set.
 	InputSelectionSet.Reset();
-	if (ClothEditorContextObject)
+	if (DataflowContextObject)
 	{
-		ensure(ClothEditorContextObject->IsUsingInputCollection());
-		if (TSharedPtr<const FManagedArrayCollection> ClothCollection = ClothEditorContextObject->GetSelectedClothCollection().Pin())
+		ensure(DataflowContextObject->IsUsingInputCollection());
+		if (const TSharedPtr<const FManagedArrayCollection> ClothCollection = DataflowContextObject->GetSelectedCollection())
 		{
-			if (TSharedPtr<Dataflow::FEngineContext> DataflowContext = ClothEditorContextObject->GetDataflowContext().Pin())
+			if (const TSharedPtr<Dataflow::FEngineContext> DataflowContext = DataflowContextObject->GetDataflowContext())
 			{
 				using namespace UE::Chaos::ClothAsset;
 				const FName InputName = SelectionNodeToUpdate->GetInputName(*DataflowContext);
@@ -539,7 +540,7 @@ void UClothMeshSelectionTool::UpdateSelectedNode()
 	checkf(SelectionNodeToUpdate, TEXT("Expected non-null pointer to Selection Node"));
 
 	// Save previous state for undo
-	if (UDataflow* const Dataflow = ClothEditorContextObject->GetDataflowAsset())
+	if (UDataflow* const Dataflow = DataflowContextObject->GetDataflowAsset())
 	{
 		GetToolManager()->GetContextTransactionsAPI()->AppendChange(Dataflow, 
 			FChaosClothAssetSelectionNode::MakeWeightMapNodeChange(*SelectionNodeToUpdate),
@@ -547,7 +548,8 @@ void UClothMeshSelectionTool::UpdateSelectedNode()
 	}
 
 	const UE::Geometry::FGroupTopologySelection& Selection = SelectionMechanic->GetActiveSelection();
-	const EClothPatternVertexType ViewMode = ClothEditorContextObject->GetConstructionViewMode();
+
+	const EClothPatternVertexType ViewMode = DataflowViewModeToClothViewMode(DataflowContextObject->GetConstructionViewMode());
 
 	TSet<int32> Indices;
 	FName GroupName = ClothCollectionGroup::SimVertices2D;
@@ -561,13 +563,13 @@ void UClothMeshSelectionTool::UpdateSelectedNode()
 
 		switch(ViewMode)
 		{
-		case UE::Chaos::ClothAsset::EClothPatternVertexType::Sim2D:
+		case EClothPatternVertexType::Sim2D:
 			GroupName = ClothCollectionGroup::SimVertices2D;
 			break;
-		case UE::Chaos::ClothAsset::EClothPatternVertexType::Sim3D:
+		case EClothPatternVertexType::Sim3D:
 			GroupName = ClothCollectionGroup::SimVertices3D;
 			break;
-		case UE::Chaos::ClothAsset::EClothPatternVertexType::Render:
+		case EClothPatternVertexType::Render:
 			GroupName = ClothCollectionGroup::RenderVertices;
 			break;
 		}
@@ -578,11 +580,11 @@ void UClothMeshSelectionTool::UpdateSelectedNode()
 
 		switch (ViewMode)
 		{
-		case UE::Chaos::ClothAsset::EClothPatternVertexType::Sim2D:
-		case UE::Chaos::ClothAsset::EClothPatternVertexType::Sim3D:
+		case EClothPatternVertexType::Sim2D:
+		case EClothPatternVertexType::Sim3D:
 			GroupName = ClothCollectionGroup::SimFaces;
 			break;
-		case UE::Chaos::ClothAsset::EClothPatternVertexType::Render:
+		case EClothPatternVertexType::Render:
 			GroupName = ClothCollectionGroup::RenderFaces;
 			break;
 		}
@@ -663,15 +665,16 @@ void UClothMeshSelectionTool::ApplyAction(EClothMeshSelectionToolActions ActionT
 
 void UClothMeshSelectionTool::ImportFromCollection(bool bImportFromSecondarySet)
 {
-	if (const TSharedPtr<const FManagedArrayCollection> ClothCollection = ClothEditorContextObject->GetSelectedClothCollection().Pin())
+	if (const TSharedPtr<const FManagedArrayCollection> ClothCollection = DataflowContextObject->GetSelectedCollection())
 	{
-		ensure(ClothEditorContextObject->IsUsingInputCollection());
+		ensure(DataflowContextObject->IsUsingInputCollection());
 
 		using namespace UE::Chaos::ClothAsset;
 		const FCollectionClothSelectionConstFacade SelectionFacade(ClothCollection.ToSharedRef());
 		if (SelectionFacade.IsValid())
 		{
-			const EClothPatternVertexType ViewMode = ClothEditorContextObject->GetConstructionViewMode();
+			const EClothPatternVertexType ViewMode = DataflowViewModeToClothViewMode(DataflowContextObject->GetConstructionViewMode());
+
 			FName GroupName = ClothCollectionGroup::SimVertices2D;
 			if (SelectionMechanic->Properties->bSelectVertices)
 			{
@@ -680,13 +683,13 @@ void UClothMeshSelectionTool::ImportFromCollection(bool bImportFromSecondarySet)
 
 				switch (ViewMode)
 				{
-				case UE::Chaos::ClothAsset::EClothPatternVertexType::Sim2D:
+				case EClothPatternVertexType::Sim2D:
 					GroupName = ClothCollectionGroup::SimVertices2D;
 					break;
-				case UE::Chaos::ClothAsset::EClothPatternVertexType::Sim3D:
+				case EClothPatternVertexType::Sim3D:
 					GroupName = ClothCollectionGroup::SimVertices3D;
 					break;
-				case UE::Chaos::ClothAsset::EClothPatternVertexType::Render:
+				case EClothPatternVertexType::Render:
 					GroupName = ClothCollectionGroup::RenderVertices;
 					break;
 				}
@@ -695,11 +698,11 @@ void UClothMeshSelectionTool::ImportFromCollection(bool bImportFromSecondarySet)
 			{
 				switch (ViewMode)
 				{
-				case UE::Chaos::ClothAsset::EClothPatternVertexType::Sim2D:
-				case UE::Chaos::ClothAsset::EClothPatternVertexType::Sim3D:
+				case EClothPatternVertexType::Sim2D:
+				case EClothPatternVertexType::Sim3D:
 					GroupName = ClothCollectionGroup::SimFaces;
 					break;
-				case UE::Chaos::ClothAsset::EClothPatternVertexType::Render:
+				case EClothPatternVertexType::Render:
 					GroupName = ClothCollectionGroup::RenderFaces;
 					break;
 				}
diff --git a/Engine/Plugins/ChaosClothAssetEditor/Source/ChaosClothAssetEditorTools/Private/ChaosClothAsset/ClothMeshSelectionTool.h b/Engine/Plugins/ChaosClothAssetEditor/Source/ChaosClothAssetEditorTools/Private/ChaosClothAsset/ClothMeshSelectionTool.h
--- a/Engine/Plugins/ChaosClothAssetEditor/Source/ChaosClothAssetEditorTools/Private/ChaosClothAsset/ClothMeshSelectionTool.h
+++ b/Engine/Plugins/ChaosClothAssetEditor/Source/ChaosClothAssetEditorTools/Private/ChaosClothAsset/ClothMeshSelectionTool.h
@@ -7,7 +7,7 @@
 #include "ClothMeshSelectionTool.generated.h"
 
 class UPolygonSelectionMechanic;
-class UClothEditorContextObject;
+class UDataflowContextObject;
 class UPreviewMesh;
 struct FChaosClothAssetSelectionNode;
 enum class EChaosClothAssetSelectionOverrideType : uint8;
@@ -150,7 +150,8 @@ private:
 	// IInteractiveToolCameraFocusAPI implementation
 	virtual FBox GetWorldSpaceFocusBox() override;
 
-	void SetClothEditorContextObject(TObjectPtr<UClothEditorContextObject> InClothEditorContextObject);
+	void SetDataflowContextObject(TObjectPtr<UDataflowContextObject> InDataflowContextObject);
+
 	bool GetSelectedNodeInfo(FString& OutMapName, UE::Geometry::FGroupTopologySelection& OutSelection, EChaosClothAssetSelectionOverrideType& OutOverrideType);
 	void UpdateSelectedNode();
 
@@ -164,7 +165,7 @@ private:
 	TObjectPtr<UPolygonSelectionMechanic> SelectionMechanic = nullptr;
 
 	UPROPERTY()
-	TObjectPtr<UClothEditorContextObject> ClothEditorContextObject = nullptr;
+	TObjectPtr<UDataflowContextObject> DataflowContextObject = nullptr;
 
 	TUniquePtr<UE::Geometry::FGroupTopology> Topology;
 
diff --git a/Engine/Plugins/ChaosClothAssetEditor/Source/ChaosClothAssetEditorTools/Private/ChaosClothAsset/ClothTransferSkinWeightsTool.cpp b/Engine/Plugins/ChaosClothAssetEditor/Source/ChaosClothAssetEditorTools/Private/ChaosClothAsset/ClothTransferSkinWeightsTool.cpp
--- a/Engine/Plugins/ChaosClothAssetEditor/Source/ChaosClothAssetEditorTools/Private/ChaosClothAsset/ClothTransferSkinWeightsTool.cpp
+++ b/Engine/Plugins/ChaosClothAssetEditor/Source/ChaosClothAssetEditorTools/Private/ChaosClothAsset/ClothTransferSkinWeightsTool.cpp
@@ -7,8 +7,8 @@
 #include "ToolSetupUtil.h"
 #include "ModelingToolTargetUtil.h"
 #include "Engine/SkeletalMesh.h"
-#include "ChaosClothAsset/ClothEditorContextObject.h"
 #include "ContextObjectStore.h"
+#include "Dataflow/DataflowContextObject.h"
 #include "Dataflow/DataflowEdNode.h"
 #include "ChaosClothAsset/TransferSkinWeightsNode.h"
 #include "BaseGizmos/TransformGizmoUtil.h"
@@ -50,7 +50,7 @@ void UClothTransferSkinWeightsTool::Setup()
 {
 	USingleSelectionMeshEditingTool::Setup();
 
-	TransferSkinWeightsNode = ClothEditorContextObject->GetSingleSelectedNodeOfType<FChaosClothAssetTransferSkinWeightsNode>();
+	TransferSkinWeightsNode = DataflowContextObject->GetSelectedNodeOfType<FChaosClothAssetTransferSkinWeightsNode>();
 	checkf(TransferSkinWeightsNode, TEXT("No Transfer Skin Weights Node is currently selected, or more than one node is selected"));
 
 	ToolProperties = NewObject<UClothTransferSkinWeightsToolProperties>(this);
@@ -250,9 +250,9 @@ TUniquePtr<UE::Geometry::FDynamicMeshOperator> UClothTransferSkinWeightsTool::Ma
 	return TransferOp;
 }
 
-void UClothTransferSkinWeightsTool::SetClothEditorContextObject(TObjectPtr<UClothEditorContextObject> InClothEditorContextObject)
+void UClothTransferSkinWeightsTool::SetDataflowEditorContextObject(TObjectPtr<UDataflowContextObject> InDataflowContextObject)
 {
-	ClothEditorContextObject = InClothEditorContextObject;
+	DataflowContextObject = InDataflowContextObject;
 }
 
 FTransform UClothTransferSkinWeightsTool::TransformFromProperties() const
diff --git a/Engine/Plugins/ChaosClothAssetEditor/Source/ChaosClothAssetEditorTools/Private/ChaosClothAsset/ClothTransferSkinWeightsTool.h b/Engine/Plugins/ChaosClothAssetEditor/Source/ChaosClothAssetEditorTools/Private/ChaosClothAsset/ClothTransferSkinWeightsTool.h
--- a/Engine/Plugins/ChaosClothAssetEditor/Source/ChaosClothAssetEditorTools/Private/ChaosClothAsset/ClothTransferSkinWeightsTool.h
+++ b/Engine/Plugins/ChaosClothAssetEditor/Source/ChaosClothAssetEditorTools/Private/ChaosClothAsset/ClothTransferSkinWeightsTool.h
@@ -10,7 +10,7 @@
 
 class UClothTransferSkinWeightsTool;
 class USkeletalMesh;
-class UClothEditorContextObject;
+class UDataflowContextObject;
 class UTransformProxy;
 class UCombinedTransformGizmo;
 class UMeshOpPreviewWithBackgroundCompute;
@@ -61,7 +61,7 @@ private:
 	// IDynamicMeshOperatorFactory
 	virtual TUniquePtr<UE::Geometry::FDynamicMeshOperator> MakeNewOperator() override;
 	
-	void SetClothEditorContextObject(TObjectPtr<UClothEditorContextObject> InClothEditorContextObject);
+	void SetDataflowEditorContextObject(TObjectPtr<UDataflowContextObject> InDataflowContextObject);
 
 	FTransform TransformFromProperties() const;
 	void SetSRTPropertiesFromTransform(const FTransform& Transform) const;
@@ -77,7 +77,7 @@ private:
 	TObjectPtr<UClothTransferSkinWeightsToolProperties> ToolProperties;
 
 	UPROPERTY(Transient)
-	TObjectPtr<UClothEditorContextObject> ClothEditorContextObject;
+	TObjectPtr<UDataflowContextObject> DataflowContextObject = nullptr;
 
 	UPROPERTY(Transient)
 	TObjectPtr<UMeshOpPreviewWithBackgroundCompute> TargetClothPreview;
diff --git a/Engine/Plugins/ChaosClothAssetEditor/Source/ChaosClothAssetEditorTools/Private/ChaosClothAsset/ClothWeightMapPaintTool.cpp b/Engine/Plugins/ChaosClothAssetEditor/Source/ChaosClothAssetEditorTools/Private/ChaosClothAsset/ClothWeightMapPaintTool.cpp
--- a/Engine/Plugins/ChaosClothAssetEditor/Source/ChaosClothAssetEditorTools/Private/ChaosClothAsset/ClothWeightMapPaintTool.cpp
+++ b/Engine/Plugins/ChaosClothAssetEditor/Source/ChaosClothAssetEditorTools/Private/ChaosClothAsset/ClothWeightMapPaintTool.cpp
@@ -33,16 +33,19 @@
 #include "CanvasItem.h"
 #include "ChaosClothAsset/ClothCollectionGroup.h"
 #include "ChaosClothAsset/ClothEditorContextObject.h"
+#include "ChaosClothAsset/ClothEditorToolBuilders.h"
 #include "ChaosClothAsset/ClothPatternVertexType.h"
 #include "ChaosClothAsset/CollectionClothFacade.h"
 #include "ChaosClothAsset/WeightedValue.h"
 #include "ContextObjectStore.h"
+#include "Dataflow/DataflowContextObject.h"
 #include "Dataflow/DataflowEdNode.h"
 #include "Dataflow/DataflowObject.h"
 #include "ChaosClothAsset/WeightMapNode.h"
 #include "GraphEditor.h"
 #include "Dataflow/DataflowSNode.h"
 #include "Dataflow/DataflowGraphEditor.h"
+#include "Dataflow/DataflowRenderingViewMode.h"
 #include "Widgets/Notifications/SNotificationList.h"
 #include "Framework/Notifications/NotificationManager.h"
 #include "Selection/PolygonSelectionMechanic.h"
@@ -109,7 +112,7 @@ void UClothEditorWeightMapPaintTool::Setup()
 	UMeshSculptToolBase::Setup();
 
 	// Get the selected weight map node
-	WeightMapNodeToUpdate = ClothEditorContextObject->GetSingleSelectedNodeOfType<FChaosClothAssetWeightMapNode>();
+	WeightMapNodeToUpdate = DataflowContextObject->GetSelectedNodeOfType<FChaosClothAssetWeightMapNode>();
 	checkf(WeightMapNodeToUpdate, TEXT("No Weight Map Node is currently selected, or more than one node is selected"));
 
 	SetToolDisplayName(LOCTEXT("ToolName", "Paint Weight Maps"));
@@ -364,19 +367,19 @@ void UClothEditorWeightMapPaintTool::Setup()
 	// configure panels
 	UpdateSubToolType(FilterProperties->SubToolType);
 
-	const bool bIsRenderMode = ClothEditorContextObject && (ClothEditorContextObject->GetConstructionViewMode() == UE::Chaos::ClothAsset::EClothPatternVertexType::Render);
 	// Setup DynamicMeshToWeight conversion and get Input weight map (if it exists)
 	InputWeightMap = TConstArrayView<float>();
-	if (ClothEditorContextObject)
+
+	if (DataflowContextObject)
 	{
-		ensure(ClothEditorContextObject->IsUsingInputCollection());
-		if (TSharedPtr<const FManagedArrayCollection> ClothCollection = ClothEditorContextObject->GetSelectedClothCollection().Pin())
+		ensure(DataflowContextObject->IsUsingInputCollection());
+		if (TSharedPtr<const FManagedArrayCollection> ClothCollection = DataflowContextObject->GetSelectedCollection())
 		{
 			using namespace UE::Chaos::ClothAsset;
 			const FNonManifoldMappingSupport NonManifoldMapping(*Mesh);
 
 			const bool bHasNonManifoldMapping = NonManifoldMapping.IsNonManifoldVertexInSource();
-			const bool bHas2D3DConversion = ClothEditorContextObject->GetConstructionViewMode() == EClothPatternVertexType::Sim2D;
+			const bool bHas2D3DConversion = DataflowViewModeToClothViewMode(DataflowContextObject->GetConstructionViewMode()) == EClothPatternVertexType::Sim2D;
 
 			bHaveDynamicMeshToWeightConversion = bHasNonManifoldMapping || bHas2D3DConversion;
 
@@ -405,8 +408,11 @@ void UClothEditorWeightMapPaintTool::Setup()
 				WeightToDynamicMesh = Cloth.GetSimVertex2DLookup();
 			}
 
-			const int32 NumPatterns = bIsRenderMode ? Cloth.GetNumRenderPatterns() : Cloth.GetNumSimPatterns();		
-			
+			const EClothPatternVertexType ViewMode = UE::Chaos::ClothAsset::DataflowViewModeToClothViewMode(DataflowContextObject->GetConstructionViewMode());
+			const bool bIsRenderMode = ViewMode == EClothPatternVertexType::Render;
+
+			const int32 NumPatterns = bIsRenderMode ? Cloth.GetNumRenderPatterns() : Cloth.GetNumSimPatterns();
+
 			PatternTriangleOffsetAndNum.SetNum(NumPatterns);
 
 			TSet<int32> NonEmptyPatternIDs;
@@ -441,7 +447,7 @@ void UClothEditorWeightMapPaintTool::Setup()
 			}
 
 			// Find the map if it exists.
-			if (TSharedPtr<Dataflow::FEngineContext> DataflowContext = ClothEditorContextObject->GetDataflowContext().Pin())
+			if (TSharedPtr<Dataflow::FEngineContext> DataflowContext = DataflowContextObject->GetDataflowContext())
 			{
 				const FName InputName = WeightMapNodeToUpdate->GetInputName(*DataflowContext);
 				if (bIsRenderMode)
@@ -544,10 +550,9 @@ void UClothEditorWeightMapPaintTool::DecreaseBrushSpeedAction()		// Actually dec
 	NotifyOfPropertyChangeByTool(FilterProperties);
 }
 
-
-void UClothEditorWeightMapPaintTool::SetClothEditorContextObject(TObjectPtr<UClothEditorContextObject> InClothEditorContextObject)
+void UClothEditorWeightMapPaintTool::SetDataflowContextObject(TObjectPtr<UDataflowContextObject> InDataflowContextObject)
 {
-	ClothEditorContextObject = InClothEditorContextObject;
+	DataflowContextObject = InDataflowContextObject;
 }
 
 void UClothEditorWeightMapPaintTool::Shutdown(EToolShutdownType ShutdownType)
@@ -1976,7 +1981,7 @@ void UClothEditorWeightMapPaintTool::UpdateSelectedNode()
 	checkf(WeightMapNodeToUpdate, TEXT("Expected non-null pointer to Add Weight Map Node"));
 
 	// Save previous state for undo
-	if (UDataflow* const Dataflow = ClothEditorContextObject->GetDataflowAsset())
+	if (UDataflow* const Dataflow = DataflowContextObject->GetDataflowAsset())
 	{
 		GetToolManager()->GetContextTransactionsAPI()->AppendChange(Dataflow, 
 			FChaosClothAssetWeightMapNode::MakeWeightMapNodeChange(*WeightMapNodeToUpdate),
@@ -1987,8 +1992,6 @@ void UClothEditorWeightMapPaintTool::UpdateSelectedNode()
 	WeightMapNodeToUpdate->MapOverrideType = UpdateWeightMapProperties->MapOverrideType;
 	WeightMapNodeToUpdate->OutputName.StringValue = UpdateWeightMapProperties->Name;
 
-	const bool bIsRenderMode = (ClothEditorContextObject->GetConstructionViewMode() == UE::Chaos::ClothAsset::EClothPatternVertexType::Render);
-
 	if (bHaveDynamicMeshToWeightConversion)
 	{
 		TArray<float> NodeWeights;
diff --git a/Engine/Plugins/ChaosClothAssetEditor/Source/ChaosClothAssetEditorTools/Private/ChaosClothAsset/ClothWeightMapPaintTool.h b/Engine/Plugins/ChaosClothAssetEditor/Source/ChaosClothAssetEditorTools/Private/ChaosClothAsset/ClothWeightMapPaintTool.h
--- a/Engine/Plugins/ChaosClothAssetEditor/Source/ChaosClothAssetEditorTools/Private/ChaosClothAsset/ClothWeightMapPaintTool.h
+++ b/Engine/Plugins/ChaosClothAssetEditor/Source/ChaosClothAssetEditorTools/Private/ChaosClothAsset/ClothWeightMapPaintTool.h
@@ -22,13 +22,14 @@
 #include "TransformTypes.h"
 #include "ToolDataVisualizer.h"
 #include "GroupTopology.h"
+#include "Dataflow/DataflowObjectInterface.h"
 #include "ClothWeightMapPaintTool.generated.h"
 
 class UMeshElementsVisualizer;
 class UWeightMapEraseBrushOpProps;
 class UWeightMapPaintBrushOpProps;
 class UWeightMapSmoothBrushOpProps;
-class UClothEditorContextObject;
+class UDataflowContextObject;
 class UPolygonSelectionMechanic;
 struct FChaosClothAssetWeightMapNode;
 enum class EChaosClothAssetWeightMapOverrideType : uint8;
@@ -322,7 +323,7 @@ public:
 
 	virtual void CommitResult(UBaseDynamicMeshComponent* Component, bool bModifiedTopology) override;
 
-	void SetClothEditorContextObject(TObjectPtr<UClothEditorContextObject> InClothEditorContextObject);
+	void SetDataflowContextObject(TObjectPtr<UDataflowContextObject> InDataflowContextObject);
 
 public:
 
@@ -478,7 +479,7 @@ protected:
 	TObjectPtr<UMeshElementsVisualizer> MeshElementsDisplay;
 
 	UPROPERTY()
-	TObjectPtr<UClothEditorContextObject> ClothEditorContextObject = nullptr;
+	TObjectPtr<UDataflowContextObject> DataflowContextObject = nullptr;
 
 	// realtime visualization
 	void OnDynamicMeshComponentChanged(UDynamicMeshComponent* Component, const FMeshVertexChange* Change, bool bRevert);
diff --git a/Engine/Plugins/ChaosClothAssetEditor/Source/ChaosClothAssetEditorTools/Public/ChaosClothAsset/ClothEditorToolBuilders.h b/Engine/Plugins/ChaosClothAssetEditor/Source/ChaosClothAssetEditorTools/Public/ChaosClothAsset/ClothEditorToolBuilders.h
--- a/Engine/Plugins/ChaosClothAssetEditor/Source/ChaosClothAssetEditorTools/Public/ChaosClothAsset/ClothEditorToolBuilders.h
+++ b/Engine/Plugins/ChaosClothAssetEditor/Source/ChaosClothAssetEditorTools/Public/ChaosClothAsset/ClothEditorToolBuilders.h
@@ -8,7 +8,15 @@
 #include "BaseTools/MeshSurfacePointMeshEditingTool.h"
 #include "ClothEditorToolBuilders.generated.h"
 
+class UDataflowContextObject;
+#if UE_ENABLE_INCLUDE_ORDER_DEPRECATED_IN_5_5
 class UClothEditorContextObject;
+#endif
+
+namespace Dataflow
+{
+	class IDataflowConstructionViewMode;
+}
 
 namespace UE::Chaos::ClothAsset
 {
@@ -29,8 +37,11 @@ class IChaosClothAssetEditorToolBuilder
 
 public:
 
+	UE_DEPRECATED(5.5, "Please use the version taking ContextObject")
+	virtual void GetSupportedViewModes(TArray<UE::Chaos::ClothAsset::EClothPatternVertexType>& Modes) const {};
+
 	/** Returns all Construction View modes that this tool can operate in. The first element should be the preferred mode to switch to if necessary. */
-	virtual void GetSupportedViewModes(const UClothEditorContextObject& ContextObject, TArray<UE::Chaos::ClothAsset::EClothPatternVertexType>& Modes) const = 0;
+	virtual void GetSupportedViewModes(const UDataflowContextObject& ContextObject, TArray<UE::Chaos::ClothAsset::EClothPatternVertexType>& Modes) const = 0;
 
 	/** Returns whether or not view can be set to wireframe when this tool is active.. */
 	virtual bool CanSetConstructionViewWireframeActive() const { return true; }
@@ -43,7 +54,7 @@ class CHAOSCLOTHASSETEDITORTOOLS_API UClothEditorWeightMapPaintToolBuilder : pub
 	GENERATED_BODY()
 
 private:
-	virtual void GetSupportedViewModes(const UClothEditorContextObject& ContextObject, TArray<UE::Chaos::ClothAsset::EClothPatternVertexType>& Modes) const override;
+	virtual void GetSupportedViewModes(const UDataflowContextObject& ContextObject, TArray<UE::Chaos::ClothAsset::EClothPatternVertexType>& Modes) const override;
 	virtual UMeshSurfacePointTool* CreateNewTool(const FToolBuilderState& SceneState) const override;
 	virtual bool CanSetConstructionViewWireframeActive() const { return false; }
 };
@@ -54,8 +65,10 @@ class CHAOSCLOTHASSETEDITORTOOLS_API UClothMeshSelectionToolBuilder : public UIn
 {
 	GENERATED_BODY()
 
+	virtual void GetSupportedViewModes(TArray<UE::Chaos::ClothAsset::EClothPatternVertexType>& Modes) const override {};
+
 private:
-	virtual void GetSupportedViewModes(const UClothEditorContextObject& ContextObject, TArray<UE::Chaos::ClothAsset::EClothPatternVertexType>& Modes) const override;
+	virtual void GetSupportedViewModes(const UDataflowContextObject& ContextObject, TArray<UE::Chaos::ClothAsset::EClothPatternVertexType>& Modes) const override;
 	virtual bool CanBuildTool(const FToolBuilderState& SceneState) const override;
 	virtual UInteractiveTool* BuildTool(const FToolBuilderState& SceneState) const override;
 	virtual const FToolTargetTypeRequirements& GetTargetRequirements() const override;
@@ -68,7 +81,7 @@ class CHAOSCLOTHASSETEDITORTOOLS_API UClothTransferSkinWeightsToolBuilder : publ
 	GENERATED_BODY()
 
 private:
-	virtual void GetSupportedViewModes(const UClothEditorContextObject& ContextObject, TArray<UE::Chaos::ClothAsset::EClothPatternVertexType>& Modes) const override;
+	virtual void GetSupportedViewModes(const UDataflowContextObject& ContextObject, TArray<UE::Chaos::ClothAsset::EClothPatternVertexType>& Modes) const override;
 	virtual USingleSelectionMeshEditingTool* CreateNewTool(const FToolBuilderState& SceneState) const override;
 
 };
@@ -78,5 +91,11 @@ namespace UE::Chaos::ClothAsset
 {
 	// Provide a list of Tool default objects for us in TInteractiveToolCommands::RegisterCommands()
 	void CHAOSCLOTHASSETEDITORTOOLS_API GetClothEditorToolDefaultObjectList(TArray<UInteractiveTool*>& ToolCDOs);
+
+	// Mapping from Dataflow View Mode to Cloth View Mode. Input object must be one of FCloth2DSimViewMode, FCloth3DSimViewMode, or FClothRenderViewMode, defined in ClothDataflowViewModes.h
+	EClothPatternVertexType CHAOSCLOTHASSETEDITORTOOLS_API DataflowViewModeToClothViewMode(const Dataflow::IDataflowConstructionViewMode* DataflowViewMode);
+
+	// Mapping from Cloth View Mode to Dataflow View Mode name. Ouptut will be one of "Cloth2DSimView", "Cloth3DSimView", or "ClothRenderView", as defined in ClothDataflowViewModes.cpp
+	FName CHAOSCLOTHASSETEDITORTOOLS_API ClothViewModeToDataflowViewModeName(EClothPatternVertexType ClothViewMode);
 }
 
diff --git a/Engine/Plugins/Experimental/Dataflow/Source/DataflowEditor/Private/Dataflow/DataflowConstructionScene.cpp b/Engine/Plugins/Experimental/Dataflow/Source/DataflowEditor/Private/Dataflow/DataflowConstructionScene.cpp
--- a/Engine/Plugins/Experimental/Dataflow/Source/DataflowEditor/Private/Dataflow/DataflowConstructionScene.cpp
+++ b/Engine/Plugins/Experimental/Dataflow/Source/DataflowEditor/Private/Dataflow/DataflowConstructionScene.cpp
@@ -148,9 +148,9 @@ void FDataflowConstructionScene::UpdateDynamicMeshComponents()
 
 						if (DynamicMesh.VertexCount())
 						{
-							if (Target == EditorContent->GetPrimarySelectedNode())
+							if (Target == EditorContent->GetSelectedNode())
 							{
-								EditorContent->SetPrimaryRenderCollection(RenderCollection);
+								EditorContent->SetRenderCollection(RenderCollection);
 							}
 							const FString MeshName = Facade.GetGeometryName()[MeshIndex];
 							AddDynamicMeshComponent({Target, MeshIndex}, MeshName, MoveTemp(DynamicMesh), {});
@@ -195,9 +195,9 @@ void FDataflowConstructionScene::UpdateDynamicMeshComponents()
 
 							if (DynamicMesh.VertexCount())
 							{
-								if (Target == EditorContent->GetPrimarySelectedNode())
+								if (Target == EditorContent->GetSelectedNode())
 								{
-									EditorContent->SetPrimaryRenderCollection(RenderCollection);
+									EditorContent->SetRenderCollection(RenderCollection);
 								}
 								const FString MeshName = Facade.GetGeometryName()[MeshIndex];
 								const FString UniqueObjectName = MakeUniqueObjectName(RootSceneActor, UDataflowEditorCollectionComponent::StaticClass(), FName(MeshName)).ToString();
diff --git a/Engine/Plugins/Experimental/Dataflow/Source/DataflowEditor/Private/Dataflow/DataflowEditorCommands.cpp b/Engine/Plugins/Experimental/Dataflow/Source/DataflowEditor/Private/Dataflow/DataflowEditorCommands.cpp
--- a/Engine/Plugins/Experimental/Dataflow/Source/DataflowEditor/Private/Dataflow/DataflowEditorCommands.cpp
+++ b/Engine/Plugins/Experimental/Dataflow/Source/DataflowEditor/Private/Dataflow/DataflowEditorCommands.cpp
@@ -14,6 +14,7 @@
 #include "Dataflow/DataflowRenderingViewMode.h"
 #include "Dataflow/DataflowSCommentNode.h"
 #include "Dataflow/DataflowSNode.h"
+#include "Dataflow/DataflowToolRegistry.h"
 #include "DataflowEditorTools/DataflowEditorWeightMapPaintTool.h"
 #include "EdGraphNode_Comment.h"
 #include "EdGraph/EdGraphNode.h"
@@ -33,7 +34,6 @@
 
 #define LOCTEXT_NAMESPACE "DataflowEditorCommands"
 
-const FString FDataflowEditorCommandsImpl::BeginWeightMapPaintToolIdentifier = TEXT("BeginWeightMapPaintTool");
 const FString FDataflowEditorCommandsImpl::AddWeightMapNodeIdentifier = TEXT("AddWeightMapNode");
 const FString FDataflowEditorCommandsImpl::RebuildSimulationSceneIdentifier = TEXT("RebuildSimulationScene");
 const FString FDataflowEditorCommandsImpl::PauseSimulationSceneIdentifier = TEXT("PauseSimulationScene");
@@ -66,7 +66,6 @@ void FDataflowEditorCommandsImpl::RegisterCommands()
 	UI_COMMAND(RemoveOptionPin, "RemoveOptionPin", "Remove the last option pin from the selected nodes.", EUserInterfaceActionType::Button, FInputChord());
 	UI_COMMAND(ZoomToFitGraph, "ZoomToFitGraph", "Fit the graph in the graph editor viewport.", EUserInterfaceActionType::None, FInputChord(EKeys::F));
 
-	UI_COMMAND(BeginWeightMapPaintTool, "Add Weight Map", "Paint weight maps on the mesh", EUserInterfaceActionType::None, FInputChord());
 	UI_COMMAND(AddWeightMapNode, "Add Weight Map", "Paint weight maps on the mesh", EUserInterfaceActionType::Button, FInputChord());
 	UI_COMMAND(RebuildSimulationScene, "RebuildSimulationScene", "Rebuild the simulation scene", EUserInterfaceActionType::ToggleButton, FInputChord());
 	UI_COMMAND(PauseSimulationScene, "PauseSimulationScene", "Pause the simulation scene", EUserInterfaceActionType::ToggleButton, FInputChord());
@@ -110,7 +109,25 @@ void FDataflowEditorCommandsImpl::RegisterCommands()
 			);
 			CreateNodesMap.Add(Parameters.TypeName, AddNode);
 		}
-	}		
+	}
+
+
+	Dataflow::FDataflowToolRegistry& ToolRegistry = Dataflow::FDataflowToolRegistry::Get();
+	const TArray<FName> NodeNames = ToolRegistry.GetNodeNames();
+	for (const FName& NodeName : NodeNames)
+	{
+		FUICommandInfo::MakeCommandInfo(
+			this->AsShared(),
+			ToolRegistry.GetToolCommandForNode(NodeName),
+			FName(NodeName.ToString() + "_Tool"),
+			LOCTEXT("DataflowTool", "Dataflow Tool"),					// TODO: Replace placeholder names
+			LOCTEXT("DataflowToolTooltip", "Dataflow Tool Tooltip"),	// TODO: Replace placeholder names
+			FSlateIcon(),
+			EUserInterfaceActionType::Button,
+			FInputChord()
+		);
+	}
+
 }
 
 void FDataflowEditorCommandsImpl::GetToolDefaultObjectList(TArray<UInteractiveTool*>& ToolCDOs)
diff --git a/Engine/Plugins/Experimental/Dataflow/Source/DataflowEditor/Private/Dataflow/DataflowEditorMode.cpp b/Engine/Plugins/Experimental/Dataflow/Source/DataflowEditor/Private/Dataflow/DataflowEditorMode.cpp
--- a/Engine/Plugins/Experimental/Dataflow/Source/DataflowEditor/Private/Dataflow/DataflowEditorMode.cpp
+++ b/Engine/Plugins/Experimental/Dataflow/Source/DataflowEditor/Private/Dataflow/DataflowEditorMode.cpp
@@ -24,6 +24,7 @@
 #include "Dataflow/DataflowSimulationViewportClient.h"
 #include "Dataflow/DataflowSNode.h"
 #include "Dataflow/DataflowToolTarget.h"
+#include "Dataflow/DataflowToolRegistry.h"
 #include "EditorModeManager.h"
 #include "EdModeInteractiveToolsContext.h"
 #include "Elements/Framework/EngineElementsLibrary.h"
@@ -108,7 +109,6 @@ void UDataflowEditorMode::AddToolTargetFactories()
 void UDataflowEditorMode::RegisterDataflowTool(TSharedPtr<FUICommandInfo> UICommand,
 	FString ToolIdentifier,
 	UInteractiveToolBuilder* Builder,
-	const IDataflowEditorToolBuilder* DataflowToolBuilder,
 	UEditorInteractiveToolsContext* const ToolsContext,
 	EToolsContextScope ToolScope)
 {
@@ -133,8 +133,37 @@ void UDataflowEditorMode::RegisterDataflowTool(TSharedPtr<FUICommandInfo> UIComm
 	const TSharedRef<FUICommandList>& CommandList = Toolkit->GetToolkitCommands();
 
 	CommandList->MapAction(UICommand, 
-		FExecuteAction::CreateWeakLambda(ToolsContext, [this, ToolsContext, ToolIdentifier, DataflowToolBuilder]()
+		FExecuteAction::CreateWeakLambda(ToolsContext, [this, ToolsContext, ToolIdentifier]()
 		{
+			UDataflowContextObject* ContextObject = ToolsContext->ContextObjectStore->FindContext<UDataflowContextObject>();
+			check(ContextObject);
+
+			// Make sure the ContextObject's selected Collection is the from the Input side of the selected node (so that the tool gets the Collection as it appears before node execution)
+
+			if (TSharedPtr<Dataflow::FEngineContext> DataflowContext = ContextObject->GetDataflowContext())
+			{
+				if (UDataflowEdNode* const SelectedNode = ContextObject->GetSelectedNode())
+				{
+					if (const TSharedPtr<FDataflowNode> DataflowNode = SelectedNode->GetDataflowNode())
+					{
+						for (const FDataflowInput* const Input : DataflowNode->GetInputs())
+						{
+							if (Input->GetType() == FName("FManagedArrayCollection"))
+							{
+								const FManagedArrayCollection DefaultValue;
+								TSharedRef<FManagedArrayCollection> Collection = MakeShared<FManagedArrayCollection>(Input->GetValue<FManagedArrayCollection>(*DataflowContext, DefaultValue));
+
+								constexpr bool bCollectionIsInput = true;
+								ContextObject->SetSelectedCollection(Collection, bCollectionIsInput);
+
+								// If we have multiple input Collections, this will just take the first one. This is what the Cloth Editor does as well (but there it also checks to see if it's a ClothCollection)
+								break;
+							}
+						}
+					}
+				}
+			}
+
 			ActiveToolsContext = ToolsContext;
 			ToolsContext->StartTool(ToolIdentifier);
 		}),
@@ -186,10 +215,24 @@ void UDataflowEditorMode::RegisterTools()
 
 	UEditorInteractiveToolsContext* const ConstructionViewportToolsContext = GetInteractiveToolsContext();
 
-	UDataflowEditorWeightMapPaintToolBuilder* WeightMapPaintToolBuilder = NewObject<UDataflowEditorWeightMapPaintToolBuilder>(this);
-	WeightMapPaintToolBuilder->SetEditorMode(this);
-	RegisterDataflowTool(CommandInfos.BeginWeightMapPaintTool, FDataflowEditorCommandsImpl::BeginWeightMapPaintToolIdentifier, WeightMapPaintToolBuilder, WeightMapPaintToolBuilder, ConstructionViewportToolsContext);
-	RegisterAddNodeCommand(CommandInfos.AddWeightMapNode, FDataflowCollectionAddScalarVertexPropertyNode::StaticType(), CommandInfos.BeginWeightMapPaintTool);
+	Dataflow::FDataflowToolRegistry& ToolRegistry = Dataflow::FDataflowToolRegistry::Get();
+	const TArray<FName> NodeNames = ToolRegistry.GetNodeNames();
+	for (const FName& RegisteredNodeName : NodeNames)
+	{
+		const TSharedPtr<FUICommandInfo> CommandInfo = ToolRegistry.GetToolCommandForNode(RegisteredNodeName);
+		UInteractiveToolBuilder* const Builder = ToolRegistry.GetToolBuilderForNode(RegisteredNodeName);
+
+		// TODO: This is here only so the Tool can hide the all meshes in the DataflowConstructionScene. That should probably be handed in this class instead.
+		if (UDataflowEditorWeightMapPaintToolBuilder* WeightMapPaintToolBuilder = Cast<UDataflowEditorWeightMapPaintToolBuilder>(Builder))
+		{
+			WeightMapPaintToolBuilder->SetEditorMode(this);
+		}
+
+		RegisterDataflowTool(CommandInfo, RegisteredNodeName.ToString() + FString(TEXT("Tool")), Builder, ConstructionViewportToolsContext);
+
+		NodeTypeToToolCommandMap.Add(RegisteredNodeName, CommandInfo);
+	}
+
 }
 
 bool UDataflowEditorMode::ShouldToolStartBeAllowed(const FString& ToolIdentifier) const
diff --git a/Engine/Plugins/Experimental/Dataflow/Source/DataflowEditor/Private/Dataflow/DataflowEditorModule.cpp b/Engine/Plugins/Experimental/Dataflow/Source/DataflowEditor/Private/Dataflow/DataflowEditorModule.cpp
--- a/Engine/Plugins/Experimental/Dataflow/Source/DataflowEditor/Private/Dataflow/DataflowEditorModule.cpp
+++ b/Engine/Plugins/Experimental/Dataflow/Source/DataflowEditor/Private/Dataflow/DataflowEditorModule.cpp
@@ -11,6 +11,8 @@
 #include "Dataflow/DataflowFunctionPropertyCustomization.h"
 #include "Dataflow/DataflowSNodeFactories.h"
 #include "Dataflow/ScalarVertexPropertyGroupCustomization.h"
+#include "Dataflow/DataflowToolRegistry.h"
+#include "DataflowEditorTools/DataflowEditorWeightMapPaintTool.h"
 #include "Dataflow/DataflowCollectionAddScalarVertexPropertyNode.h"
 
 #include "PropertyEditorModule.h"
@@ -34,6 +36,11 @@ void FDataflowEditorModule::StartupModule()
 	}
 
 	Dataflow::RenderingCallbacks();
+
+	Dataflow::FDataflowToolRegistry& ToolRegistry = Dataflow::FDataflowToolRegistry::Get();
+
+	UDataflowEditorWeightMapPaintToolBuilder* const ToolBuilder = NewObject<UDataflowEditorWeightMapPaintToolBuilder>();
+	ToolRegistry.AddNodeToToolMapping(FDataflowCollectionAddScalarVertexPropertyNode::StaticType(), ToolBuilder);
 }
 
 void FDataflowEditorModule::ShutdownModule()
@@ -46,6 +53,9 @@ void FDataflowEditorModule::ShutdownModule()
 		PropertyModule->UnregisterCustomPropertyTypeLayout(ScalarVertexPropertyGroupName);
 		PropertyModule->UnregisterCustomPropertyTypeLayout(DataflowFunctionPropertyName);
 	}
+
+	Dataflow::FDataflowToolRegistry& ToolRegistry = Dataflow::FDataflowToolRegistry::Get();
+	ToolRegistry.RemoveNodeToToolMapping(FDataflowCollectionAddScalarVertexPropertyNode::StaticType());
 }
 
 IMPLEMENT_MODULE(FDataflowEditorModule, DataflowEditor)
diff --git a/Engine/Plugins/Experimental/Dataflow/Source/DataflowEditor/Private/Dataflow/DataflowEditorToolkit.cpp b/Engine/Plugins/Experimental/Dataflow/Source/DataflowEditor/Private/Dataflow/DataflowEditorToolkit.cpp
--- a/Engine/Plugins/Experimental/Dataflow/Source/DataflowEditor/Private/Dataflow/DataflowEditorToolkit.cpp
+++ b/Engine/Plugins/Experimental/Dataflow/Source/DataflowEditor/Private/Dataflow/DataflowEditorToolkit.cpp
@@ -608,6 +608,11 @@ void FDataflowEditorToolkit::CreateWidgets()
 				AssetDetailsEditor = CreateAssetDetailsEditorWidget({EditorContent->GetDataflowAsset()});
 			}
 			GraphEditor = CreateGraphEditorWidget(DataflowAsset, NodeDetailsEditor);
+
+			// Synchronize the EditorContent's selected node with the GraphEditor
+			UDataflowEdNode* const InitialSelectedNode = Cast<UDataflowEdNode>(GraphEditor->GetSingleSelectedNode());
+			EditorContent->SetSelectedNode(InitialSelectedNode);
+
 			CreateSimulationViewportClient();
 
 			FAdvancedPreviewSceneModule& AdvancedPreviewSceneModule = FModuleManager::LoadModuleChecked<FAdvancedPreviewSceneModule>("AdvancedPreviewScene");
@@ -915,7 +920,7 @@ void FDataflowEditorToolkit::OnNodeSelectionChanged(const TSet<UObject*>& InNewS
 				}
 			}
 
-			EditorContent->SetPrimarySelectedNode(nullptr);
+			EditorContent->SetSelectedNode(nullptr);
 
 			if( UDataflowEditorMode* const DataflowMode = Cast<UDataflowEditorMode>(EditorModeManager->GetActiveScriptableMode(UDataflowEditorMode::EM_DataflowEditorModeId)) )
 			{
@@ -927,13 +932,26 @@ void FDataflowEditorToolkit::OnNodeSelectionChanged(const TSet<UObject*>& InNewS
 					ToolsContext->EndTool(EToolShutdownType::Completed);
 				}
 
-				EditorContent->SetPrimarySelectedNode(PrimarySelection);
+				EditorContent->SetSelectedNode(PrimarySelection);
 
 				// Call the node's OnSelected function. Some nodes use this to cache information from the inputs (e.g. FDataflowCollectionAddScalarVertexPropertyNode::CachedCollectionGroupNames)
 				TSharedPtr<Dataflow::FEngineContext> DataflowContext = EditorContent->GetDataflowContext();
 				if (PrimarySelection && DataflowContext.IsValid())
 				{
 					PrimarySelection->GetDataflowNode()->OnSelected(*DataflowContext);
+
+					// Update selected Collection in the ContextObject
+					const TSharedPtr<FDataflowNode> DataflowNode = PrimarySelection->GetDataflowNode();
+					for (const FDataflowOutput* const Output : DataflowNode->GetOutputs())
+					{
+						if (Output->GetType() == FName(TEXT("FManagedArrayCollection")))
+						{
+							const FManagedArrayCollection DefaultValue;
+							TSharedRef<FManagedArrayCollection> Collection = MakeShared<FManagedArrayCollection>(Output->GetValue<FManagedArrayCollection>(*DataflowContext, DefaultValue));
+							constexpr bool bCollectionIsInput = false;
+							EditorContent->SetSelectedCollection(Collection, bCollectionIsInput);
+						}
+					}
 				}
 			}
 
@@ -1067,6 +1085,16 @@ void FDataflowEditorToolkit::Tick(float DeltaTime)
 					nullptr, nullptr, EditorContent->GetTerminalAsset(), EditorContent->GetDataflowTerminal());
 			}
 			EditorContent->SetLastModifiedTimestamp(TimeStamp);
+
+			// Ensure the context object's selected node matches the selected node in the graph editor
+			// TODO: Create an Editor Context Object that can just hold a reference to the graph editor, rather than keeping these in sync
+			if (GraphEditor->GetNumberOfSelectedNodes() == 1)
+			{
+				const UDataflowEdNode* const ContextObjectSelectedNode = EditorContent->GetSelectedNode();
+				const UDataflowEdNode* const EditorSelectedNode = Cast<UDataflowEdNode>(GraphEditor->GetSingleSelectedNode());
+				ensure(EditorSelectedNode == ContextObjectSelectedNode);
+			}
+
 		}
 	}
 }
diff --git a/Engine/Plugins/Experimental/Dataflow/Source/DataflowEditor/Private/Dataflow/DataflowToolRegistry.cpp b/Engine/Plugins/Experimental/Dataflow/Source/DataflowEditor/Private/Dataflow/DataflowToolRegistry.cpp
new file mode 100644
--- /dev/null
+++ b/Engine/Plugins/Experimental/Dataflow/Source/DataflowEditor/Private/Dataflow/DataflowToolRegistry.cpp
@@ -0,0 +1,56 @@
+// Copyright Epic Games, Inc. All Rights Reserved.
+
+#include "Dataflow/DataflowToolRegistry.h"
+#include "Misc/LazySingleton.h"
+#include "Framework/Commands/UICommandInfo.h"
+
+namespace Dataflow
+{
+
+	FDataflowToolRegistry& FDataflowToolRegistry::Get()
+	{
+		return TLazySingleton<FDataflowToolRegistry>::Get();
+	}
+
+	void FDataflowToolRegistry::TearDown()
+	{
+		TLazySingleton<FDataflowToolRegistry>::TearDown();
+	}
+
+	void FDataflowToolRegistry::AddNodeToToolMapping(const FName& NodeName, TObjectPtr<UInteractiveToolBuilder> ToolBuilder)
+	{
+		// FUICommandInfo is uninitialized, it will be created later in FDataflowEditorCommandsImpl::RegisterCommands
+		NodeTypeToToolMap.Add(NodeName, { ToolBuilder, nullptr });
+	}
+
+	void FDataflowToolRegistry::RemoveNodeToToolMapping(const FName& NodeName)
+	{
+		NodeTypeToToolMap.Remove(NodeName);
+	}
+
+	TArray<FName> FDataflowToolRegistry::GetNodeNames() const
+	{
+		TArray<FName> NodeNames;
+		NodeTypeToToolMap.GetKeys(NodeNames);
+		return NodeNames;
+	}
+	
+	TSharedPtr<FUICommandInfo>& FDataflowToolRegistry::GetToolCommandForNode(const FName& NodeName)
+	{
+		check(NodeTypeToToolMap.Contains(NodeName));
+		return NodeTypeToToolMap[NodeName].ToolCommand;
+	}
+
+	UInteractiveToolBuilder* FDataflowToolRegistry::GetToolBuilderForNode(const FName& NodeName)
+	{
+		check(NodeTypeToToolMap.Contains(NodeName));
+		return NodeTypeToToolMap[NodeName].ToolBuilder;
+	}
+
+	const UInteractiveToolBuilder* FDataflowToolRegistry::GetToolBuilderForNode(const FName& NodeName) const
+	{
+		check(NodeTypeToToolMap.Contains(NodeName));
+		return NodeTypeToToolMap[NodeName].ToolBuilder;
+	}
+
+}
diff --git a/Engine/Plugins/Experimental/Dataflow/Source/DataflowEditor/Private/DataflowEditorTools/DataflowEditorWeightMapPaintTool.cpp b/Engine/Plugins/Experimental/Dataflow/Source/DataflowEditor/Private/DataflowEditorTools/DataflowEditorWeightMapPaintTool.cpp
--- a/Engine/Plugins/Experimental/Dataflow/Source/DataflowEditor/Private/DataflowEditorTools/DataflowEditorWeightMapPaintTool.cpp
+++ b/Engine/Plugins/Experimental/Dataflow/Source/DataflowEditor/Private/DataflowEditorTools/DataflowEditorWeightMapPaintTool.cpp
@@ -100,11 +100,11 @@ bool UDataflowEditorWeightMapPaintToolBuilder::CanBuildTool(const FToolBuilderSt
 			{
 				if (UDataflowBaseContent* ContextObject = SceneState.ToolManager->GetContextObjectStore()->FindContext<UDataflowBaseContent>())
 				{
-					if (ContextObject->GetPrimarySelectedNode() == Component->Node)
+					if (ContextObject->GetSelectedNode() == Component->Node)
 					{
 						if (const TSharedPtr<Dataflow::FEngineContext> EvaluationContext = ContextObject->GetDataflowContext())
 						{
-							if (const FDataflowNode* PrimarySelection = ContextObject->GetPrimarySelectedNodeOfType<FDataflowCollectionAddScalarVertexPropertyNode>())
+							if (const FDataflowNode* PrimarySelection = ContextObject->GetSelectedNodeOfType<FDataflowCollectionAddScalarVertexPropertyNode>())
 							{
 								return HasManagedArrayCollection(PrimarySelection, EvaluationContext);
 							}
@@ -167,7 +167,7 @@ void UDataflowEditorWeightMapPaintTool::Setup()
 	UMeshSculptToolBase::Setup();
 
 	// Get the selected weight map node
-	WeightMapNodeToUpdate = DataflowEditorContextObject->GetPrimarySelectedNodeOfType<FDataflowCollectionAddScalarVertexPropertyNode>();
+	WeightMapNodeToUpdate = DataflowEditorContextObject->GetSelectedNodeOfType<FDataflowCollectionAddScalarVertexPropertyNode>();
 	checkf(WeightMapNodeToUpdate, TEXT("No Weight Map Node is currently selected, or more than one node is selected"));
 
 	SetToolDisplayName(LOCTEXT("ToolName", "Paint Weight Maps"));
@@ -349,7 +349,7 @@ void UDataflowEditorWeightMapPaintTool::Setup()
 	// Setup DynamicMeshToWeight conversion
 	if (DataflowEditorContextObject)
 	{
-		if (const TSharedPtr<const FManagedArrayCollection> Collection = DataflowEditorContextObject->GetPrimaryRenderCollection())
+		if (const TSharedPtr<const FManagedArrayCollection> Collection = DataflowEditorContextObject->GetRenderCollection())
 		{
 			using namespace Dataflow;
 			const FNonManifoldMappingSupport NonManifoldMapping(*Mesh);
diff --git a/Engine/Plugins/Experimental/Dataflow/Source/DataflowEditor/Public/Dataflow/DataflowEditorCommands.h b/Engine/Plugins/Experimental/Dataflow/Source/DataflowEditor/Public/Dataflow/DataflowEditorCommands.h
--- a/Engine/Plugins/Experimental/Dataflow/Source/DataflowEditor/Public/Dataflow/DataflowEditorCommands.h
+++ b/Engine/Plugins/Experimental/Dataflow/Source/DataflowEditor/Public/Dataflow/DataflowEditorCommands.h
@@ -57,7 +57,9 @@ public:
 
 	TMap< FName, TSharedPtr<FUICommandInfo> > CreateNodesMap;
 
+	UE_DEPRECATED(5.5, "Dataflow Tool commands are now stored in FDataflowToolRegistry")
 	const static FString BeginWeightMapPaintToolIdentifier;
+	UE_DEPRECATED(5.5, "Dataflow Tool commands are now stored in FDataflowToolRegistry")
 	TSharedPtr<FUICommandInfo> BeginWeightMapPaintTool;
 	
 	const static FString AddWeightMapNodeIdentifier;
diff --git a/Engine/Plugins/Experimental/Dataflow/Source/DataflowEditor/Public/Dataflow/DataflowEditorMode.h b/Engine/Plugins/Experimental/Dataflow/Source/DataflowEditor/Public/Dataflow/DataflowEditorMode.h
--- a/Engine/Plugins/Experimental/Dataflow/Source/DataflowEditor/Public/Dataflow/DataflowEditorMode.h
+++ b/Engine/Plugins/Experimental/Dataflow/Source/DataflowEditor/Public/Dataflow/DataflowEditorMode.h
@@ -114,7 +114,6 @@ private:
 	void RegisterDataflowTool(TSharedPtr<FUICommandInfo> UICommand,
 							  FString ToolIdentifier,
 							  UInteractiveToolBuilder* Builder,
-							  const IDataflowEditorToolBuilder* DataflowToolBuilder,
 							  UEditorInteractiveToolsContext* const ToolsContext,
 							  EToolsContextScope ToolScope = EToolsContextScope::Default);
 
diff --git a/Engine/Plugins/Experimental/Dataflow/Source/DataflowEditor/Public/Dataflow/DataflowToolRegistry.h b/Engine/Plugins/Experimental/Dataflow/Source/DataflowEditor/Public/Dataflow/DataflowToolRegistry.h
new file mode 100644
--- /dev/null
+++ b/Engine/Plugins/Experimental/Dataflow/Source/DataflowEditor/Public/Dataflow/DataflowToolRegistry.h
@@ -0,0 +1,46 @@
+// Copyright Epic Games, Inc. All Rights Reserved.
+
+#pragma once
+
+#include "UObject/NameTypes.h"
+#include "Templates/SharedPointer.h"
+#include "UObject/ObjectPtr.h"
+
+class FUICommandInfo;
+class UInteractiveToolBuilder;
+
+namespace Dataflow
+{
+	class DATAFLOWEDITOR_API FDataflowToolRegistry
+	{
+	public:
+
+		static FDataflowToolRegistry& Get();
+		static void TearDown();
+
+		void AddNodeToToolMapping(const FName& NodeName, TObjectPtr<UInteractiveToolBuilder> ToolBuilder);
+		void RemoveNodeToToolMapping(const FName& NodeName);
+
+		TArray<FName> GetNodeNames() const;
+
+		TSharedPtr<FUICommandInfo>& GetToolCommandForNode(const FName& NodeName);
+
+		UInteractiveToolBuilder* GetToolBuilderForNode(const FName& NodeName);
+		const UInteractiveToolBuilder* GetToolBuilderForNode(const FName& NodeName) const;
+
+	private:
+
+		struct FToolInfo
+		{
+			// Specified when registering the tool
+			TObjectPtr<UInteractiveToolBuilder> ToolBuilder;
+
+			// Constructed automatically in FDataflowEditorCommandsImpl::RegisterCommands
+			TSharedPtr<FUICommandInfo> ToolCommand;
+		};
+
+		TMap<FName, FToolInfo> NodeTypeToToolMap;
+	
+	};
+
+}
diff --git a/Engine/Source/Runtime/Experimental/Dataflow/Engine/Private/Dataflow/DataflowContent.cpp b/Engine/Source/Runtime/Experimental/Dataflow/Engine/Private/Dataflow/DataflowContent.cpp
--- a/Engine/Source/Runtime/Experimental/Dataflow/Engine/Private/Dataflow/DataflowContent.cpp
+++ b/Engine/Source/Runtime/Experimental/Dataflow/Engine/Private/Dataflow/DataflowContent.cpp
@@ -253,14 +253,6 @@ namespace DataflowContextHelpers
 	template DATAFLOWENGINE_API TObjectPtr<UDataflowSkeletalContent> CreateNewDataflowContent(const TObjectPtr<UObject>& ContentOwner);
 }
 
-
-void UDataflowContextObject::AddReferencedObjects(UObject* InThis, FReferenceCollector& Collector)
-{
-	UDataflowContextObject* This = CastChecked<UDataflowContextObject>(InThis);
-	Collector.AddReferencedObject(This->PrimarySelectedNode);
-	Super::AddReferencedObjects(InThis, Collector);
-}
-
 TObjectPtr<UDataflowBaseContent> IDataflowContentOwner::BuildDataflowContent()
 {
 	if(TObjectPtr<UDataflowBaseContent> DataflowContent = CreateDataflowContent())
@@ -339,11 +331,6 @@ void UDataflowBaseContent::SetDataflowAsset(const TObjectPtr<UDataflow>& Dataflo
 	SetSimulationDirty(true);
 }
 
-TObjectPtr<UDataflow> UDataflowBaseContent::GetDataflowAsset() const 
-{
-	return DataflowGraph;
-}
-
 #if WITH_EDITOR
 
 void UDataflowBaseContent::PostEditChangeProperty(FPropertyChangedEvent& PropertyChangedEvent) 
diff --git a/Engine/Source/Runtime/Experimental/Dataflow/Engine/Private/Dataflow/DataflowContextObject.cpp b/Engine/Source/Runtime/Experimental/Dataflow/Engine/Private/Dataflow/DataflowContextObject.cpp
new file mode 100644
--- /dev/null
+++ b/Engine/Source/Runtime/Experimental/Dataflow/Engine/Private/Dataflow/DataflowContextObject.cpp
@@ -0,0 +1,14 @@
+// Copyright Epic Games, Inc. All Rights Reserved.
+
+#include "Dataflow/DataflowContextObject.h"
+
+#include UE_INLINE_GENERATED_CPP_BY_NAME(DataflowContextObject)
+
+
+void UDataflowContextObject::AddReferencedObjects(UObject* InThis, FReferenceCollector& Collector)
+{
+	UDataflowContextObject* This = CastChecked<UDataflowContextObject>(InThis);
+	Collector.AddReferencedObject(This->SelectedNode);
+	Super::AddReferencedObjects(InThis, Collector);
+}
+
diff --git a/Engine/Source/Runtime/Experimental/Dataflow/Engine/Public/Dataflow/DataflowContent.h b/Engine/Source/Runtime/Experimental/Dataflow/Engine/Public/Dataflow/DataflowContent.h
--- a/Engine/Source/Runtime/Experimental/Dataflow/Engine/Public/Dataflow/DataflowContent.h
+++ b/Engine/Source/Runtime/Experimental/Dataflow/Engine/Public/Dataflow/DataflowContent.h
@@ -2,13 +2,12 @@
 
 #pragma once
 
+#include "DataflowContextObject.h"
 #include "UObject/Interface.h"
-#include "Dataflow/DataflowEdNode.h"
 #include "Templates/SharedPointer.h"
 #include "DataflowContent.generated.h"
 
 class FDataflowEditorToolkit;
-class UDataflow;
 class USkeletalMesh;
 class USkeleton;
 class USkeletalMeshComponent;
@@ -17,15 +16,6 @@ class UDataflowBaseContent;
 class FPreviewScene;
 class UAnimSingleNodeInstance;
 class AActor;
-class UDataflowEdNode;
-struct FManagedArrayCollection;
-
-namespace Dataflow
-{
-	class IDataflowConstructionViewMode;
-	template<class Base> class TEngineContext;
-	typedef TEngineContext<FContextSingle> FEngineContext;
-}
 
 namespace DataflowContextHelpers
 {
@@ -34,55 +24,6 @@ namespace DataflowContextHelpers
 	DATAFLOWENGINE_API TObjectPtr<T> CreateNewDataflowContent(const TObjectPtr<UObject>& ContentOwner);
 }
 
-
-/** 
- * Context object used for selection/rendering 
- */
-
-UCLASS()
-class DATAFLOWENGINE_API UDataflowContextObject : public UObject
-{
-	GENERATED_BODY()
-public:
-
-	/** Selection Collection Access */
-	void SetPrimarySelectedNode(TObjectPtr<UDataflowEdNode> InSelectedNode) { PrimarySelectedNode = InSelectedNode; }
-	TObjectPtr<UDataflowEdNode> GetPrimarySelectedNode() const { return PrimarySelectedNode; }
-
-	/** Render Collection used to generate the DynamicMesh3D on the PrimarySelection */
-	void SetPrimaryRenderCollection(const TSharedPtr<FManagedArrayCollection>& InCollection) { PrimaryRenderCollection = InCollection; }
-	TSharedPtr<const FManagedArrayCollection> GetPrimaryRenderCollection() const { return PrimaryRenderCollection; }
-
-	/** ViewMode Access */
-	void SetConstructionViewMode(const Dataflow::IDataflowConstructionViewMode* InMode) { ConstructionViewMode = InMode; }
-	const Dataflow::IDataflowConstructionViewMode* GetConstructionViewMode() const { return ConstructionViewMode; }
-
-	/** Get a single selected node of the specified type. Return nullptr if the specified node is not selected, or if multiple nodes are selected*/
-	template<typename NodeType>
-	NodeType* GetPrimarySelectedNodeOfType() const 
-	{
-		if (PrimarySelectedNode && PrimarySelectedNode->GetDataflowNode()) 
-		{
-			return PrimarySelectedNode->GetDataflowNode()->AsType<NodeType>();
-		}
-		return nullptr;
-	}
-
-	//~ UObject interface
-	static void AddReferencedObjects(UObject* InThis, FReferenceCollector& Collector);
-
-protected:
-
-	/** Render collection to be used */
-	TSharedPtr<FManagedArrayCollection> PrimaryRenderCollection = nullptr;
-
-	/** Primary node that is selected in the graph */
-	TObjectPtr<UDataflowEdNode> PrimarySelectedNode = nullptr;
-
-	/** Construction view mode for the context object @todo(michael) : is it only for construction or for simulation as well*/
-	const Dataflow::IDataflowConstructionViewMode* ConstructionViewMode = nullptr;
-};
-
 UINTERFACE(MinimalAPI)
 class UDataflowContentOwner : public UInterface
 {
@@ -158,10 +99,8 @@ public:
 	*	Context - Dataflow Evaluation State
 	*   Dataflow context stores the evaluated state of the graph. 
 	*/
-	void SetDataflowContext(const TSharedPtr<Dataflow::FEngineContext>& InContext);
-	const TSharedPtr<Dataflow::FEngineContext>& GetDataflowContext() const { return DataflowContext; }
-	TSharedPtr<Dataflow::FEngineContext>& GetDataflowContext() { return DataflowContext; }
-
+	virtual void SetDataflowContext(const TSharedPtr<Dataflow::FEngineContext>& InContext) override;
+ 
 	/** Rebuild the owner dependent datas  */
 	void UpdateContentDatas();
 
@@ -176,8 +115,7 @@ public:
 	TObjectPtr<UObject> GetDataflowOwner() const;
 	
 	/** Data flow asset accessors (through the context) */
-	void SetDataflowAsset(const TObjectPtr<UDataflow>& InAsset);
-	TObjectPtr<UDataflow> GetDataflowAsset() const;
+	virtual void SetDataflowAsset(const TObjectPtr<UDataflow>& InAsset) override;
 
 	/** Data flow terminal accessors */
 	void SetDataflowTerminal(const FString& InPath) { DataflowTerminal = InPath;  SetConstructionDirty(true); SetSimulationDirty(true);}
@@ -210,18 +148,11 @@ protected:
 	/** Data flow terminal path for evaluation */
 	UPROPERTY(Transient, SkipSerialization)
 	FString DataflowTerminal = "";
-
-	/** Dataflow graph for evaluation */
-	UPROPERTY(Transient, SkipSerialization)
-	TObjectPtr<UDataflow> DataflowGraph;
 	
 	/** Data flow terminal path for evaluation */
 	UPROPERTY(Transient, SkipSerialization)
 	TObjectPtr<UObject> TerminalAsset = nullptr;
 
-	/**  Engine context (data flow owner/asset) to be used for dataflow evaluation */
-    TSharedPtr<Dataflow::FEngineContext> DataflowContext = nullptr;
-
     /** Last data flow evaluated node time stamp */
 	Dataflow::FTimestamp LastModifiedTimestamp = Dataflow::FTimestamp::Invalid;
 
diff --git a/Engine/Source/Runtime/Experimental/Dataflow/Engine/Public/Dataflow/DataflowContextObject.h b/Engine/Source/Runtime/Experimental/Dataflow/Engine/Public/Dataflow/DataflowContextObject.h
new file mode 100644
--- /dev/null
+++ b/Engine/Source/Runtime/Experimental/Dataflow/Engine/Public/Dataflow/DataflowContextObject.h
@@ -0,0 +1,103 @@
+// Copyright Epic Games, Inc. All Rights Reserved.
+
+#pragma once
+
+#include "UObject/Object.h"
+#include "Dataflow/DataflowEdNode.h"
+#include "Dataflow/DataflowNodeParameters.h"
+#include "DataflowContextObject.generated.h"
+
+struct FManagedArrayCollection;
+class UDataflow;
+namespace Dataflow
+{
+	class IDataflowConstructionViewMode;
+	template<class Base> class TEngineContext;
+	typedef TEngineContext<FContextSingle> FEngineContext;
+}
+
+/** 
+ * Context object used for selection/rendering 
+ */
+
+UCLASS()
+class DATAFLOWENGINE_API UDataflowContextObject : public UObject
+{
+	GENERATED_BODY()
+public:
+
+	/** Selection Collection Access */
+	void SetSelectedNode(TObjectPtr<UDataflowEdNode> InSelectedNode) { SelectedNode = InSelectedNode; }
+	TObjectPtr<UDataflowEdNode> GetSelectedNode() const { return SelectedNode; }
+
+	/** Get a single selected node of the specified type. Return nullptr if the specified node is not selected, or if multiple nodes are selected*/
+	template<typename NodeType>
+	NodeType* GetSelectedNodeOfType() const
+	{
+		if (SelectedNode && SelectedNode->GetDataflowNode())
+		{
+			return SelectedNode->GetDataflowNode()->AsType<NodeType>();
+		}
+		return nullptr;
+	}
+
+	/** Render Collection used to generate the DynamicMesh3D from the selected node */
+	void SetRenderCollection(const TSharedPtr<FManagedArrayCollection>& InCollection) { RenderCollection = InCollection; }
+	TSharedPtr<const FManagedArrayCollection> GetRenderCollection() const { return RenderCollection; }
+
+	/** ViewMode Access */
+	void SetConstructionViewMode(const Dataflow::IDataflowConstructionViewMode* InMode) { ConstructionViewMode = InMode; }
+	const Dataflow::IDataflowConstructionViewMode* GetConstructionViewMode() const { return ConstructionViewMode; }
+	
+	/**
+	*	Context - Dataflow Evaluation State
+	*   Dataflow context stores the evaluated state of the graph.
+	*/
+	virtual void SetDataflowContext(const TSharedPtr<Dataflow::FEngineContext>& InDataflowContext) { DataflowContext = InDataflowContext; }
+	const TSharedPtr<Dataflow::FEngineContext>& GetDataflowContext() const { return DataflowContext; }
+	TSharedPtr<Dataflow::FEngineContext>& GetDataflowContext() { return DataflowContext; }
+
+	/** Dataflow Asset */
+	virtual void SetDataflowAsset(const TObjectPtr<UDataflow>& InAsset) { DataflowGraph = InAsset; }
+	TObjectPtr<UDataflow> GetDataflowAsset() const { return DataflowGraph;  }
+
+	/** Collection passing through the currently selected node and a flag indicating whether it's on a node input or output */
+	void SetSelectedCollection(TSharedPtr<const FManagedArrayCollection> InSelectedCollection, bool bInUsingInputCollection) 
+	{ 
+		SelectedCollection = InSelectedCollection; 
+		bUsingInputCollection = bInUsingInputCollection;
+	}
+	const TSharedPtr<const FManagedArrayCollection> GetSelectedCollection() const { return SelectedCollection; }
+	bool IsUsingInputCollection() const { return bUsingInputCollection; }
+
+
+	//~ UObject interface
+	static void AddReferencedObjects(UObject* InThis, FReferenceCollector& Collector);
+
+protected:
+
+	/** Render collection to be used */
+	TSharedPtr<FManagedArrayCollection> RenderCollection = nullptr;
+
+	/** Node that is selected in the graph */
+	UPROPERTY(Transient, SkipSerialization)
+	TObjectPtr<UDataflowEdNode> SelectedNode = nullptr;
+
+	/** Construction view mode for the context object @todo(michael) : is it only for construction or for simulation as well*/
+	const Dataflow::IDataflowConstructionViewMode* ConstructionViewMode = nullptr;
+
+	/** Engine context (data flow owner/asset) to be used for dataflow evaluation */
+	TSharedPtr<Dataflow::FEngineContext> DataflowContext = nullptr;
+
+	/** Dataflow graph for evaluation */
+	UPROPERTY(Transient, SkipSerialization)
+	TObjectPtr<UDataflow> DataflowGraph = nullptr;
+
+	/** Managed Array Collection passing through the currently selected node */
+	TSharedPtr<const FManagedArrayCollection> SelectedCollection = nullptr;
+
+	/** Whether the Collection is taken from a node Input (i.e. it's in the the state before node execution) */
+	bool bUsingInputCollection = false;
+
+};
+
