Cloth Editor: remove secondary selection from Selection Tool. Allow importing secondary selection from input collection, to allow safe data migration.

diff --git a/Engine/Plugins/ChaosClothAsset/Source/ChaosClothAsset/Private/ChaosClothAsset/ClothGeometryTools.cpp b/Engine/Plugins/ChaosClothAsset/Source/ChaosClothAsset/Private/ChaosClothAsset/ClothGeometryTools.cpp
--- a/Engine/Plugins/ChaosClothAsset/Source/ChaosClothAsset/Private/ChaosClothAsset/ClothGeometryTools.cpp
+++ b/Engine/Plugins/ChaosClothAsset/Source/ChaosClothAsset/Private/ChaosClothAsset/ClothGeometryTools.cpp
@@ -1425,6 +1425,15 @@ namespace UE::Chaos::ClothAsset
 		return false;
 	}
 
+	bool FClothGeometryTools::ConvertSelectionToNewGroupType(const TSharedRef<const FManagedArrayCollection>& ClothCollection, const FName& SelectionName, const FName& GroupName, TSet<int32>& OutSelectionSet)
+	{
+		constexpr bool bSecondarySelection = false;
+		PRAGMA_DISABLE_DEPRECATION_WARNINGS
+		return ConvertSelectionToNewGroupType(ClothCollection, SelectionName, GroupName, bSecondarySelection, OutSelectionSet);
+		PRAGMA_ENABLE_DEPRECATION_WARNINGS
+	}
+
+
 	void FClothGeometryTools::TransferWeightMap(
 		const TConstArrayView<FVector3f>& SourcePositions,
 		const TConstArrayView<FIntVector3>& InSourceIndices,
diff --git a/Engine/Plugins/ChaosClothAsset/Source/ChaosClothAsset/Public/ChaosClothAsset/ClothGeometryTools.h b/Engine/Plugins/ChaosClothAsset/Source/ChaosClothAsset/Public/ChaosClothAsset/ClothGeometryTools.h
--- a/Engine/Plugins/ChaosClothAsset/Source/ChaosClothAsset/Public/ChaosClothAsset/ClothGeometryTools.h
+++ b/Engine/Plugins/ChaosClothAsset/Source/ChaosClothAsset/Public/ChaosClothAsset/ClothGeometryTools.h
@@ -128,8 +128,20 @@ namespace UE::Chaos::ClothAsset
 		 * @param OutSelectionSet copy of the selection. Unchanged when function returns false.
 		 * @return success (will return false if the selection is not found or conversion is not possible)
 		 */
+		UE_DEPRECATED(5.5, "Please use the version with no bSecondarySelection parameter")
 		static bool ConvertSelectionToNewGroupType(const TSharedRef<const FManagedArrayCollection>& ClothCollection, const FName& SelectionName, const FName& GroupName, bool bSecondarySelection, TSet<int32>& OutSelectionSet);
 
+		/**
+		 * Get a copy of the selection, converting to the desired group if possible.
+		 * Currently only conversions between vertex and face components on the same mesh type are supported.
+		 * @param ClothCollection to query
+		 * @param SelectionName the selection name
+		 * @param GroupName the group name
+		 * @param OutSelectionSet copy of the selection. Unchanged when function returns false.
+		 * @return success (will return false if the selection is not found or conversion is not possible)
+		 */
+		static bool ConvertSelectionToNewGroupType(const TSharedRef<const FManagedArrayCollection>& ClothCollection, const FName& SelectionName, const FName& GroupName, TSet<int32>& OutSelectionSet);
+
 		/**
 		 * Transfer a vertex weight map from a source to target mesh
 		 */
diff --git a/Engine/Plugins/ChaosClothAssetEditor/Source/ChaosClothAssetDataflowNodes/Private/ChaosClothAsset/SelectionNode.cpp b/Engine/Plugins/ChaosClothAssetEditor/Source/ChaosClothAssetDataflowNodes/Private/ChaosClothAsset/SelectionNode.cpp
--- a/Engine/Plugins/ChaosClothAssetEditor/Source/ChaosClothAssetDataflowNodes/Private/ChaosClothAsset/SelectionNode.cpp
+++ b/Engine/Plugins/ChaosClothAssetEditor/Source/ChaosClothAssetDataflowNodes/Private/ChaosClothAsset/SelectionNode.cpp
@@ -65,10 +65,13 @@ bool TransferSelectionSet(const TSharedRef<const FManagedArrayCollection>& Trans
 	TSet<int32> TransferSet;
 	const FName DesiredTransferGroup = bIsValidRenderSelection ? ClothCollectionGroup::RenderVertices :
 		(SimTransferType == EChaosClothAssetWeightMapTransferType::Use2DSimMesh ? ClothCollectionGroup::SimVertices2D : ClothCollectionGroup::SimVertices3D);
+	
+	PRAGMA_DISABLE_DEPRECATION_WARNINGS
 	if (!FClothGeometryTools::ConvertSelectionToNewGroupType(TransferClothCollection, InInputName, DesiredTransferGroup, bIsSecondarySelection, TransferSet))
 	{
 		return false;
 	}
+	PRAGMA_ENABLE_DEPRECATION_WARNINGS
 
 	// Convert to weights that are 0 on unselected vertices and 1 on selected vertices
 	TArray<float> TransferWeights;
@@ -303,7 +306,7 @@ void FChaosClothAssetSelectionNode::SetAssetValue(TObjectPtr<UObject> Asset, Dat
 							if (Private::TransferSelectionSet<false>(TransferClothCollection, ClothCollection, InInputName, SelectionGroupName, SimTransferType, TransferSelectionThreshold, PrimaryFinalSelection))
 							{
 								TSet<int32> InputSelection;
-								FClothGeometryTools::ConvertSelectionToNewGroupType(ClothCollection, InInputName, SelectionGroupName, false, InputSelection);
+								FClothGeometryTools::ConvertSelectionToNewGroupType(ClothCollection, InInputName, SelectionGroupName, InputSelection);
 
 								MutableThis->SetIndices(InputSelection, PrimaryFinalSelection);
 							}
@@ -312,7 +315,10 @@ void FChaosClothAssetSelectionNode::SetAssetValue(TObjectPtr<UObject> Asset, Dat
 							if (Private::TransferSelectionSet<true>(TransferClothCollection, ClothCollection, InInputName, SelectionSecondaryGroupName, SimTransferType, TransferSelectionThreshold, SecondaryFinalSelection))
 							{
 								TSet<int32> InputSelection;
+
+								PRAGMA_DISABLE_DEPRECATION_WARNINGS
 								FClothGeometryTools::ConvertSelectionToNewGroupType(ClothCollection, InInputName, SelectionGroupName, true, InputSelection);
+								PRAGMA_ENABLE_DEPRECATION_WARNINGS
 
 								MutableThis->SetSecondaryIndices(InputSelection, SecondaryFinalSelection);
 							}
@@ -389,7 +395,7 @@ void FChaosClothAssetSelectionNode::Evaluate(Dataflow::FContext& Context, const
 		check(SelectionFacade.IsValid());
 
 		TSet<int32> InputSelectionSet;
-		FClothGeometryTools::ConvertSelectionToNewGroupType(SelectionCollection, InInputName, SelectionGroupName, false, InputSelectionSet);
+		FClothGeometryTools::ConvertSelectionToNewGroupType(SelectionCollection, InInputName, SelectionGroupName, InputSelectionSet);
 		TSet<int32> FinalSet;
 		CalculateFinalSet(InputSelectionSet, FinalSet);
 
@@ -402,7 +408,9 @@ void FChaosClothAssetSelectionNode::Evaluate(Dataflow::FContext& Context, const
 			const FName SecondarySelectionGroupName = (*SecondaryGroup.Name);
 
 			InputSelectionSet.Reset();
+			PRAGMA_DISABLE_DEPRECATION_WARNINGS
 			FClothGeometryTools::ConvertSelectionToNewGroupType(SelectionCollection, InInputName, SelectionGroupName, true, InputSelectionSet);
+			PRAGMA_ENABLE_DEPRECATION_WARNINGS
 			FinalSet.Reset();
 			CalculateFinalSecondarySet(InputSelectionSet, FinalSet);
 			TSet<int32>& SecondarySelectionSet = SelectionFacade.FindOrAddSelectionSecondarySet(SecondarySelectionName, SecondarySelectionGroupName);
diff --git a/Engine/Plugins/ChaosClothAssetEditor/Source/ChaosClothAssetEditorTools/Private/ChaosClothAsset/ClothMeshSelectionTool.cpp b/Engine/Plugins/ChaosClothAssetEditor/Source/ChaosClothAssetEditorTools/Private/ChaosClothAsset/ClothMeshSelectionTool.cpp
--- a/Engine/Plugins/ChaosClothAssetEditor/Source/ChaosClothAssetEditorTools/Private/ChaosClothAsset/ClothMeshSelectionTool.cpp
+++ b/Engine/Plugins/ChaosClothAssetEditor/Source/ChaosClothAssetEditorTools/Private/ChaosClothAsset/ClothMeshSelectionTool.cpp
@@ -305,7 +305,6 @@ void UClothMeshSelectionTool::Setup()
 	});
 
 	// Order of operations is important here:
-	// ToolProperties->RestoreProperties should happen before GetSelectedNodeInfo so that we can know whether we need to get Primary or Secondary set
 	// ToolProperties->WatchProperty(ToolProperties->Name) should happen after GetSelectedNodeInfo so that we can set the OriginalName
 
 	ToolProperties->RestoreProperties(this);
@@ -355,8 +354,6 @@ void UClothMeshSelectionTool::Setup()
 	AddToolPropertySource(ToolProperties);
 
 	AddToolPropertySource(SelectionMechanic->Properties);
-
-	UpdatePrimarySecondaryMessage();
 }
 
 void UClothMeshSelectionTool::OnShutdown(EToolShutdownType ShutdownType)
@@ -366,12 +363,7 @@ void UClothMeshSelectionTool::OnShutdown(EToolShutdownType ShutdownType)
 		GetToolManager()->BeginUndoTransaction(LOCTEXT("SelectionToolTransactionName", "Mesh Selection"));
 		UpdateSelectedNode();
 		GetToolManager()->EndUndoTransaction();
-	}
 
-	// Invalidate the node even if we are hitting cancel. We could have saved new selection information to the node by switching primary/secondary modes and we'd expect 
-	// that information to make its way into the ClothCollection
-	if (SelectionNodeToUpdate)
-	{
 		SelectionNodeToUpdate->Invalidate();
 	}
 
@@ -493,8 +485,6 @@ bool UClothMeshSelectionTool::GetSelectedNodeInfo(FString& OutSelectionName, UE:
 		}
 	};
 
-	const bool bSecondarySelection = ToolProperties && ToolProperties->bSecondarySelection;
-
 	// Get the input set.
 	InputSelectionSet.Reset();
 	if (DataflowContextObject)
@@ -506,25 +496,16 @@ bool UClothMeshSelectionTool::GetSelectedNodeInfo(FString& OutSelectionName, UE:
 			{
 				using namespace UE::Chaos::ClothAsset;
 				const FName InputName = SelectionNodeToUpdate->GetInputName(*DataflowContext);
-				const FName GroupName = bSecondarySelection ? FName(*SelectionNodeToUpdate->SecondaryGroup.Name) : FName(*SelectionNodeToUpdate->Group.Name);
+				const FName GroupName = FName(*SelectionNodeToUpdate->Group.Name);
 
-				FClothGeometryTools::ConvertSelectionToNewGroupType(ClothCollection.ToSharedRef(), InputName, GroupName, bSecondarySelection, InputSelectionSet);
+				FClothGeometryTools::ConvertSelectionToNewGroupType(ClothCollection.ToSharedRef(), InputName, GroupName, InputSelectionSet);
 			}
 		}
 	}
 
-	if (bSecondarySelection)
-	{
-		TSet<int32> FinalSet;
-		SelectionNodeToUpdate->CalculateFinalSecondarySet(InputSelectionSet, FinalSet);
-		AppendSet(SelectionNodeToUpdate->SecondaryGroup, FinalSet);
-	}
-	else
-	{
-		TSet<int32> FinalSet;
-		SelectionNodeToUpdate->CalculateFinalSet(InputSelectionSet, FinalSet);
-		AppendSet(SelectionNodeToUpdate->Group, FinalSet);
-	}
+	TSet<int32> FinalSet;
+	SelectionNodeToUpdate->CalculateFinalSet(InputSelectionSet, FinalSet);
+	AppendSet(SelectionNodeToUpdate->Group, FinalSet);
 
 	OutSelectionName = SelectionNodeToUpdate->Name;
 	OutOverrideType = SelectionNodeToUpdate->SelectionOverrideType;
@@ -613,18 +594,10 @@ void UClothMeshSelectionTool::UpdateSelectedNode()
 
 	SelectionNodeToUpdate->Name = ToolProperties->Name;
 	SelectionNodeToUpdate->SelectionOverrideType = ToolProperties->SelectionOverrideType;
-	if (!ToolProperties->bSecondarySelection)
-	{
-		TSet<int32> FinalSet;
-		GetFinalSet(SelectionNodeToUpdate->Group, FinalSet);
-		SelectionNodeToUpdate->SetIndices(InputSelectionSet, FinalSet);
-	}
-	else
-	{
-		TSet<int32> FinalSet;
-		GetFinalSet(SelectionNodeToUpdate->SecondaryGroup, FinalSet);
-		SelectionNodeToUpdate->SetSecondaryIndices(InputSelectionSet, FinalSet);
-	}
+
+	TSet<int32> FinalSet;
+	GetFinalSet(SelectionNodeToUpdate->Group, FinalSet);
+	SelectionNodeToUpdate->SetIndices(InputSelectionSet, FinalSet);
 	
 }
 
@@ -648,9 +621,6 @@ void UClothMeshSelectionTool::ApplyAction(EClothMeshSelectionToolActions ActionT
 	case EClothMeshSelectionToolActions::ImportSecondaryFromCollection:
 		ImportFromCollection(/*bImportFromSecondarySet = */ true);
 		break;
-	case EClothMeshSelectionToolActions::TogglePrimarySecondary:
-		TogglePrimarySecondaryAction();
-		break;
 	case EClothMeshSelectionToolActions::GrowSelection:
 		GrowSelection();
 		break;
@@ -776,54 +746,6 @@ void UClothMeshSelectionTool::ImportFromCollection(bool bImportFromSecondarySet)
 	}
 }
 
-void UClothMeshSelectionTool::TogglePrimarySecondaryAction()
-{
-	// Save any changes made in the current mode to the node
-	if (bAnyChangeMade)
-	{
-		UpdateSelectedNode();
-		bAnyChangeMade = false;
-	}
-
-	// Toggle
-	ToolProperties->bSecondarySelection = !ToolProperties->bSecondarySelection;
-
-	// Re-initialize the Selection from the selected Dataflow node
-	FString ExistingSelectionName;
-	FGroupTopologySelection ExistingNodeSelection;
-	EChaosClothAssetSelectionOverrideType ExistingOverrideType;
-	GetSelectedNodeInfo(ExistingSelectionName, ExistingNodeSelection, ExistingOverrideType);
-
-	constexpr bool bBroadcastChange = false;
-	SelectionMechanic->SetSelection(ExistingNodeSelection, bBroadcastChange);
-
-	if (ExistingNodeSelection.SelectedCornerIDs.Num() > 0)
-	{
-		check(ExistingNodeSelection.SelectedEdgeIDs.Num() == 0);
-		check(ExistingNodeSelection.SelectedGroupIDs.Num() == 0);
-		SelectionMechanic->Properties->bSelectVertices = true;
-		SelectionMechanic->Properties->bSelectFaces = false;
-	}
-	else
-	{
-		SelectionMechanic->Properties->bSelectVertices = false;
-		SelectionMechanic->Properties->bSelectFaces = true;
-	}
-
-	ToolProperties->Name = ExistingSelectionName;
-	ToolProperties->SelectionOverrideType = ExistingOverrideType;
-
-	PreviewMesh->FastNotifySecondaryTrianglesChanged();
-
-	UpdatePrimarySecondaryMessage();
-}
-
-void UClothMeshSelectionTool::UpdatePrimarySecondaryMessage()
-{
-	const FText Message = ToolProperties->bSecondarySelection ? LOCTEXT("SecondarySelectionMode", "Secondary Selection") : LOCTEXT("PrimarySelectionMode", "Primary Selection");
-	GetToolManager()->DisplayMessage(Message, EToolMessageLevel::UserWarning);
-}
-
 void UClothMeshSelectionTool::GrowSelection()
 {
 	using namespace UE::Geometry;
diff --git a/Engine/Plugins/ChaosClothAssetEditor/Source/ChaosClothAssetEditorTools/Private/ChaosClothAsset/ClothMeshSelectionTool.h b/Engine/Plugins/ChaosClothAssetEditor/Source/ChaosClothAssetEditorTools/Private/ChaosClothAsset/ClothMeshSelectionTool.h
--- a/Engine/Plugins/ChaosClothAssetEditor/Source/ChaosClothAssetEditorTools/Private/ChaosClothAsset/ClothMeshSelectionTool.h
+++ b/Engine/Plugins/ChaosClothAssetEditor/Source/ChaosClothAssetEditorTools/Private/ChaosClothAsset/ClothMeshSelectionTool.h
@@ -25,7 +25,6 @@ enum class EClothMeshSelectionToolActions
 
 	ImportFromCollection,
 	ImportSecondaryFromCollection,
-	TogglePrimarySecondary,
 
 	GrowSelection,
 	ShrinkSelection,
@@ -57,24 +56,18 @@ public:
 
 	void PostAction(EClothMeshSelectionToolActions Action);
 
-	UFUNCTION(CallInEditor, Category = Operations)
+	UFUNCTION(CallInEditor, Category = Import)
 	void ImportFromCollection()
 	{
 		PostAction(EClothMeshSelectionToolActions::ImportFromCollection);
 	}
 
-	UFUNCTION(CallInEditor, Category = Operations)
+	UFUNCTION(CallInEditor, Category = Migrate)
 	void ImportSecondaryFromCollection()
 	{
 		PostAction(EClothMeshSelectionToolActions::ImportSecondaryFromCollection);
 	}
 
-	UFUNCTION(CallInEditor, Category = Operations)
-	void TogglePrimarySecondary()
-	{
-		PostAction(EClothMeshSelectionToolActions::TogglePrimarySecondary);
-	}
-
 	UFUNCTION(CallInEditor, Category = Selection)
 	void GrowSelection()
 	{
@@ -108,13 +101,6 @@ public:
 	UPROPERTY(EditAnywhere, Transient, Category = Selection, meta = (TransientToolProperty))
 	EChaosClothAssetSelectionOverrideType SelectionOverrideType;
 
-	/**
-	* Whether the user is editing the primary or secondary selection set. Stored as a property here so that the tool will remember what the
-	* user was doing the last time the tool shut down.
-	*/
-	UPROPERTY()
-	bool bSecondarySelection = false;
-
 	UPROPERTY(EditAnywhere, Category = Visualization, meta = (DisplayName = "Show Vertices"))
 	bool bShowVertices = false;
 	
@@ -193,8 +179,6 @@ private:
 	virtual void ApplyAction(EClothMeshSelectionToolActions ActionType);
 
 	void ImportFromCollection(bool bImportFromSecondarySet);
-	void TogglePrimarySecondaryAction();
-	void UpdatePrimarySecondaryMessage();
 
 	void GrowSelection();
 	void ShrinkSelection();
