Dataflow - Rework and tidy up Evaluation code.
Remove the Context timestamp, and use the cached output and node timestamps instead.
In the Dataflow editors, the terminal nodes' timestamp now replaces the context timestamp.

#rnx
#tests Ran evaluation in both Dataflow and Cloth Asset editors.

diff --git a/Engine/Plugins/ChaosClothAssetEditor/Source/ChaosClothAssetDataflowNodes/Private/ChaosClothAsset/ConnectableValueCustomization.cpp b/Engine/Plugins/ChaosClothAssetEditor/Source/ChaosClothAssetDataflowNodes/Private/ChaosClothAsset/ConnectableValueCustomization.cpp
--- a/Engine/Plugins/ChaosClothAssetEditor/Source/ChaosClothAssetDataflowNodes/Private/ChaosClothAsset/ConnectableValueCustomization.cpp
+++ b/Engine/Plugins/ChaosClothAssetEditor/Source/ChaosClothAssetDataflowNodes/Private/ChaosClothAsset/ConnectableValueCustomization.cpp
@@ -163,7 +163,7 @@ PRAGMA_ENABLE_DEPRECATION_WARNINGS
 									Text = *static_cast<FString*>(Data);  // Default value if the property isn't an input, or isn't connected
 									if (const FDataflowInput* const DataflowInput = DataflowNode->FindInput(Data))
 									{
-										Dataflow::FContextThreaded Context(FPlatformTime::Cycles64());
+										Dataflow::FContextThreaded Context;
 										Text = DataflowInput->GetValue<FString>(Context, Text);
 									}
 								}
diff --git a/Engine/Plugins/ChaosClothAssetEditor/Source/ChaosClothAssetDataflowNodes/Private/ChaosClothAsset/WeightMapNode.cpp b/Engine/Plugins/ChaosClothAssetEditor/Source/ChaosClothAssetDataflowNodes/Private/ChaosClothAsset/WeightMapNode.cpp
--- a/Engine/Plugins/ChaosClothAssetEditor/Source/ChaosClothAssetDataflowNodes/Private/ChaosClothAsset/WeightMapNode.cpp
+++ b/Engine/Plugins/ChaosClothAssetEditor/Source/ChaosClothAssetDataflowNodes/Private/ChaosClothAsset/WeightMapNode.cpp
@@ -54,7 +54,7 @@ void FChaosClothAssetWeightMapNode::OnTransfer()
 	using namespace UE::Chaos::ClothAsset;
 
 	// Create a temporary context (until we find a way to re-use the one from the tool calling this function)
-	Dataflow::FContextThreaded Context(FPlatformTime::Cycles64());
+	Dataflow::FContextThreaded Context;
 
 	// Transfer weight map if the transfer collection input has changed and is valid
 	FManagedArrayCollection InCollection = GetValue<FManagedArrayCollection>(Context, &Collection);
diff --git a/Engine/Plugins/ChaosClothAssetEditor/Source/ChaosClothAssetEditor/Private/ChaosClothAsset/ClothEditorToolkit.cpp b/Engine/Plugins/ChaosClothAssetEditor/Source/ChaosClothAssetEditor/Private/ChaosClothAsset/ClothEditorToolkit.cpp
--- a/Engine/Plugins/ChaosClothAssetEditor/Source/ChaosClothAssetEditor/Private/ChaosClothAsset/ClothEditorToolkit.cpp
+++ b/Engine/Plugins/ChaosClothAssetEditor/Source/ChaosClothAssetEditor/Private/ChaosClothAsset/ClothEditorToolkit.cpp
@@ -56,34 +56,6 @@ const FName FChaosClothAssetEditorToolkit::GraphCanvasTabId(TEXT("ChaosClothAsse
 const FName FChaosClothAssetEditorToolkit::NodeDetailsTabId(TEXT("ChaosClothAssetEditor_NodeDetails"));
 const FName FChaosClothAssetEditorToolkit::SimulationVisualizationTabID(TEXT("ChaosClothAssetEditor_SimulationVisualizationTab"));
 
-namespace Private
-{
-	UDataflow* GetDataflowFrom(const UObject* InObject)
-	{
-		if (const UClass* Class = InObject->GetClass())
-		{
-			if (const FProperty* Property = Class->FindPropertyByName(FName("DataflowAsset")))
-			{
-				return *Property->ContainerPtrToValuePtr<UDataflow*>(InObject);
-			}
-		}
-		return nullptr;
-	}
-
-	FString GetDataflowTerminalFrom(const UObject* InObject)
-	{
-		if (const UClass* Class = InObject->GetClass())
-		{
-			if (const FProperty* Property = Class->FindPropertyByName(FName("DataflowTerminal")))
-			{
-				return *Property->ContainerPtrToValuePtr<FString>(InObject);
-			}
-		}
-		return FString();
-	}
-}
-
-
 FChaosClothAssetEditorToolkit::FChaosClothAssetEditorToolkit(UAssetEditor* InOwningAssetEditor)
 	: FBaseCharacterFXEditorToolkit(InOwningAssetEditor, FName("ChaosClothAssetEditor"))
 {
@@ -255,9 +227,8 @@ void FChaosClothAssetEditorToolkit::Tick(float DeltaTime)
 	TickCommands.Broadcast();
 	TickCommands.Clear();
 
-	// Evaluate terminal node if one is assigned, but only if needed (don't force it)
-	constexpr bool bForceOperation = false;
-	EvaluateNode(nullptr, bForceOperation);
+	// Evaluate terminal node if one is assigned
+	EvaluateNode(nullptr, nullptr);
 
 	InvalidateViews();
 }
@@ -324,7 +295,7 @@ void FChaosClothAssetEditorToolkit::CreateWidgets()
 
 	if (ClothAsset)
 	{
-		UDataflow* const Dataflow = Private::GetDataflowFrom(ClothAsset);
+		UDataflow* const Dataflow = ClothAsset->GetDataflow();
 
 		// TODO: Figure out how to create the GraphEditor widgets when the ClothAsset doesn't have a Dataflow property set
 		if (Dataflow)
@@ -497,7 +468,7 @@ void FChaosClothAssetEditorToolkit::PostInitAssetEditor()
 
 	ClothMode->DataflowGraph = GetDataflow();
 	ClothMode->SetDataflowGraphEditor(GraphEditor);
-	DataflowContext = TSharedPtr<Dataflow::FEngineContext>(new Dataflow::FClothAssetDataflowContext(GetAsset(), GetDataflow(), Dataflow::FTimestamp::Invalid));
+	DataflowContext = TSharedPtr<Dataflow::FEngineContext>(new Dataflow::FClothAssetDataflowContext(GetAsset(), GetDataflow()));
 	ClothMode->SetDataflowContext(DataflowContext);
 
 	// Handle Dataflow asset reload event
@@ -517,8 +488,8 @@ void FChaosClothAssetEditorToolkit::GetSaveableObjects(TArray<UObject*>& OutObje
 {
 	FBaseCharacterFXEditorToolkit::GetSaveableObjects(OutObjects);
 
-	const UChaosClothAsset* const ClothAsset = GetAsset();
-	UDataflow* const DataflowAsset = Private::GetDataflowFrom(ClothAsset);
+	UChaosClothAsset* const ClothAsset = GetAsset();
+	UDataflow* const DataflowAsset = ClothAsset->GetDataflow();
 	if (DataflowAsset)
 	{
 		check(DataflowAsset->IsAsset());
@@ -534,7 +505,7 @@ void FChaosClothAssetEditorToolkit::OnAssetsSaved(const TArray<UObject*>& SavedO
 	const UChaosClothAsset* const ClothAsset = GetAsset();
 	check(ClothAsset);
 
-	if (UDataflow* const DataflowAsset = Private::GetDataflowFrom(ClothAsset))
+	if (const UDataflow* const DataflowAsset = ClothAsset->GetDataflow())
 	{
 		if (DataflowAsset->Dataflow)
 		{
@@ -845,10 +816,10 @@ void FChaosClothAssetEditorToolkit::OnFinishedChangingAssetProperties(const FPro
 
 	if (ChangedProperty && ChangedProperty->GetFName() == TEXT("DataflowAsset"))
 	{
-		const UChaosClothAsset* const ClothAsset = GetAsset();
+		UChaosClothAsset* const ClothAsset = GetAsset();
 		if (ClothAsset)
 		{
-			UDataflow* const Dataflow = Private::GetDataflowFrom(ClothAsset);
+			UDataflow* const Dataflow = ClothAsset->GetDataflow();
 
 			if (Dataflow)
 			{
@@ -871,115 +842,42 @@ void FChaosClothAssetEditorToolkit::OnFinishedChangingAssetProperties(const FPro
 	}
 }
 
-void FChaosClothAssetEditorToolkit::EvaluateNode(FDataflowNode* Node, bool bForceOperation)
+void FChaosClothAssetEditorToolkit::EvaluateNode(const FDataflowNode* Node, const FDataflowOutput* Output)
 {
+	check(DataflowContext);
+
 	UChaosClothAsset* const ClothAsset = GetAsset();
 	UDataflow* const Dataflow = GetDataflow(); 
 
 	if (Dataflow && ClothAsset)
 	{
-		const float NumSteps = 2.f;
+		constexpr float NumSteps = 1.f;
 		FScopedSlowTask SlowTask(NumSteps, LOCTEXT("ChaosClothAssetEditorToolkitEvaluateNode", "Evaluating nodes..."));
 		SlowTask.MakeDialogDelayed(1.f);
-
-		const TSharedPtr<Dataflow::FGraph> Graph = Dataflow->GetDataflow();
-		check(Graph);
-		
-		const FName DataflowTerminalName = *Private::GetDataflowTerminalFrom(ClothAsset);
-
-		if (!Node)
-		{
-			// Retrieve terminal node
-			if (const TSharedPtr<FDataflowNode> TerminalNode = Graph->FindBaseNode(DataflowTerminalName))
-			{
-				Node = TerminalNode->AsType<FDataflowTerminalNode>();
-			}
-			if (!Node)
-			{
-				return;  // No terminal node set
-			}
-		}
-
-		// Set evaluation context
-		check(DataflowContext);
-		if (bForceOperation)
-		{
-			LastDataflowNodeTimestamp = Dataflow::FTimestamp::Invalid;
-		}
-
-		Dataflow::FTimestamp OldTimestamp = LastDataflowNodeTimestamp;
-
-		// Set all asset values by visiting the graph from the specified node
-		TArray<FDataflowTerminalNode*> TerminalNodes;
-		TerminalNodes.Reserve(Graph->GetFilteredNodes(FDataflowTerminalNode::StaticType()).Num());
-		
-		TQueue<FDataflowNode*> Queue;
-		Queue.Enqueue(Node);
-
-		TSet<FDataflowNode*> VisitedNodes;
-
-		do
-		{
-			FDataflowNode* VisitedNode;
-			Queue.Dequeue(VisitedNode);
-			VisitedNodes.Add(VisitedNode);
-			if (VisitedNode->GetTimestamp() >= LastDataflowNodeTimestamp)
-			{
-				UE_LOG(LogChaosClothAssetEditor, VeryVerbose, TEXT("EvaluateNode - Visiting node %s"), *VisitedNode->GetName().ToString());
-
-				for (FDataflowInput* const Input : VisitedNode->GetInputs())
-				{
-					for (FDataflowOutput* const ConnectedOutput : Input->GetConnectedOutputs())
-					{
-						if (FDataflowNode* const OwningNode = ConnectedOutput->GetOwningNode())
-						{
-							if (!VisitedNodes.Contains(OwningNode))
-							{
-								Queue.Enqueue(OwningNode);
-							}
-
-							if (FDataflowTerminalNode* const TerminalNode = OwningNode->AsType<FDataflowTerminalNode>())
-							{
-								UE_LOG(LogChaosClothAssetEditor, VeryVerbose, TEXT("EvaluateNode - Found terminal node %s"), *TerminalNode->GetName().ToString());
-								TerminalNodes.Emplace(TerminalNode);
-							}
-						}
-					}
-				}
-			}
-		} while (!Queue.IsEmpty());
-
-		TSet<const FDataflowNode*> ProcessedTerminalNodes;
-		ProcessedTerminalNodes.Reserve(TerminalNodes.Num());  // Only cull the already visited nodes at this later stage, so that it always find nodes using the longest path in the graph
-
-		for (int32 TerminalNodeIndex = TerminalNodes.Num() - 1; TerminalNodeIndex >= 0; --TerminalNodeIndex)
-		{
-			FDataflowTerminalNode* const TerminalNode = TerminalNodes[TerminalNodeIndex];
-
-			bool bHasAlreadyBeenProcessed;
-			ProcessedTerminalNodes.FindOrAdd(TerminalNode, &bHasAlreadyBeenProcessed);
-
-			if (!bHasAlreadyBeenProcessed)
-			{
-				UE_LOG(LogChaosClothAssetEditor, VeryVerbose, TEXT("EvaluateNode - Setting asset value for terminal node %s"), *TerminalNodes[TerminalNodeIndex]->GetName().ToString());
-				TerminalNode->SetAssetValue(ClothAsset, *DataflowContext);
-			}
-		}
 		SlowTask.EnterProgressFrame(1.f);
 
-		// Evaluate the node and update LastDataflowNodeTimestamp if doing so (even if it isn't a terminal node despite the name of the command)
-		// If node is a terminal node, it will have its asset value updated prior to its output being evaluated
-		FDataflowEditorCommands::EvaluateTerminalNode(*DataflowContext, LastDataflowNodeTimestamp, Dataflow, Node, nullptr, ClothAsset);
-
-		SlowTask.EnterProgressFrame(1.f);
+		// Evaluate the node and update LastDataflowNodeTimestamp if doing so
+		const Dataflow::FTimestamp OldTimestamp = LastDataflowNodeTimestamp;
+		Node = FDataflowEditorCommands::EvaluateNode(
+			*DataflowContext,
+			LastDataflowNodeTimestamp,
+			Dataflow,
+			Node,
+			Output,
+			ClothAsset->GetDataflowTerminal(),
+			ClothAsset);
 
 		// Refresh editor
-		if (OldTimestamp.Value < LastDataflowNodeTimestamp.Value)
+		if (Node && OldTimestamp < LastDataflowNodeTimestamp)
 		{
-			OnClothAssetChanged();
+			if (Node->GetName() == ClothAsset->GetDataflowTerminal())
+			{
+				OnClothAssetChanged();
+			}
 
 			// Refresh the construction viewport
-			if (TSharedPtr<FDataflowNode> SelectedDataflowNode = GetSelectedDataflowNode())
+			TSharedPtr<FDataflowNode> SelectedDataflowNode = GetSelectedDataflowNode();
+			if (Node == SelectedDataflowNode.Get())
 			{
 				TSharedPtr<FManagedArrayCollection> InputCollection = GetInputClothCollectionIfPossible(SelectedDataflowNode, DataflowContext);
 				TSharedPtr<FManagedArrayCollection> Collection = GetClothCollectionIfPossible(SelectedDataflowNode, DataflowContext);
@@ -996,9 +894,9 @@ TSharedRef<SDataflowGraphEditor> FChaosClothAssetEditorToolkit::CreateGraphEdito
 	ensure(Dataflow);
 	using namespace Dataflow;
 
-	const auto EvalLambda = [this](FDataflowNode* Node, FDataflowOutput* /*Out*/)
+	const auto EvalLambda = [this](const FDataflowNode* Node, const FDataflowOutput* Output)
 	{
-		EvaluateNode(Node);
+		EvaluateNode(Node, Output);
 	};
 
 	SGraphEditor::FGraphEditorEvents InEvents;
@@ -1025,9 +923,9 @@ void FChaosClothAssetEditorToolkit::ReinitializeGraphEditorWidget()
 
 	ensure(Dataflow);
 
-	const auto EvalLambda = [this](FDataflowNode* Node, FDataflowOutput* /*Out*/)
+	const auto EvalLambda = [this](const FDataflowNode* Node, const FDataflowOutput* Output)
 	{
-		EvaluateNode(Node);
+		EvaluateNode(Node, Output);
 	};
 
 	SGraphEditor::FGraphEditorEvents InEvents;
@@ -1130,15 +1028,6 @@ TSharedPtr<FManagedArrayCollection> FChaosClothAssetEditorToolkit::GetClothColle
 		{
 			if (Output->GetType() == FName("FManagedArrayCollection"))
 			{
-				// Set asset value first if this node is a terminal node as the output might need the result of this operation
-				if (const FDataflowTerminalNode* const TerminalNode = InDataflowNode->AsType<const FDataflowTerminalNode>())
-				{
-					if (UChaosClothAsset* const ClothAsset = GetAsset())
-					{
-						TerminalNode->SetAssetValue(ClothAsset, *Context);
-					}
-				}
-
 				const FManagedArrayCollection DefaultValue;
 				TSharedRef<FManagedArrayCollection> Collection = MakeShared<FManagedArrayCollection>(Output->GetValue<FManagedArrayCollection>(*Context, DefaultValue));
 
diff --git a/Engine/Plugins/ChaosClothAssetEditor/Source/ChaosClothAssetEditor/Private/ChaosClothAsset/ClothEditorToolkit.h b/Engine/Plugins/ChaosClothAssetEditor/Source/ChaosClothAssetEditor/Private/ChaosClothAsset/ClothEditorToolkit.h
--- a/Engine/Plugins/ChaosClothAssetEditor/Source/ChaosClothAssetEditor/Private/ChaosClothAsset/ClothEditorToolkit.h
+++ b/Engine/Plugins/ChaosClothAssetEditor/Source/ChaosClothAssetEditor/Private/ChaosClothAsset/ClothEditorToolkit.h
@@ -31,8 +31,8 @@ namespace Dataflow
 	public:
 		DATAFLOW_CONTEXT_INTERNAL(TEngineContext<FContextSingle>, FClothAssetDataflowContext);
 
-		FClothAssetDataflowContext(UObject* InOwner, UDataflow* InGraph, FTimestamp InTimestamp)
-			: Super(InOwner, InTimestamp)
+		FClothAssetDataflowContext(UObject* InOwner, UDataflow* InGraph)
+			: Super(InOwner)
 		{}
 	};
 }
@@ -124,9 +124,7 @@ private:
 	void InvalidateViews();
 
 	// Dataflow
-	UE_DEPRECATED(5.4, "Use EvaluateNode(FDataflowNode*, bool) instead.")
-	void EvaluateNode(FDataflowNode* Node, FDataflowOutput* Out) { EvaluateNode(Node); }
-	void EvaluateNode(FDataflowNode* Node, bool bForceOperation = true);
+	void EvaluateNode(const FDataflowNode* Node, const FDataflowOutput* Output);
 	TSharedRef<SDataflowGraphEditor> CreateGraphEditorWidget();
 	void ReinitializeGraphEditorWidget();
 	TSharedPtr<IStructureDetailsView> CreateNodeDetailsEditorWidget(UObject* ObjectToEdit);
diff --git a/Engine/Plugins/Experimental/Dataflow/Source/DataflowEditor/Private/Dataflow/DataflowConstructionScene.cpp b/Engine/Plugins/Experimental/Dataflow/Source/DataflowEditor/Private/Dataflow/DataflowConstructionScene.cpp
--- a/Engine/Plugins/Experimental/Dataflow/Source/DataflowEditor/Private/Dataflow/DataflowConstructionScene.cpp
+++ b/Engine/Plugins/Experimental/Dataflow/Source/DataflowEditor/Private/Dataflow/DataflowConstructionScene.cpp
@@ -68,27 +68,32 @@ void FDataflowConstructionScene::AddReferencedObjects(FReferenceCollector& Colle
 	Collector.AddReferencedObjects(WireframeElements);
 }
 
-FORCEINLINE Dataflow::FTimestamp LatestTimestamp(const UDataflow* Dataflow, const ::Dataflow::FContext* Context)
-{
-	if (Dataflow && Context)
-	{
-		return FMath::Max(Dataflow->GetRenderingTimestamp().Value, Context->GetTimestamp().Value);
-	}
-	return ::Dataflow::FTimestamp::Invalid;
-}
-
 void FDataflowConstructionScene::TickDataflowScene(const float DeltaSeconds)
 {
 	if (const TObjectPtr<UDataflowBaseContent>& EditorContent = GetEditorContent())
 	{
-		if (const TSharedPtr<Dataflow::FContext> DataflowContext = EditorContent->GetDataflowContext())
+		if (const UDataflow* Dataflow = EditorContent->GetDataflowAsset())
 		{
-			if (const UDataflow* Dataflow = EditorContent->GetDataflowAsset())
+			if (Dataflow->GetDataflow())
 			{
-				const Dataflow::FTimestamp SystemTimestamp = DataflowContext->GetTimestamp().Value;
-				if (LastRenderedTimestamp < SystemTimestamp || EditorContent->IsConstructionDirty())
+				Dataflow::FTimestamp SystemTimestamp = Dataflow::FTimestamp::Invalid;
+				bool bMustUpdateConstructionScene = false;
+				for (TObjectPtr<const UDataflowBaseContent> DataflowBaseContent : GetTerminalContents())
+				{
+					const FName DataflowTerminalName(DataflowBaseContent->GetDataflowTerminal());
+					if (TSharedPtr<const FDataflowNode> DataflowTerminalNode = Dataflow->GetDataflow()->FindBaseNode(DataflowTerminalName))
+					{
+						SystemTimestamp = DataflowTerminalNode->GetTimestamp();
+					}
+
+					if (LastRenderedTimestamp < SystemTimestamp)
+					{
+						LastRenderedTimestamp = SystemTimestamp;
+						bMustUpdateConstructionScene = true;
+					}
+				}
+				if (bMustUpdateConstructionScene || EditorContent->IsConstructionDirty())
 				{
-					LastRenderedTimestamp = SystemTimestamp;
 					UpdateConstructionScene();
 				}
 			}
diff --git a/Engine/Plugins/Experimental/Dataflow/Source/DataflowEditor/Private/Dataflow/DataflowEditorCommands.cpp b/Engine/Plugins/Experimental/Dataflow/Source/DataflowEditor/Private/Dataflow/DataflowEditorCommands.cpp
--- a/Engine/Plugins/Experimental/Dataflow/Source/DataflowEditor/Private/Dataflow/DataflowEditorCommands.cpp
+++ b/Engine/Plugins/Experimental/Dataflow/Source/DataflowEditor/Private/Dataflow/DataflowEditorCommands.cpp
@@ -165,113 +165,34 @@ void FDataflowEditorCommands::Unregister()
 	return FDataflowEditorCommandsImpl::Unregister();
 }
 
-void FDataflowEditorCommands::EvaluateSelectedNodes(const FGraphPanelSelectionSet& SelectedNodes, FDataflowEditorCommands::FGraphEvaluationCallback Evaluate)
+const FDataflowNode* FDataflowEditorCommands::EvaluateNode(Dataflow::FContext& Context, Dataflow::FTimestamp& InOutLastNodeTimestamp,
+	const UDataflow* Dataflow, const FDataflowNode* Node, const FDataflowOutput* Output, const FString& NodeName, UObject* Asset)
 {
-	const bool bIsInPIEOrSimulate = GEditor->PlayWorld != NULL || GEditor->bIsSimulatingInEditor;
-	if (!bIsInPIEOrSimulate)
-	{
-		for (UObject* Node : SelectedNodes)
-		{
-			if (UDataflowEdNode* EdNode = dynamic_cast<UDataflowEdNode*>(Node))
-			{
-				if (const TSharedPtr<Dataflow::FGraph> DataflowGraph = EdNode->GetDataflowGraph())
-				{
-					if (const TSharedPtr<FDataflowNode> DataflowNode = DataflowGraph->FindBaseNode(EdNode->GetDataflowNodeGuid()))
-					{
-						if (DataflowNode->bActive)
-						{
-							if (DataflowNode->GetOutputs().Num())
-							{
-								for (FDataflowConnection* NodeOutput : DataflowNode->GetOutputs())
-								{
-									Evaluate(DataflowNode.Get(), (FDataflowOutput*)NodeOutput);
-								}
-							}
-							else
-							{
-								Evaluate(DataflowNode.Get(), nullptr);
-							}
-						}
-					}
-				}
-			}
-		}
-	}
-}
+	UE_LOG(LogChaosDataflow, Verbose, TEXT("FDataflowEditorCommands::EvaluateNode(): Node [%s], NodeName [%s] Output [%s]"), Node ? *Node->GetName().ToString() : TEXT("nullptr"), *NodeName, Output ? *Output->GetName().ToString() : TEXT("nullptr"));
 
-void FDataflowEditorCommands::EvaluateNode(Dataflow::FContext& Context, Dataflow::FTimestamp& OutLastNodeTimestamp,
-	const UDataflow* Dataflow, const FDataflowNode* InNode, const FDataflowOutput* Output, FString NodeName)
-{
-	const bool bIsInPIEOrSimulate = GEditor->PlayWorld != NULL || GEditor->bIsSimulatingInEditor;
-	if (!bIsInPIEOrSimulate)
+	if (!Node && Dataflow)
 	{
-		if (Dataflow)
+		if (const TSharedPtr<const Dataflow::FGraph> Graph = Dataflow->GetDataflow())
 		{
-			const FDataflowNode* Node = InNode;
-			if (Node == nullptr)
-			{
-				if (const TSharedPtr<const Dataflow::FGraph> Graph = Dataflow->GetDataflow())
-				{
-					if (TSharedPtr<const FDataflowNode> GraphNode = Graph->FindBaseNode(FName(NodeName)))
-					{
-						Node = GraphNode.Get();
-					}
-				}
-			}
-
-			if (Node)
-			{
-				if (Context.GetTimestamp() < Node->GetTimestamp())
-				{
-					Context.Evaluate(Node, Output);
-					OutLastNodeTimestamp = Context.GetTimestamp();
-				}
-			}
+			Node = Graph->FindBaseNode(FName(NodeName)).Get();
 		}
 	}
-}
-
-void FDataflowEditorCommands::EvaluateTerminalNode(Dataflow::FContext& Context, Dataflow::FTimestamp& OutLastNodeTimestamp,
-	const UDataflow* Dataflow, const FDataflowNode* InNode, const FDataflowOutput* Output, UObject* InAsset, FString NodeName)
-{
-	const bool bIsInPIEOrSimulate = GEditor->PlayWorld != NULL || GEditor->bIsSimulatingInEditor;
-	if (!bIsInPIEOrSimulate)
+	if (Node && InOutLastNodeTimestamp < Node->GetTimestamp())
 	{
-		if (Dataflow)
-		{
-			const FDataflowNode* Node = InNode;
-			if (Node == nullptr)
-			{
-				if (const TSharedPtr<const Dataflow::FGraph> Graph = Dataflow->GetDataflow())
-				{
-					if (TSharedPtr<const FDataflowNode> GraphNode = Graph->FindBaseNode(FName(NodeName)))
-					{
-						Node = GraphNode.Get();
-					}
-				}
-			}
+		Context.Evaluate(Node, Output);
 
-			if (Node)
+		if (Asset)
+		{
+			if (const FDataflowTerminalNode* const TerminalNode = Node->AsType<const FDataflowTerminalNode>())
 			{
-				if (Context.GetTimestamp() < Node->GetTimestamp())
-				{
-					if (const FDataflowTerminalNode* TerminalNode = Node->AsType<const FDataflowTerminalNode>())
-					{
-						if (InAsset)
-						{
-							TerminalNode->SetAssetValue(InAsset, Context);  // Kriss.Gossart: Must set asset value before call to Evaluate
-																			// This dependency is neccessary for now for the cloth editor code.
-																			// It is order dependant, with some nodes relying on the non const 
-																			// SetAssetValue to be called first in order to have a valid const Evaluation.
-						}
-					}
-
-					Context.Evaluate(Node, Output);
-					OutLastNodeTimestamp = Context.GetTimestamp();
-				}
+				UE_LOG(LogChaosDataflow, Verbose, TEXT("FDataflowTerminalNode::SetAssetValue(): TerminalNode [%s], Asset [%s]"), *TerminalNode->GetName().ToString(), *Asset->GetName());
+				TerminalNode->SetAssetValue(Asset, Context);
 			}
 		}
+
+		InOutLastNodeTimestamp = Node->GetTimestamp();
 	}
+	return Node;
 }
 
 bool FDataflowEditorCommands::OnNodeVerifyTitleCommit(const FText& NewText, UEdGraphNode* GraphNode, FText& OutErrorMessage)
diff --git a/Engine/Plugins/Experimental/Dataflow/Source/DataflowEditor/Private/Dataflow/DataflowEditorToolkit.cpp b/Engine/Plugins/Experimental/Dataflow/Source/DataflowEditor/Private/Dataflow/DataflowEditorToolkit.cpp
--- a/Engine/Plugins/Experimental/Dataflow/Source/DataflowEditor/Private/Dataflow/DataflowEditorToolkit.cpp
+++ b/Engine/Plugins/Experimental/Dataflow/Source/DataflowEditor/Private/Dataflow/DataflowEditorToolkit.cpp
@@ -1059,36 +1059,45 @@ void FDataflowEditorToolkit::Tick(float DeltaTime)
 	{
 		if (EditorContent->GetDataflowAsset())
 		{
-			Dataflow::FTimestamp TimeStamp = EditorContent->GetLastModifiedTimestamp();
+			Dataflow::FTimestamp InitTimeStamp = EditorContent->GetLastModifiedTimestamp();
 			if (!EditorContent->GetDataflowContext())
 			{
-				EditorContent->SetDataflowContext(MakeShared<Dataflow::FEngineContext>(EditorContent->GetDataflowOwner(), Dataflow::FTimestamp::Invalid));
-				TimeStamp = Dataflow::FTimestamp::Invalid;
+				EditorContent->SetDataflowContext(MakeShared<Dataflow::FEngineContext>(EditorContent->GetDataflowOwner()));
+				InitTimeStamp = Dataflow::FTimestamp::Invalid;
 			}
 
 			// Update the list of dataflow terminal contents 
-			DataflowEditor->UpdateTerminalContents(TimeStamp);
-			
+			DataflowEditor->UpdateTerminalContents(InitTimeStamp);
+
 			// OnTick evaluation only pulls the terminal nodes. The other evaluations can be specific nodes.
 			// We only evaluate multiple terminal nodes if the dataflow owner is a UDataflow (Owner == Asset)
-			if(!GetTerminalContents().IsEmpty() && (EditorContent->GetDataflowOwner() == EditorContent->GetDataflowAsset()) && EditorContent->GetDataflowContext())
+			if (!GetTerminalContents().IsEmpty() && (EditorContent->GetDataflowOwner() == EditorContent->GetDataflowAsset()) && EditorContent->GetDataflowContext())
 			{
-				const Dataflow::FTimestamp InitTimeStamp = TimeStamp;
-				for(const TObjectPtr<UDataflowBaseContent>& TerminalContent : GetTerminalContents())
+				for (const TObjectPtr<UDataflowBaseContent>& TerminalContent : GetTerminalContents())
 				{
-					Dataflow::FTimestamp TerminalTimeStamp = InitTimeStamp;
-					FDataflowEditorCommands::EvaluateTerminalNode(*EditorContent->GetDataflowContext().Get(), TerminalTimeStamp, EditorContent->GetDataflowAsset(),
-						nullptr, nullptr, TerminalContent->GetTerminalAsset(), TerminalContent->GetDataflowTerminal());
+					if (const UDataflow* const Dataflow = EditorContent->GetDataflowAsset())
+					{
+						if (const TSharedPtr<const Dataflow::FGraph> Graph = Dataflow->GetDataflow())
+						{
+							const FName TerminalNodeName(TerminalContent->GetDataflowTerminal());
+							const FDataflowNode* const Node = Graph->FindBaseNode(TerminalNodeName).Get();
 
-					TimeStamp = FMath::Max(TimeStamp, TerminalTimeStamp);
+							Dataflow::FTimestamp TerminalNodeTimeStamp = InitTimeStamp;
+							EvaluateNode(Node, nullptr, TerminalNodeTimeStamp);  // When Node is null, EvaluateNode falls back on the EditorContent terminal node
+
+							// Take the Max of the existing time stamp, as other terminal nodes might have more recent invalidations
+							const Dataflow::FTimestamp LastModifiedTimestamp = FMath::Max(EditorContent->GetLastModifiedTimestamp(), TerminalNodeTimeStamp);
+							EditorContent->SetLastModifiedTimestamp(LastModifiedTimestamp);
+						}
+					}
 				}
 			}
 			else
 			{
-				FDataflowEditorCommands::EvaluateTerminalNode(*EditorContent->GetDataflowContext().Get(), TimeStamp, EditorContent->GetDataflowAsset(),
-					nullptr, nullptr, EditorContent->GetTerminalAsset(), EditorContent->GetDataflowTerminal());
+				Dataflow::FTimestamp TerminalNodeTimeStamp = InitTimeStamp;
+				EvaluateNode(nullptr, nullptr, TerminalNodeTimeStamp);
+				EditorContent->SetLastModifiedTimestamp(TerminalNodeTimeStamp);
 			}
-			EditorContent->SetLastModifiedTimestamp(TimeStamp);
 
 			// Ensure the context object's selected node matches the selected node in the graph editor
 			// TODO: Create an Editor Context Object that can just hold a reference to the graph editor, rather than keeping these in sync
@@ -1108,27 +1117,42 @@ TStatId FDataflowEditorToolkit::GetStatId() const
 	RETURN_QUICK_DECLARE_CYCLE_STAT(FDataflowEditorToolkit, STATGROUP_Tickables);
 }
 
-TSharedRef<SDataflowGraphEditor> FDataflowEditorToolkit::CreateGraphEditorWidget(UDataflow* DataflowToEdit, TSharedPtr<IStructureDetailsView> InNodeDetailsEditor)
+void FDataflowEditorToolkit::EvaluateNode(const FDataflowNode* Node, const FDataflowOutput* Output, Dataflow::FTimestamp& InOutTimestamp)
 {
-	ensure(DataflowToEdit);
-	using namespace Dataflow;
+	UE_LOG(LogChaosDataflow, Verbose, TEXT("FDataflowEditorToolkit::EvaluateNode(): Node [%s], Output [%s]"), Node ? *Node->GetName().ToString() : TEXT("nullptr"), Output ? *Output->GetName().ToString() : TEXT("nullptr"));
 
-	FDataflowEditorCommands::FGraphEvaluationCallback Evaluate = [&](FDataflowNode* Node, FDataflowOutput* Out)
+	const bool bIsInPIEOrSimulate = GEditor->PlayWorld || GEditor->bIsSimulatingInEditor;
+	if (!bIsInPIEOrSimulate)  // TODO: make this test optional for some toolkit
 	{
 		if (const TObjectPtr<UDataflowBaseContent>& EditorContent = GetEditorContent())
 		{
 			if (EditorContent->GetDataflowAsset())
 			{
-				Node->Invalidate();
-				Dataflow::FTimestamp TimeStamp = Dataflow::FTimestamp::Invalid;
-				
-				FDataflowEditorCommands::EvaluateTerminalNode(*EditorContent->GetDataflowContext().Get(), TimeStamp, EditorContent->GetDataflowAsset(),
-					Node, Out, EditorContent->GetTerminalAsset(), EditorContent->GetDataflowTerminal());
-
-				EditorContent->SetLastModifiedTimestamp(TimeStamp);
+				// If Node is null, the terminal node with the given name will be used instead
+				FDataflowEditorCommands::EvaluateNode(*EditorContent->GetDataflowContext().Get(), InOutTimestamp, EditorContent->GetDataflowAsset(),
+					Node, Output, EditorContent->GetDataflowTerminal(), EditorContent->GetTerminalAsset());
 			}
 		}
-	};
+	}
+}
+
+TSharedRef<SDataflowGraphEditor> FDataflowEditorToolkit::CreateGraphEditorWidget(UDataflow* DataflowToEdit, TSharedPtr<IStructureDetailsView> InNodeDetailsEditor)
+{
+	ensure(DataflowToEdit);
+	using namespace Dataflow;
+
+	const FDataflowEditorCommands::FGraphEvaluationCallback Evaluate =
+		[this](const FDataflowNode* Node, const FDataflowOutput* Output)
+		{
+			if (const TObjectPtr<UDataflowBaseContent>& EditorContent = GetEditorContent())
+			{
+				Dataflow::FTimestamp LastNodeTimestamp = EditorContent->GetLastModifiedTimestamp();
+
+				EvaluateNode(Node, Output, LastNodeTimestamp);
+
+				EditorContent->SetLastModifiedTimestamp(LastNodeTimestamp);
+			}
+		};
 	
 	DataflowEditor->UpdateTerminalContents(Dataflow::FTimestamp::Invalid);
 	
diff --git a/Engine/Plugins/Experimental/Dataflow/Source/DataflowEditor/Private/Dataflow/DataflowEditorUtil.cpp b/Engine/Plugins/Experimental/Dataflow/Source/DataflowEditor/Private/Dataflow/DataflowEditorUtil.cpp
--- a/Engine/Plugins/Experimental/Dataflow/Source/DataflowEditor/Private/Dataflow/DataflowEditorUtil.cpp
+++ b/Engine/Plugins/Experimental/Dataflow/Source/DataflowEditor/Private/Dataflow/DataflowEditorUtil.cpp
@@ -123,13 +123,13 @@ namespace Dataflow
 		{
 			if (!Content->GetDataflowContext())
 			{
-				Content->SetDataflowContext(MakeShared<FEngineContext>(Content->GetDataflowOwner(), FTimestamp::Invalid));
+				Content->SetDataflowContext(MakeShared<FEngineContext>(Content->GetDataflowOwner()));
 			}
 			return Content->GetDataflowContext();
 		}
 
 		ensure(false);
-		return MakeShared<FEngineContext>(nullptr, FTimestamp::Invalid);
+		return MakeShared<FEngineContext>(nullptr);
 	}
 
 	bool CanRenderNodeOutput(const UDataflowEdNode& EdNode, const UDataflowBaseContent& EditorContent, const IDataflowConstructionViewMode& ViewMode)
diff --git a/Engine/Plugins/Experimental/Dataflow/Source/DataflowEditor/Private/Dataflow/DataflowGraphEditor.cpp b/Engine/Plugins/Experimental/Dataflow/Source/DataflowEditor/Private/Dataflow/DataflowGraphEditor.cpp
--- a/Engine/Plugins/Experimental/Dataflow/Source/DataflowEditor/Private/Dataflow/DataflowGraphEditor.cpp
+++ b/Engine/Plugins/Experimental/Dataflow/Source/DataflowEditor/Private/Dataflow/DataflowGraphEditor.cpp
@@ -179,19 +179,42 @@ const TSharedPtr<Dataflow::FEngineContext> SDataflowGraphEditor::GetDataflowCont
 
 void SDataflowGraphEditor::EvaluateNode()
 {
-	if (EvaluateGraphCallback)
-	{
-		FDataflowEditorCommands::EvaluateSelectedNodes(GetSelectedNodes(), EvaluateGraphCallback);
-	}
-	else
+	UE_LOG(LogChaosDataflow, Verbose, TEXT("SDataflowGraphEditor::EvaluateNode(): Nodes [%s]"),
+		*FString::JoinBy(GetSelectedNodes().Array(), TEXT(", "), [](const UObject* SelectedNode)
+			{
+				return Cast<UDataflowEdNode>(SelectedNode) && Cast<UDataflowEdNode>(SelectedNode)->GetDataflowNode() ? 
+					Cast<UDataflowEdNode>(SelectedNode)->GetDataflowNode()->GetName().ToString() : FString();
+			}));
+
+	using namespace Dataflow;
+
+	TOptional<FContextThreaded> DefaultContext(EvaluateGraphCallback ? TOptional<FContextThreaded>() : TOptional<FContextThreaded>(FContextThreaded()));
+
+	for (UObject* Node : GetSelectedNodes())
 	{
-		FDataflowEditorCommands::FGraphEvaluationCallback LocalEvaluateCallback = [](FDataflowNode* Node, FDataflowOutput* Out)
+		if (UDataflowEdNode* const EdNode = Cast<UDataflowEdNode>(Node))
 		{
-			using namespace Dataflow;
-			FContextThreaded(FPlatformTime::Cycles64()).Evaluate(Node, Out);
-		};
-
-		FDataflowEditorCommands::EvaluateSelectedNodes(GetSelectedNodes(), LocalEvaluateCallback);
+			if (const TSharedPtr<FGraph> DataflowGraph = EdNode->GetDataflowGraph())
+			{
+				if (const TSharedPtr<FDataflowNode> DataflowNode = DataflowGraph->FindBaseNode(EdNode->GetDataflowNodeGuid()))
+				{
+					if (DataflowNode->bActive)
+					{
+						DataflowNode->Invalidate();  // Force evaluation
+
+						if (EvaluateGraphCallback)
+						{
+							EvaluateGraphCallback(DataflowNode.Get(), nullptr);  // Evaluation processes all outputs when passing a null Output
+						}
+						else
+						{
+							check(DefaultContext);
+							DefaultContext->Evaluate(DataflowNode.Get(), nullptr);
+						}
+					}
+				}
+			}
+		}
 	}
 }
 
diff --git a/Engine/Plugins/Experimental/Dataflow/Source/DataflowEditor/Public/Dataflow/DataflowEditorCommands.h b/Engine/Plugins/Experimental/Dataflow/Source/DataflowEditor/Public/Dataflow/DataflowEditorCommands.h
--- a/Engine/Plugins/Experimental/Dataflow/Source/DataflowEditor/Public/Dataflow/DataflowEditorCommands.h
+++ b/Engine/Plugins/Experimental/Dataflow/Source/DataflowEditor/Public/Dataflow/DataflowEditorCommands.h
@@ -99,7 +99,7 @@ public:
 class DATAFLOWEDITOR_API FDataflowEditorCommands
 {
 public:
-	typedef TFunction<void(FDataflowNode*, FDataflowOutput*)> FGraphEvaluationCallback;
+	typedef TFunction<void(const FDataflowNode*, const FDataflowOutput*)> FGraphEvaluationCallback;
 	typedef TFunction<void(const FGeometry& MyGeometry, const FDragDropEvent& DragDropEvent)> FOnDragDropEventCallback;
 
 	static void Register();
@@ -108,20 +108,20 @@ public:
 	static const FDataflowEditorCommandsImpl& Get();
 
 	/*
-	*  EvaluateSelectedNodes
+	* Node evaluation utility function.
+	* 
+	* @param Context  The evaluation context.
+	* @param InOutLastNodeTimestamp  The last evaluation time used to trigger the evaluation when the node's timestamp is more recent than this value. If the node is evaluated, the value also gets updated with the evaluated node's current timestamp.
+	* @param Dataflow  The dataflow asset used to search for the NodeName when Node is nullptr.
+	* @param Node  The node to evaluate. When null, a node with the given NodeName will be evaluated instead if it exists.
+	* @param Output  The node's output to evaluate. When no output are specified, all outputs will be evaluated.
+	* @param NodeName  When no node is specified, then the node will be searched instead within the Dataflow's graph using NodeName, otherwise NodeName is ignored.
+	* @param Asset  When Asset is non null, if the node is a terminal node, and if the node timestamp is more recent than InOutLastNodeTimestamp. then the node SetAssetValue method will be called on this asset.
+	* @return  The node that has been evaluated if any.
 	*/
-	static void EvaluateSelectedNodes(const FGraphPanelSelectionSet& SelectedNodes, FGraphEvaluationCallback);
-
-	/*
-	* EvaluateGraph
-	*/
-	static void EvaluateNode(Dataflow::FContext& Context, Dataflow::FTimestamp& OutLastNodeTimestamp,
-		const UDataflow* Dataflow, const FDataflowNode* Node = nullptr, const FDataflowOutput* Out = nullptr, 
-		FString NodeName = FString()); // @todo(Dataflow) deprecate  
-
-	static void EvaluateTerminalNode(Dataflow::FContext& Context, Dataflow::FTimestamp& OutLastNodeTimestamp,
-		const UDataflow* Dataflow, const FDataflowNode* Node = nullptr, const FDataflowOutput* Out = nullptr,
-		UObject* InAsset = nullptr, FString NodeName = FString());
+	static const FDataflowNode* EvaluateNode(Dataflow::FContext& Context, Dataflow::FTimestamp& InOutLastNodeTimestamp,
+		const UDataflow* Dataflow, const FDataflowNode* Node, const FDataflowOutput* Output = nullptr, 
+		const FString& NodeName = FString(), UObject* Asset = nullptr);
 
 	/*
 	*  DeleteNodes
diff --git a/Engine/Plugins/Experimental/Dataflow/Source/DataflowEditor/Public/Dataflow/DataflowEditorToolkit.h b/Engine/Plugins/Experimental/Dataflow/Source/DataflowEditor/Public/Dataflow/DataflowEditorToolkit.h
--- a/Engine/Plugins/Experimental/Dataflow/Source/DataflowEditor/Public/Dataflow/DataflowEditorToolkit.h
+++ b/Engine/Plugins/Experimental/Dataflow/Source/DataflowEditor/Public/Dataflow/DataflowEditorToolkit.h
@@ -91,6 +91,9 @@ protected:
 	void OnConstructionViewSelectionChanged(const TArray<UPrimitiveComponent*>& SelectedComponents);
 	// Callback to remove the closed one from the listener views
 	void OnTabClosed(TSharedRef<SDockTab> Tab);
+
+	// Node evaluation
+	void EvaluateNode(const FDataflowNode* Node, const FDataflowOutput* Output, Dataflow::FTimestamp& InOutTimestamp);
 	
 private:
 	
diff --git a/Engine/Plugins/Experimental/Dataflow/Source/DataflowSimulation/Private/Dataflow/DataflowSimulationManager.cpp b/Engine/Plugins/Experimental/Dataflow/Source/DataflowSimulation/Private/Dataflow/DataflowSimulationManager.cpp
--- a/Engine/Plugins/Experimental/Dataflow/Source/DataflowSimulation/Private/Dataflow/DataflowSimulationManager.cpp
+++ b/Engine/Plugins/Experimental/Dataflow/Source/DataflowSimulation/Private/Dataflow/DataflowSimulationManager.cpp
@@ -449,8 +449,7 @@ void UDataflowSimulationManager::AddSimulationInterface(IDataflowSimulationInter
 			Dataflow::FDataflowSimulationData& DataflowData = SimulationData.FindOrAdd(DataflowAsset);
 			if(!DataflowData.SimulationContext.IsValid())
 			{
-				DataflowData.SimulationContext = MakeShared<Dataflow::FDataflowSimulationContext>(
-					DataflowAsset, Dataflow::FTimestamp::Invalid);
+				DataflowData.SimulationContext = MakeShared<Dataflow::FDataflowSimulationContext>(DataflowAsset);
 			}
 			DataflowData.SimulationInterfaces.FindOrAdd(SimulationInterface->GetSimulationType()).Add(SimulationInterface);
 		}
diff --git a/Engine/Plugins/Experimental/Dataflow/Source/DataflowSimulation/Public/Dataflow/DataflowSimulationContext.h b/Engine/Plugins/Experimental/Dataflow/Source/DataflowSimulation/Public/Dataflow/DataflowSimulationContext.h
--- a/Engine/Plugins/Experimental/Dataflow/Source/DataflowSimulation/Public/Dataflow/DataflowSimulationContext.h
+++ b/Engine/Plugins/Experimental/Dataflow/Source/DataflowSimulation/Public/Dataflow/DataflowSimulationContext.h
@@ -17,8 +17,8 @@ namespace Dataflow
 	public:
 		DATAFLOW_CONTEXT_INTERNAL(TEngineContext<Base>, TSimulationContext);
 
-		TSimulationContext(const TObjectPtr<UObject>& InOwner, FTimestamp InTimestamp)
-				: Super(InOwner, InTimestamp)
+		explicit TSimulationContext(const TObjectPtr<UObject>& InOwner)
+				: Super(InOwner)
 		{}
 		
 		virtual ~TSimulationContext() override {};
diff --git a/Engine/Source/Runtime/Experimental/Chaos/Private/ChaosLog.cpp b/Engine/Source/Runtime/Experimental/Chaos/Private/ChaosLog.cpp
--- a/Engine/Source/Runtime/Experimental/Chaos/Private/ChaosLog.cpp
+++ b/Engine/Source/Runtime/Experimental/Chaos/Private/ChaosLog.cpp
@@ -11,6 +11,7 @@ DEFINE_LOG_CATEGORY(LogChaosThread);
 DEFINE_LOG_CATEGORY(LogChaosSimulation);
 DEFINE_LOG_CATEGORY(LogChaosDebug);
 DEFINE_LOG_CATEGORY(LogChaos);
+DEFINE_LOG_CATEGORY(LogChaosDataflow);
 
 CSV_DEFINE_CATEGORY_MODULE(CHAOS_API, Chaos, true);
 CSV_DEFINE_CATEGORY_MODULE(CHAOS_API, PhysicsVerbose, false);
diff --git a/Engine/Source/Runtime/Experimental/Chaos/Public/ChaosLog.h b/Engine/Source/Runtime/Experimental/Chaos/Public/ChaosLog.h
--- a/Engine/Source/Runtime/Experimental/Chaos/Public/ChaosLog.h
+++ b/Engine/Source/Runtime/Experimental/Chaos/Public/ChaosLog.h
@@ -9,6 +9,7 @@ CHAOS_API DECLARE_LOG_CATEGORY_EXTERN(LogChaosGeneral, Log, All);
 CHAOS_API DECLARE_LOG_CATEGORY_EXTERN(LogChaosThread, Log, All);
 CHAOS_API DECLARE_LOG_CATEGORY_EXTERN(LogChaosSimulation, Log, All);
 CHAOS_API DECLARE_LOG_CATEGORY_EXTERN(LogChaosDebug, Log, All);
+CHAOS_API DECLARE_LOG_CATEGORY_EXTERN(LogChaosDataflow, Log, All);
 
 CHAOS_API DECLARE_LOG_CATEGORY_EXTERN(LogChaos, Log, All);
 
diff --git a/Engine/Source/Runtime/Experimental/Dataflow/Core/Private/Dataflow/DataflowConnection.cpp b/Engine/Source/Runtime/Experimental/Dataflow/Core/Private/Dataflow/DataflowConnection.cpp
--- a/Engine/Source/Runtime/Experimental/Dataflow/Core/Private/Dataflow/DataflowConnection.cpp
+++ b/Engine/Source/Runtime/Experimental/Dataflow/Core/Private/Dataflow/DataflowConnection.cpp
@@ -62,6 +62,11 @@ FGuid FDataflowConnection::GetOwningNodeGuid() const
 	return OwningNode ? OwningNode->GetGuid() : FGuid();
 }
 
+Dataflow::FTimestamp FDataflowConnection::GetOwningNodeTimestamp() const
+{
+	return OwningNode ? OwningNode->GetTimestamp() : Dataflow::FTimestamp::Invalid;
+}
+
 uint32 FDataflowConnection::GetOwningNodeValueHash() const
 {
 	return OwningNode ? OwningNode->GetValueHash() : 0;
diff --git a/Engine/Source/Runtime/Experimental/Dataflow/Core/Private/Dataflow/DataflowInputOutput.cpp b/Engine/Source/Runtime/Experimental/Dataflow/Core/Private/Dataflow/DataflowInputOutput.cpp
--- a/Engine/Source/Runtime/Experimental/Dataflow/Core/Private/Dataflow/DataflowInputOutput.cpp
+++ b/Engine/Source/Runtime/Experimental/Dataflow/Core/Private/Dataflow/DataflowInputOutput.cpp
@@ -250,6 +250,8 @@ bool FDataflowOutput::Evaluate(Dataflow::FContext& Context) const
 
 bool FDataflowOutput::EvaluateImpl(Dataflow::FContext& Context) const
 {
+	UE_LOG(LogChaosDataflow, Verbose, TEXT("FDataflowOutput::EvaluateImpl(): Node [%s], Output [%s]"), *GetOwningNode()->GetName().ToString(), *GetName().ToString());
+
 	Dataflow::FContextScopedCallstack Callstack(Context, this);
 	if (Callstack.IsLoopDetected())
 	{ 
@@ -260,9 +262,11 @@ bool FDataflowOutput::EvaluateImpl(Dataflow::FContext& Context) const
 	// check if the cache has a valid version
 	if(Context.HasData(CacheKey(), OwningNode->LastModifiedTimestamp))
 	{
+		UE_LOG(LogChaosDataflow, Verbose, TEXT("FDataflowOutput::EvaluateImpl(): Context has data, NodeTimestamp [%ul], CacheTimestamp [%lu]"), OwningNode->LastModifiedTimestamp.Value, Context.GetTimestamp(CacheKey()).Value);
 		return true;
 	}
 	// if not, evaluate
+	UE_LOG(LogChaosDataflow, Verbose, TEXT("FDataflowNode::Evaluate(): Node [%s], Output [%s], NodeTimestamp [%lu]"), *GetOwningNode()->GetName().ToString(), *GetName().ToString(), OwningNode->LastModifiedTimestamp.Value);
 	OwningNode->Evaluate(Context, this);
 	// Validation
 	if (!Context.HasData(CacheKey()))
diff --git a/Engine/Source/Runtime/Experimental/Dataflow/Core/Private/Dataflow/DataflowNodeParameters.cpp b/Engine/Source/Runtime/Experimental/Dataflow/Core/Private/Dataflow/DataflowNodeParameters.cpp
--- a/Engine/Source/Runtime/Experimental/Dataflow/Core/Private/Dataflow/DataflowNodeParameters.cpp
+++ b/Engine/Source/Runtime/Experimental/Dataflow/Core/Private/Dataflow/DataflowNodeParameters.cpp
@@ -14,6 +14,15 @@ namespace Dataflow
 	uint64 FTimestamp::Invalid = 0;
 	uint64 FTimestamp::Current() { return FPlatformTime::Cycles64(); }
 
+	FTimestamp FContext::GetTimestamp(FContextCacheKey Key) const
+	{
+		if (const TUniquePtr<FContextCacheElementBase>* const Cache = const_cast<FContext*>(this)->GetDataImpl(Key))
+		{
+			return (*Cache)->GetTimestamp();
+		}
+		return FTimestamp::Invalid;
+	}
+
 	void FContext::PushToCallstack(const FDataflowConnection* Connection)
 	{
 #if DATAFLOW_EDITOR_EVALUATION
@@ -53,28 +62,32 @@ namespace Dataflow
 
 	void BeginContextEvaluation(FContext& Context, const FDataflowNode* Node, const FDataflowOutput* Output)
 	{
-		if (Node != nullptr)
+		if (Output)
+		{
+			Context.Evaluate(*Output);
+		}
+		else if (Node)
 		{
-			Context.Timestamp = FTimestamp(FPlatformTime::Cycles64());
 			if (Node->NumOutputs())
 			{
-				if (Output)
+				for (const FDataflowOutput* const NodeOutput : Node->GetOutputs())
 				{
-					Node->Evaluate(Context, Output);
-				}
-				else
-				{
-					for (FDataflowOutput* NodeOutput : Node->GetOutputs())
-					{
-						Node->Evaluate(Context, NodeOutput);
-					}
+					Context.Evaluate(*NodeOutput);
 				}
 			}
 			else
 			{
+				// TODO: When no outputs are specified, this call to Evaluate should really be removed.
+				//       The purpose of the node evaluation function is to evaluate outputs.
+				//       Therefore if a node has no outputs, then it shouldn't need any evaluation.
+				UE_LOG(LogChaosDataflow, Verbose, TEXT("FDataflowNode::Evaluate(): Node [%s], Output [nullptr], NodeTimestamp [%lu]"), *Node->GetName().ToString(), Node->LastModifiedTimestamp.Value);
 				Node->Evaluate(Context, nullptr);
 			}
 		}
+		else
+		{
+			ensureMsgf(false, TEXT("Invalid arguments, either Node or Output needs to be non null."));
+		}
 	}
 
 	void FContextSingle::Evaluate(const FDataflowNode* Node, const FDataflowOutput* Output)
@@ -84,6 +97,7 @@ namespace Dataflow
 
 	bool FContextSingle::Evaluate(const FDataflowOutput& Connection)
 	{
+		UE_LOG(LogChaosDataflow, Verbose, TEXT("FContextSingle::Evaluate(): Node [%s], Output [%s]"), *Connection.GetOwningNode()->GetName().ToString(), *Connection.GetName().ToString());
 		return Connection.EvaluateImpl(*this);
 	}
 
@@ -96,6 +110,7 @@ namespace Dataflow
 
 	bool FContextThreaded::Evaluate(const FDataflowOutput& Connection)
 	{
+		UE_LOG(LogChaosDataflow, Verbose, TEXT("FContextThreaded::Evaluate(): Node [%s], Output [%s]"), *Connection.GetOwningNode()->GetName().ToString(), *Connection.GetName().ToString());
 		Connection.OutputLock->Lock(); ON_SCOPE_EXIT{ Connection.OutputLock->Unlock(); };
 		return Connection.EvaluateImpl(*this);
 	}
diff --git a/Engine/Source/Runtime/Experimental/Dataflow/Core/Public/Dataflow/DataflowConnection.h b/Engine/Source/Runtime/Experimental/Dataflow/Core/Public/Dataflow/DataflowConnection.h
--- a/Engine/Source/Runtime/Experimental/Dataflow/Core/Public/Dataflow/DataflowConnection.h
+++ b/Engine/Source/Runtime/Experimental/Dataflow/Core/Public/Dataflow/DataflowConnection.h
@@ -183,6 +183,7 @@ public:
 
 	DATAFLOWCORE_API FGuid GetOwningNodeGuid() const;
 	DATAFLOWCORE_API uint32 GetOwningNodeValueHash() const;
+	DATAFLOWCORE_API Dataflow::FTimestamp GetOwningNodeTimestamp() const;
 
 	const FProperty* GetProperty() const { return Property; }
 
diff --git a/Engine/Source/Runtime/Experimental/Dataflow/Core/Public/Dataflow/DataflowInputOutput.h b/Engine/Source/Runtime/Experimental/Dataflow/Core/Public/Dataflow/DataflowInputOutput.h
--- a/Engine/Source/Runtime/Experimental/Dataflow/Core/Public/Dataflow/DataflowInputOutput.h
+++ b/Engine/Source/Runtime/Experimental/Dataflow/Core/Public/Dataflow/DataflowInputOutput.h
@@ -193,7 +193,7 @@ public:
 	{
 		if (Property)
 		{
-			Context.SetData(CacheKey(), Property, Forward<T>(InVal), GetOwningNodeGuid(), GetOwningNodeValueHash(), Dataflow::FTimestamp::Current());
+			Context.SetData(CacheKey(), Property, Forward<T>(InVal), GetOwningNodeGuid(), GetOwningNodeValueHash(), GetOwningNodeTimestamp());
 		}
 	}
 
@@ -206,7 +206,7 @@ public:
 				using FSingleType = typename decltype(SingleTypePolicy)::FType;
 				FSingleType ValueToSet{};
 				FDataflowConverter<typename TAnyType::FStorageType>::To(InVal, ValueToSet);
-				Context.SetData(CacheKey(), GetProperty(), Forward<FSingleType>(ValueToSet), GetOwningNodeGuid(), GetOwningNodeValueHash(), Dataflow::FTimestamp::Current());
+				Context.SetData(CacheKey(), GetProperty(), Forward<FSingleType>(ValueToSet), GetOwningNodeGuid(), GetOwningNodeValueHash(), GetOwningNodeTimestamp());
 			});
 	}
 
@@ -215,7 +215,7 @@ public:
 	{
 		if (!this->Evaluate(Context))
 		{
-			Context.SetData(CacheKey(), Property, Default, GetOwningNodeGuid(), GetOwningNodeValueHash(), Dataflow::FTimestamp::Current());
+			Context.SetData(CacheKey(), Property, Default, GetOwningNodeGuid(), GetOwningNodeValueHash(), GetOwningNodeTimestamp());
 		}
 
 		if (Context.HasData(CacheKey()))
@@ -226,6 +226,11 @@ public:
 		return Default;
 	}
 
+	bool HasCachedValue(Dataflow::FContext& Context) const
+	{
+		return Context.HasData(CacheKey(), GetOwningNodeTimestamp());
+	}
+
 	// there's no need for a templatized version as the parameter will not be used
 	// the method do check if the type of the input is the same as the output type though 
 	DATAFLOWCORE_API void ForwardInput(const Dataflow::FConnectionReference& InputReference, Dataflow::FContext& Context) const;
@@ -259,7 +264,7 @@ const T& FDataflowInput::GetValue(Dataflow::FContext& Context, const T& Default)
 	{
 		if (!ConnectionOut->Evaluate(Context))
 		{
-			Context.SetData(ConnectionOut->CacheKey(), Property, Default, GetOwningNodeGuid(), GetOwningNodeValueHash(), Dataflow::FTimestamp::Current());
+			Context.SetData(ConnectionOut->CacheKey(), Property, Default, GetOwningNodeGuid(), GetOwningNodeValueHash(), GetOwningNodeTimestamp());
 		}
 		if (Context.HasData(ConnectionOut->CacheKey()))
 		{
diff --git a/Engine/Source/Runtime/Experimental/Dataflow/Core/Public/Dataflow/DataflowNodeParameters.h b/Engine/Source/Runtime/Experimental/Dataflow/Core/Public/Dataflow/DataflowNodeParameters.h
--- a/Engine/Source/Runtime/Experimental/Dataflow/Core/Public/Dataflow/DataflowNodeParameters.h
+++ b/Engine/Source/Runtime/Experimental/Dataflow/Core/Public/Dataflow/DataflowNodeParameters.h
@@ -182,13 +182,8 @@ namespace Dataflow
 		FContextCache DataStore;
 
 	public:
-		FContext(FTimestamp InTimestamp)
-			: Timestamp(InTimestamp)
-		{}
-
-		virtual ~FContext() {}
-		
-		FTimestamp Timestamp = FTimestamp::Invalid;
+		FContext() = default;
+		virtual ~FContext() = default;
 
 		static FName StaticType() { return FName("FContext"); }
 
@@ -244,7 +239,6 @@ namespace Dataflow
 			return Default;
 		}
 
-		
 		virtual bool HasDataImpl(FContextCacheKey Key, FTimestamp InTimestamp = FTimestamp::Invalid) = 0;
 		
 		bool HasData(FContextCacheKey Key, FTimestamp InTimestamp = FTimestamp::Invalid)
@@ -262,12 +256,13 @@ namespace Dataflow
 
 		virtual void Serialize(FArchive& Ar)
 		{
+			FTimestamp Timestamp = FTimestamp::Invalid;
 			Ar << Timestamp;
 			Ar << DataStore;
 		}
 
+		DATAFLOWCORE_API FTimestamp GetTimestamp(FContextCacheKey Key) const;
 
-		FTimestamp GetTimestamp() const { return Timestamp; }
 		virtual void Evaluate(const FDataflowNode* Node, const FDataflowOutput* Output) = 0;
 		virtual bool Evaluate(const FDataflowOutput& Connection) = 0;
 
@@ -308,9 +303,7 @@ namespace Dataflow
 	public:
 		DATAFLOW_CONTEXT_INTERNAL(FContext, FContextSingle);
 
-		FContextSingle(FTimestamp InTime)
-			: FContext(InTime)
-		{}
+		FContextSingle() = default;
 
 		virtual void SetDataImpl(FContextCacheKey Key, TUniquePtr<FContextCacheElementBase>&& DataStoreEntry) override
 		{
@@ -345,8 +338,8 @@ namespace Dataflow
 		DATAFLOW_CONTEXT_INTERNAL(FContext, FContextThreaded);
 
 
-		FContextThreaded(FTimestamp InTime)
-			: FContext(InTime)
+		FContextThreaded()
+			: FContext()
 		{
 			CacheLock = MakeShared<FCriticalSection>();
 		}
@@ -359,7 +352,7 @@ namespace Dataflow
 			// Threaded evaluation can only set an output once per context evaluation. Otherwise
 			// downstream nodes that are extracting the data will get currupted store entries. 
 			TUniquePtr<FContextCacheElementBase>* CurrentData = DataStore.Find(Key);
-			if (!CurrentData || !(*CurrentData) || (*CurrentData)->GetTimestamp() < GetTimestamp())
+			if (!CurrentData || !(*CurrentData) || (*CurrentData)->GetTimestamp() < DataStoreEntry->GetTimestamp())
 			{
 				DataStore.Emplace(Key, MoveTemp(DataStoreEntry));
 			}
diff --git a/Engine/Source/Runtime/Experimental/Dataflow/Engine/Private/Dataflow/DataflowBlueprintLibrary.cpp b/Engine/Source/Runtime/Experimental/Dataflow/Engine/Private/Dataflow/DataflowBlueprintLibrary.cpp
--- a/Engine/Source/Runtime/Experimental/Dataflow/Engine/Private/Dataflow/DataflowBlueprintLibrary.cpp
+++ b/Engine/Source/Runtime/Experimental/Dataflow/Engine/Private/Dataflow/DataflowBlueprintLibrary.cpp
@@ -14,10 +14,12 @@ void UDataflowBlueprintLibrary::EvaluateTerminalNodeByName(UDataflow* Dataflow,
 		{
 			if (const FDataflowTerminalNode* TerminalNode = Node->AsType<const FDataflowTerminalNode>())
 			{
-				Dataflow::FEngineContext Context(ResultAsset, FPlatformTime::Cycles64());
+				UE_LOG(LogChaosDataflow, Verbose, TEXT("UDataflowBlueprintLibrary::EvaluateTerminalNodeByName(): Node [%s]"), *TerminalNodeName.ToString());
+				Dataflow::FEngineContext Context(ResultAsset);
 				TerminalNode->Evaluate(Context);
 				if (ResultAsset)
 				{
+					UE_LOG(LogChaosDataflow, Verbose, TEXT("FDataflowTerminalNode::SetAssetValue(): TerminalNode [%s], Asset [%s]"), *TerminalNodeName.ToString(), *ResultAsset->GetName());
 					TerminalNode->SetAssetValue(ResultAsset, Context);
 				}
 			}
diff --git a/Engine/Source/Runtime/Experimental/Dataflow/Engine/Private/Dataflow/DataflowContent.cpp b/Engine/Source/Runtime/Experimental/Dataflow/Engine/Private/Dataflow/DataflowContent.cpp
--- a/Engine/Source/Runtime/Experimental/Dataflow/Engine/Private/Dataflow/DataflowContent.cpp
+++ b/Engine/Source/Runtime/Experimental/Dataflow/Engine/Private/Dataflow/DataflowContent.cpp
@@ -156,8 +156,6 @@ namespace DataflowContextHelpers
 				return false;
 			}
 		}
-		Context->Timestamp = NewTimestamp;
-
 		return true;
 	}
 
@@ -308,7 +306,7 @@ void UDataflowBaseContent::SetDataflowOwner(const TObjectPtr<UObject>& InOwner)
 {
 	if(!DataflowContext)
 	{
-		DataflowContext = MakeShared<Dataflow::FEngineContext>(nullptr, Dataflow::FTimestamp::Invalid);
+		DataflowContext = MakeShared<Dataflow::FEngineContext>(nullptr);
 	}
 	DataflowContext->Owner = InOwner;  
 	SetConstructionDirty(true);
@@ -324,7 +322,7 @@ void UDataflowBaseContent::SetDataflowAsset(const TObjectPtr<UDataflow>& Dataflo
 {
 	if(!DataflowContext)
 	{
-		DataflowContext = MakeShared<Dataflow::FEngineContext>(nullptr, Dataflow::FTimestamp::Invalid);
+		DataflowContext = MakeShared<Dataflow::FEngineContext>(nullptr);
 	}
 	DataflowGraph = DataflowAsset;  
 	SetConstructionDirty(true);
@@ -377,7 +375,7 @@ void UDataflowBaseContent::Serialize(FArchive& Ar)
 
 	if (!DataflowContext)
 	{
-		DataflowContext = MakeShared<Dataflow::FEngineContext>(nullptr, LastModifiedTimestamp);
+		DataflowContext = MakeShared<Dataflow::FEngineContext>(nullptr);
 	}
 	DataflowContext->Serialize(Ar);
 }
diff --git a/Engine/Source/Runtime/Experimental/Dataflow/Engine/Public/Dataflow/DataflowObjectInterface.h b/Engine/Source/Runtime/Experimental/Dataflow/Engine/Public/Dataflow/DataflowObjectInterface.h
--- a/Engine/Source/Runtime/Experimental/Dataflow/Engine/Public/Dataflow/DataflowObjectInterface.h
+++ b/Engine/Source/Runtime/Experimental/Dataflow/Engine/Public/Dataflow/DataflowObjectInterface.h
@@ -16,8 +16,8 @@ namespace Dataflow
 	public:
 		DATAFLOW_CONTEXT_INTERNAL(Base, TEngineContext);
 
-		TEngineContext(const TObjectPtr<UObject>& InOwner, FTimestamp InTimestamp)
-			: Base(InTimestamp)
+		explicit TEngineContext(const TObjectPtr<UObject>& InOwner)
+			: Base()
 			, Owner(InOwner)
 		{}
 
