Dataflow Editor: add mechanism to FDataflowToolRegistry to allow enabling/disabling commands when tools begin or end. This allows hotkeys to be used in Dataflow tools that are defined in other modules.
- Also enable shortcut keys in Dataflow Editor Construction Viewport.

diff --git a/Engine/Plugins/ChaosClothAssetEditor/Source/ChaosClothAssetEditorTools/Private/ChaosClothAsset/ChaosClothAssetEditorToolsModule.cpp b/Engine/Plugins/ChaosClothAssetEditor/Source/ChaosClothAssetEditorTools/Private/ChaosClothAsset/ChaosClothAssetEditorToolsModule.cpp
--- a/Engine/Plugins/ChaosClothAssetEditor/Source/ChaosClothAssetEditorTools/Private/ChaosClothAsset/ChaosClothAssetEditorToolsModule.cpp
+++ b/Engine/Plugins/ChaosClothAssetEditor/Source/ChaosClothAssetEditorTools/Private/ChaosClothAsset/ChaosClothAssetEditorToolsModule.cpp
@@ -7,6 +7,7 @@
 #include "ChaosClothAsset/SelectionNode.h"
 #include "ChaosClothAsset/TransferSkinWeightsNode.h"
 #include "ChaosClothAsset/WeightMapNode.h"
+#include "ChaosClothAsset/ClothToolActionCommandBindings.h"
 
 namespace UE::Chaos::ClothAsset
 {
@@ -16,9 +17,12 @@ namespace UE::Chaos::ClothAsset
 		virtual void StartupModule() override
 		{
 			Dataflow::FDataflowToolRegistry& ToolRegistry = Dataflow::FDataflowToolRegistry::Get();
-			ToolRegistry.AddNodeToToolMapping(FChaosClothAssetWeightMapNode::StaticType(), NewObject<UClothEditorWeightMapPaintToolBuilder>());
-			ToolRegistry.AddNodeToToolMapping(FChaosClothAssetSelectionNode::StaticType(), NewObject<UClothMeshSelectionToolBuilder>());
-			ToolRegistry.AddNodeToToolMapping(FChaosClothAssetTransferSkinWeightsNode::StaticType(), NewObject<UClothTransferSkinWeightsToolBuilder>());
+		
+			const TSharedRef<const FClothToolActionCommandBindings> ClothToolActions = MakeShared<FClothToolActionCommandBindings>();
+
+			ToolRegistry.AddNodeToToolMapping(FChaosClothAssetWeightMapNode::StaticType(), NewObject<UClothEditorWeightMapPaintToolBuilder>(), ClothToolActions);
+			ToolRegistry.AddNodeToToolMapping(FChaosClothAssetSelectionNode::StaticType(), NewObject<UClothMeshSelectionToolBuilder>(), ClothToolActions);
+			ToolRegistry.AddNodeToToolMapping(FChaosClothAssetTransferSkinWeightsNode::StaticType(), NewObject<UClothTransferSkinWeightsToolBuilder>(), ClothToolActions);
 		}
 
 		virtual void ShutdownModule() override
diff --git a/Engine/Plugins/ChaosClothAssetEditor/Source/ChaosClothAssetEditorTools/Private/ChaosClothAsset/ClothToolActionCommandBindings.cpp b/Engine/Plugins/ChaosClothAssetEditor/Source/ChaosClothAssetEditorTools/Private/ChaosClothAsset/ClothToolActionCommandBindings.cpp
new file mode 100644
--- /dev/null
+++ b/Engine/Plugins/ChaosClothAssetEditor/Source/ChaosClothAssetEditorTools/Private/ChaosClothAsset/ClothToolActionCommandBindings.cpp
@@ -0,0 +1,64 @@
+// Copyright Epic Games, Inc. All Rights Reserved.
+
+#include "ChaosClothAsset/ClothToolActionCommandBindings.h"
+
+#define LOCTEXT_NAMESPACE "ClothToolActionCommandBindings"
+
+FClothEditorWeightMapPaintToolActionCommands::FClothEditorWeightMapPaintToolActionCommands() : FClothToolActionCommands<FClothEditorWeightMapPaintToolActionCommands, UClothEditorWeightMapPaintTool>(
+	TEXT("ClothEditorWeightMapPaintToolContext"), LOCTEXT("ClothEditorWeightMapPaintToolContext", "Cloth Weight Map Paint Tool Context"))
+{}
+
+FClothMeshSelectionToolActionCommands::FClothMeshSelectionToolActionCommands() : FClothToolActionCommands<FClothMeshSelectionToolActionCommands, UClothMeshSelectionTool>(
+	TEXT("ClothSelectionToolContext"), LOCTEXT("ClothSelectionToolContext", "Cloth Selection Tool Context"))
+{}
+
+FClothTransferSkinWeightsToolActionCommands::FClothTransferSkinWeightsToolActionCommands() : FClothToolActionCommands<FClothTransferSkinWeightsToolActionCommands, UClothTransferSkinWeightsTool>(
+	TEXT("ClothTransferSkinWeightsToolContext"), LOCTEXT("ClothTransferSkinWeightsToolContext", "Cloth Transfer Skin Weights Tool Context"))
+{}
+
+FClothToolActionCommandBindings::FClothToolActionCommandBindings()
+{
+	// Note: if a TCommands<> doesn't actually register any commands then it will be deleted. Only WeightMapPaintTool currently has key commands, but we will include the other tools
+	// here so that hotkeys can be added to them in the future. This means we need to check if the objects are registered before trying to use them below.
+
+	FClothEditorWeightMapPaintToolActionCommands::Register();
+	FClothMeshSelectionToolActionCommands::Register();
+	FClothTransferSkinWeightsToolActionCommands::Register();
+}
+
+void FClothToolActionCommandBindings::UnbindActiveCommands(const TSharedPtr<FUICommandList>& UICommandList) const
+{
+	if (FClothEditorWeightMapPaintToolActionCommands::IsRegistered())
+	{
+		FClothEditorWeightMapPaintToolActionCommands::Get().UnbindActiveCommands(UICommandList);
+	}
+
+	if (FClothMeshSelectionToolActionCommands::IsRegistered())
+	{
+		FClothMeshSelectionToolActionCommands::Get().UnbindActiveCommands(UICommandList);
+	}
+
+	if (FClothTransferSkinWeightsToolActionCommands::IsRegistered())
+	{
+		FClothTransferSkinWeightsToolActionCommands::Get().UnbindActiveCommands(UICommandList);
+	}
+}
+
+void FClothToolActionCommandBindings::BindCommandsForCurrentTool(const TSharedPtr<FUICommandList>& UICommandList, UInteractiveTool* Tool) const
+{
+	if (ExactCast<UClothEditorWeightMapPaintTool>(Tool) && FClothEditorWeightMapPaintToolActionCommands::IsRegistered())
+	{
+		FClothEditorWeightMapPaintToolActionCommands::Get().BindCommandsForCurrentTool(UICommandList, Tool);
+	}
+	else if (ExactCast<UClothMeshSelectionTool>(Tool) && FClothMeshSelectionToolActionCommands::IsRegistered())
+	{
+		FClothMeshSelectionToolActionCommands::Get().BindCommandsForCurrentTool(UICommandList, Tool);
+	}
+	else if (ExactCast<UClothTransferSkinWeightsTool>(Tool) && FClothTransferSkinWeightsToolActionCommands::IsRegistered())
+	{
+		FClothTransferSkinWeightsToolActionCommands::Get().BindCommandsForCurrentTool(UICommandList, Tool);
+	}
+}
+
+#undef LOCTEXT_NAMESPACE
+
diff --git a/Engine/Plugins/ChaosClothAssetEditor/Source/ChaosClothAssetEditorTools/Private/ChaosClothAsset/ClothToolActionCommandBindings.h b/Engine/Plugins/ChaosClothAssetEditor/Source/ChaosClothAssetEditorTools/Private/ChaosClothAsset/ClothToolActionCommandBindings.h
new file mode 100644
--- /dev/null
+++ b/Engine/Plugins/ChaosClothAssetEditor/Source/ChaosClothAssetEditorTools/Private/ChaosClothAsset/ClothToolActionCommandBindings.h
@@ -0,0 +1,65 @@
+// Copyright Epic Games, Inc. All Rights Reserved.
+
+#pragma once
+
+#include "Dataflow/DataflowToolRegistry.h"
+#include "Tools/InteractiveToolsCommands.h"
+#include "Styling/AppStyle.h"			// Style name is a required parameter of TInteractiveToolCommands/FBindingContext
+
+// For GetToolDefaultObjectList in class template 
+#include "ClothWeightMapPaintTool.h"	
+#include "ClothMeshSelectionTool.h"
+#include "ClothTransferSkinWeightsTool.h"
+
+
+// TInteractiveToolCommands<> are typically used to bind the current set of available hotkey commands when a tool starts/ends. However we cannot store multiple actions 
+// with the same key activation in a single TInteractiveToolCommands object (even if they are only active in different tools), so we end up creating one TInteractiveToolCommands per tool. 
+// This is also what Modeling Tools does -- see comments in ModelingToolsActions.h.
+
+
+// Base class with TInteractiveToolCommands<> boilerplate
+
+template<typename T, typename ToolClass>
+class FClothToolActionCommands : public TInteractiveToolCommands<T>
+{
+public:
+	FClothToolActionCommands(const FName InContextName, const FText& InContextDesc) : TInteractiveToolCommands<T>(InContextName, InContextDesc, NAME_None, FAppStyle::GetAppStyleSetName())
+	{}
+
+	virtual void GetToolDefaultObjectList(TArray<UInteractiveTool*>& ToolCDOs) override
+	{
+		ToolCDOs.Add(GetMutableDefault<ToolClass>());
+	}
+};
+
+// Derived classes for individual tools
+
+class FClothEditorWeightMapPaintToolActionCommands : public FClothToolActionCommands<FClothEditorWeightMapPaintToolActionCommands, UClothEditorWeightMapPaintTool>
+{
+public:
+	FClothEditorWeightMapPaintToolActionCommands();
+};
+
+class FClothMeshSelectionToolActionCommands : public FClothToolActionCommands<FClothMeshSelectionToolActionCommands, UClothMeshSelectionTool>
+{
+public:
+	FClothMeshSelectionToolActionCommands();
+};
+
+class FClothTransferSkinWeightsToolActionCommands : public FClothToolActionCommands<FClothTransferSkinWeightsToolActionCommands, UClothTransferSkinWeightsTool>
+{
+public:
+	FClothTransferSkinWeightsToolActionCommands();
+};
+
+
+// Tool action registry entry
+
+class FClothToolActionCommandBindings : public Dataflow::FDataflowToolRegistry::IDataflowToolActionCommands
+{
+
+public:
+	FClothToolActionCommandBindings();
+	virtual void UnbindActiveCommands(const TSharedPtr<FUICommandList>& UICommandList) const override;
+	virtual void BindCommandsForCurrentTool(const TSharedPtr<FUICommandList>& UICommandList, UInteractiveTool* Tool) const override;
+};
diff --git a/Engine/Plugins/Experimental/Dataflow/Source/DataflowEditor/Private/Dataflow/DataflowConstructionViewportClient.cpp b/Engine/Plugins/Experimental/Dataflow/Source/DataflowEditor/Private/Dataflow/DataflowConstructionViewportClient.cpp
--- a/Engine/Plugins/Experimental/Dataflow/Source/DataflowEditor/Private/Dataflow/DataflowConstructionViewportClient.cpp
+++ b/Engine/Plugins/Experimental/Dataflow/Source/DataflowEditor/Private/Dataflow/DataflowConstructionViewportClient.cpp
@@ -65,6 +65,22 @@ USelection* FDataflowConstructionViewportClient::GetSelectedComponents()
 	return ModeTools->GetSelectedComponents(); 
 }
 
+bool FDataflowConstructionViewportClient::InputKey(const FInputKeyEventArgs& EventArgs)
+{
+	// See if any tool commands want to handle the key event
+	const TSharedPtr<FUICommandList> PinnedToolCommandList = ToolCommandList.Pin();
+	if (EventArgs.Event != IE_Released && PinnedToolCommandList.IsValid())
+	{
+		const FModifierKeysState KeyState = FSlateApplication::Get().GetModifierKeys();
+		if (PinnedToolCommandList->ProcessCommandBindings(EventArgs.Key, KeyState, (EventArgs.Event == IE_Repeat)))
+		{
+			return true;
+		}
+	}
+
+	return FEditorViewportClient::InputKey(EventArgs);
+}
+
 
 void FDataflowConstructionViewportClient::ProcessClick(FSceneView& View, HHitProxy* HitProxy, FKey Key, EInputEvent Event, uint32 HitX, uint32 HitY)
 {
diff --git a/Engine/Plugins/Experimental/Dataflow/Source/DataflowEditor/Private/Dataflow/DataflowEditorCommands.cpp b/Engine/Plugins/Experimental/Dataflow/Source/DataflowEditor/Private/Dataflow/DataflowEditorCommands.cpp
--- a/Engine/Plugins/Experimental/Dataflow/Source/DataflowEditor/Private/Dataflow/DataflowEditorCommands.cpp
+++ b/Engine/Plugins/Experimental/Dataflow/Source/DataflowEditor/Private/Dataflow/DataflowEditorCommands.cpp
@@ -130,23 +130,17 @@ void FDataflowEditorCommandsImpl::RegisterCommands()
 
 }
 
-void FDataflowEditorCommandsImpl::GetToolDefaultObjectList(TArray<UInteractiveTool*>& ToolCDOs)
-{
-	ToolCDOs.Add(GetMutableDefault<UDataflowEditorWeightMapPaintTool>());
-}
 
 void FDataflowEditorCommandsImpl::UpdateToolCommandBinding(UInteractiveTool* Tool, TSharedPtr<FUICommandList> UICommandList, bool bUnbind)
 {
-	if (FDataflowEditorCommandsImpl::IsRegistered())
+	Dataflow::FDataflowToolRegistry& ToolRegistry = Dataflow::FDataflowToolRegistry::Get();
+	if (bUnbind)
 	{
-		if (bUnbind)
-		{
-			FDataflowEditorCommandsImpl::Get().UnbindActiveCommands(UICommandList);
-		}
-		else
-		{
-			FDataflowEditorCommandsImpl::Get().BindCommandsForCurrentTool(UICommandList, Tool);
-		}
+		ToolRegistry.UnbindActiveCommands(UICommandList);
+	}
+	else
+	{
+		ToolRegistry.BindCommandsForCurrentTool(UICommandList, Tool);
 	}
 }
 
diff --git a/Engine/Plugins/Experimental/Dataflow/Source/DataflowEditor/Private/Dataflow/DataflowEditorModule.cpp b/Engine/Plugins/Experimental/Dataflow/Source/DataflowEditor/Private/Dataflow/DataflowEditorModule.cpp
--- a/Engine/Plugins/Experimental/Dataflow/Source/DataflowEditor/Private/Dataflow/DataflowEditorModule.cpp
+++ b/Engine/Plugins/Experimental/Dataflow/Source/DataflowEditor/Private/Dataflow/DataflowEditorModule.cpp
@@ -24,6 +24,49 @@ const FColor FDataflowEditorModule::SurfaceColor = FLinearColor(0.6, 0.6, 0.6).T
 static const FName ScalarVertexPropertyGroupName = TEXT("ScalarVertexPropertyGroup");
 static const FName DataflowFunctionPropertyName = TEXT("DataflowFunctionProperty");
 
+class FDataflowEditorWeightMapPaintToolActionCommands : public TInteractiveToolCommands<FDataflowEditorWeightMapPaintToolActionCommands>
+{
+public:
+	FDataflowEditorWeightMapPaintToolActionCommands() : 
+		TInteractiveToolCommands<FDataflowEditorWeightMapPaintToolActionCommands>(
+			TEXT("DataflowEditorWeightMapPaintToolContext"),
+			LOCTEXT("DataflowEditorWeightMapPaintToolContext", "Dataflow Weight Map Paint Tool Context"),
+			NAME_None,
+			FAppStyle::GetAppStyleSetName())
+	{}
+
+	virtual void GetToolDefaultObjectList(TArray<UInteractiveTool*>& ToolCDOs) override
+	{
+		ToolCDOs.Add(GetMutableDefault<UDataflowEditorWeightMapPaintTool>());
+	}
+};
+
+
+class FDataflowToolActionCommandBindings : public Dataflow::FDataflowToolRegistry::IDataflowToolActionCommands
+{
+public:
+	FDataflowToolActionCommandBindings()
+	{
+		FDataflowEditorWeightMapPaintToolActionCommands::Register();
+	}
+
+	virtual void UnbindActiveCommands(const TSharedPtr<FUICommandList>& UICommandList) const override
+	{
+		checkf(FDataflowEditorWeightMapPaintToolActionCommands::IsRegistered(), TEXT("Expected WeightMapPaintTool actions to have been registered"));
+		FDataflowEditorWeightMapPaintToolActionCommands::Get().UnbindActiveCommands(UICommandList);
+	}
+
+	virtual void BindCommandsForCurrentTool(const TSharedPtr<FUICommandList>& UICommandList, UInteractiveTool* Tool) const override
+	{
+		if (ExactCast<UDataflowEditorWeightMapPaintTool>(Tool))
+		{
+			checkf(FDataflowEditorWeightMapPaintToolActionCommands::IsRegistered(), TEXT("Expected WeightMapPaintTool actions to have been registered"));
+			FDataflowEditorWeightMapPaintToolActionCommands::Get().BindCommandsForCurrentTool(UICommandList, Tool);
+		}
+	}
+};
+
+
 void FDataflowEditorModule::StartupModule()
 {
 	FDataflowEditorStyle::Get();
@@ -40,7 +83,10 @@ void FDataflowEditorModule::StartupModule()
 	Dataflow::FDataflowToolRegistry& ToolRegistry = Dataflow::FDataflowToolRegistry::Get();
 
 	UDataflowEditorWeightMapPaintToolBuilder* const ToolBuilder = NewObject<UDataflowEditorWeightMapPaintToolBuilder>();
-	ToolRegistry.AddNodeToToolMapping(FDataflowCollectionAddScalarVertexPropertyNode::StaticType(), ToolBuilder);
+	TSharedRef<const FDataflowToolActionCommandBindings> Actions = MakeShared<FDataflowToolActionCommandBindings>();
+	ToolRegistry.AddNodeToToolMapping(FDataflowCollectionAddScalarVertexPropertyNode::StaticType(), ToolBuilder, Actions);
+
+	FDataflowEditorCommands::Register();
 }
 
 void FDataflowEditorModule::ShutdownModule()
@@ -54,6 +100,8 @@ void FDataflowEditorModule::ShutdownModule()
 		PropertyModule->UnregisterCustomPropertyTypeLayout(DataflowFunctionPropertyName);
 	}
 
+	FDataflowEditorCommands::Unregister();
+
 	Dataflow::FDataflowToolRegistry& ToolRegistry = Dataflow::FDataflowToolRegistry::Get();
 	ToolRegistry.RemoveNodeToToolMapping(FDataflowCollectionAddScalarVertexPropertyNode::StaticType());
 }
diff --git a/Engine/Plugins/Experimental/Dataflow/Source/DataflowEditor/Private/Dataflow/DataflowToolRegistry.cpp b/Engine/Plugins/Experimental/Dataflow/Source/DataflowEditor/Private/Dataflow/DataflowToolRegistry.cpp
--- a/Engine/Plugins/Experimental/Dataflow/Source/DataflowEditor/Private/Dataflow/DataflowToolRegistry.cpp
+++ b/Engine/Plugins/Experimental/Dataflow/Source/DataflowEditor/Private/Dataflow/DataflowToolRegistry.cpp
@@ -17,10 +17,10 @@ namespace Dataflow
 		TLazySingleton<FDataflowToolRegistry>::TearDown();
 	}
 
-	void FDataflowToolRegistry::AddNodeToToolMapping(const FName& NodeName, TObjectPtr<UInteractiveToolBuilder> ToolBuilder)
+	void FDataflowToolRegistry::AddNodeToToolMapping(const FName& NodeName, TObjectPtr<UInteractiveToolBuilder> ToolBuilder, const TSharedRef<const IDataflowToolActionCommands>& ToolActionRegistry)
 	{
 		// FUICommandInfo is uninitialized, it will be created later in FDataflowEditorCommandsImpl::RegisterCommands
-		NodeTypeToToolMap.Add(NodeName, { ToolBuilder, nullptr });
+		NodeTypeToToolMap.Add(NodeName, { ToolBuilder, ToolActionRegistry, nullptr });
 	}
 
 	void FDataflowToolRegistry::RemoveNodeToToolMapping(const FName& NodeName)
@@ -53,4 +53,20 @@ namespace Dataflow
 		return NodeTypeToToolMap[NodeName].ToolBuilder;
 	}
 
+	void FDataflowToolRegistry::UnbindActiveCommands(const TSharedPtr<FUICommandList>& UICommandList) const
+	{
+		for (const TPair<FName, FToolInfo>& Entry : NodeTypeToToolMap)
+		{
+			Entry.Value.ToolActionCommands->UnbindActiveCommands(UICommandList);
+		}
+	}
+
+	void FDataflowToolRegistry::BindCommandsForCurrentTool(const TSharedPtr<FUICommandList>& UICommandList, UInteractiveTool* Tool) const
+	{
+		for (const TPair<FName, FToolInfo>& Entry : NodeTypeToToolMap)
+		{
+			Entry.Value.ToolActionCommands->BindCommandsForCurrentTool(UICommandList, Tool);
+		}
+	}
+
 }
diff --git a/Engine/Plugins/Experimental/Dataflow/Source/DataflowEditor/Public/Dataflow/DataflowConstructionViewportClient.h b/Engine/Plugins/Experimental/Dataflow/Source/DataflowEditor/Public/Dataflow/DataflowConstructionViewportClient.h
--- a/Engine/Plugins/Experimental/Dataflow/Source/DataflowEditor/Public/Dataflow/DataflowConstructionViewportClient.h
+++ b/Engine/Plugins/Experimental/Dataflow/Source/DataflowEditor/Public/Dataflow/DataflowConstructionViewportClient.h
@@ -58,6 +58,7 @@ private:
 
 	// FEditorViewportClient interface
 	virtual void Tick(float DeltaSeconds) override;
+	virtual bool InputKey(const FInputKeyEventArgs& EventArgs) override;
 	virtual void ProcessClick(FSceneView& View, HHitProxy* HitProxy, FKey Key, EInputEvent Event, uint32 HitX, uint32 HitY) override;
 
 	/** Toolkit used to create the viewport client */
diff --git a/Engine/Plugins/Experimental/Dataflow/Source/DataflowEditor/Public/Dataflow/DataflowEditorCommands.h b/Engine/Plugins/Experimental/Dataflow/Source/DataflowEditor/Public/Dataflow/DataflowEditorCommands.h
--- a/Engine/Plugins/Experimental/Dataflow/Source/DataflowEditor/Public/Dataflow/DataflowEditorCommands.h
+++ b/Engine/Plugins/Experimental/Dataflow/Source/DataflowEditor/Public/Dataflow/DataflowEditorCommands.h
@@ -44,7 +44,8 @@ public:
 	 virtual void RegisterCommands() override;
 
 	// TInteractiveToolCommands<>
-	 virtual void GetToolDefaultObjectList(TArray<UInteractiveTool*>& ToolCDOs) override;
+	// Each tool will have its own TInteractiveToolCommands<> object stored in the DataflowToolRegistry, so this should not return anything
+	virtual void GetToolDefaultObjectList(TArray<UInteractiveTool*>& ToolCDOs) override {}
 
 	/**
 	* Add or remove commands relevant to Tool to the given UICommandList.
diff --git a/Engine/Plugins/Experimental/Dataflow/Source/DataflowEditor/Public/Dataflow/DataflowToolRegistry.h b/Engine/Plugins/Experimental/Dataflow/Source/DataflowEditor/Public/Dataflow/DataflowToolRegistry.h
--- a/Engine/Plugins/Experimental/Dataflow/Source/DataflowEditor/Public/Dataflow/DataflowToolRegistry.h
+++ b/Engine/Plugins/Experimental/Dataflow/Source/DataflowEditor/Public/Dataflow/DataflowToolRegistry.h
@@ -8,6 +8,8 @@
 
 class FUICommandInfo;
 class UInteractiveToolBuilder;
+class FUICommandList;
+class UInteractiveTool;
 
 namespace Dataflow
 {
@@ -15,10 +17,20 @@ namespace Dataflow
 	{
 	public:
 
+		// Interface for binding/unbinding tool actions. When a tool begins or ends we switch out the currently available FUICommandList. This allows multiple tools to have
+		// individual hotkey actions with the same key chords, for example.
+		class DATAFLOWEDITOR_API IDataflowToolActionCommands
+		{
+		public:
+			virtual ~IDataflowToolActionCommands() = default;
+			virtual void UnbindActiveCommands(const TSharedPtr<FUICommandList>& UICommandList) const = 0;
+			virtual void BindCommandsForCurrentTool(const TSharedPtr<FUICommandList>& UICommandList, UInteractiveTool* Tool) const = 0;
+		};
+
 		static FDataflowToolRegistry& Get();
 		static void TearDown();
 
-		void AddNodeToToolMapping(const FName& NodeName, TObjectPtr<UInteractiveToolBuilder> ToolBuilder);
+		void AddNodeToToolMapping(const FName& NodeName, TObjectPtr<UInteractiveToolBuilder> ToolBuilder, const TSharedRef<const IDataflowToolActionCommands>& ToolActionCommands);
 		void RemoveNodeToToolMapping(const FName& NodeName);
 
 		TArray<FName> GetNodeNames() const;
@@ -28,12 +40,16 @@ namespace Dataflow
 		UInteractiveToolBuilder* GetToolBuilderForNode(const FName& NodeName);
 		const UInteractiveToolBuilder* GetToolBuilderForNode(const FName& NodeName) const;
 
+		void UnbindActiveCommands(const TSharedPtr<FUICommandList>& UICommandList) const;
+		void BindCommandsForCurrentTool(const TSharedPtr<FUICommandList>& UICommandList, UInteractiveTool* Tool) const;
+
 	private:
 
 		struct FToolInfo
 		{
 			// Specified when registering the tool
 			TObjectPtr<UInteractiveToolBuilder> ToolBuilder;
+			TSharedRef<const IDataflowToolActionCommands> ToolActionCommands;
 
 			// Constructed automatically in FDataflowEditorCommandsImpl::RegisterCommands
 			TSharedPtr<FUICommandInfo> ToolCommand;
