Cloth Editor: when clicking a node for the first time start any corresponding tool immediately. This change just creates the DynamicMeshComponent before trying to start the tool.
Also: clear the bDynamicMeshComponentInitDeferred flag after testing it in the Tick function, not inside ReinitializeDynamicMeshComponents(), to prevent unexpected behavior calling ReinitializeDynamicMeshComponents out of sequence.

diff --git a/Engine/Plugins/ChaosClothAssetEditor/Source/ChaosClothAssetEditor/Private/ChaosClothAsset/ClothEditorMode.cpp b/Engine/Plugins/ChaosClothAssetEditor/Source/ChaosClothAssetEditor/Private/ChaosClothAsset/ClothEditorMode.cpp
--- a/Engine/Plugins/ChaosClothAssetEditor/Source/ChaosClothAssetEditor/Private/ChaosClothAsset/ClothEditorMode.cpp
+++ b/Engine/Plugins/ChaosClothAssetEditor/Source/ChaosClothAssetEditor/Private/ChaosClothAsset/ClothEditorMode.cpp
@@ -879,7 +879,6 @@ void UChaosClothAssetEditorMode::ReinitializeDynamicMeshComponents()
 	WireframeDraw = nullptr;
 	ClothSeamDraw = nullptr;
 	SurfaceNormalDraw = nullptr;
-	bDynamicMeshComponentInitDeferred = false;
 
 	TSharedPtr<FManagedArrayCollection> Collection = bDynamicMeshUseInputCollection ? GetInputClothCollection() : GetClothCollection();
 	if (!Collection)
@@ -1261,6 +1260,15 @@ void UChaosClothAssetEditorMode::ModeTick(float DeltaTime)
 		bShouldClearTeleportFlag = true;		// clear the flag next tick
 	}
 
+	if (bDynamicMeshComponentInitDeferred)
+	{
+		ReinitializeDynamicMeshComponents();
+
+		// The first time we get a valid mesh, refocus the camera on it
+		FirstTimeFocusRestSpaceViewport();
+
+		bDynamicMeshComponentInitDeferred = false;
+	}
 
 	if (!NodeTypeForPendingToolStart.IsNone() && !GetToolManager()->HasActiveTool(EToolSide::Left))
 	{
@@ -1275,14 +1283,6 @@ void UChaosClothAssetEditorMode::ModeTick(float DeltaTime)
 		NodeTypeForPendingToolStart = FName();
 	}
 
-	if (bDynamicMeshComponentInitDeferred)
-	{
-		ReinitializeDynamicMeshComponents();
-
-		// The first time we get a valid mesh, refocus the camera on it
-		FirstTimeFocusRestSpaceViewport();
-	}
-
 	const bool bIsInPIEOrSIE = GEditor->PlayWorld != NULL || GEditor->bIsSimulatingInEditor;
 	const bool bShouldPause = PreviewScene->GetPreviewSceneDescription()->bPauseWhilePlayingInEditor && bIsInPIEOrSIE;
 
