Cloth Editor: fix for the mesh in the construction viewport not updating when node properties in the details panel are changed

diff --git a/Engine/Plugins/ChaosClothAssetEditor/Source/ChaosClothAssetEditor/Private/ChaosClothAsset/ClothEditorToolkit.cpp b/Engine/Plugins/ChaosClothAssetEditor/Source/ChaosClothAssetEditor/Private/ChaosClothAsset/ClothEditorToolkit.cpp
--- a/Engine/Plugins/ChaosClothAssetEditor/Source/ChaosClothAssetEditor/Private/ChaosClothAsset/ClothEditorToolkit.cpp
+++ b/Engine/Plugins/ChaosClothAssetEditor/Source/ChaosClothAssetEditor/Private/ChaosClothAsset/ClothEditorToolkit.cpp
@@ -1178,10 +1178,19 @@ void FChaosClothAssetEditorToolkit::OnNodeSelectionChanged(const TSet<UObject*>&
 								{
 									// Warning: Do not execute code that rebuilds the UI in this lambda as it is called by the UI!
 
-									TSharedPtr<FDataflowNode> SelectedDataflowNode = GetSelectedDataflowNode();
+									const TSharedPtr<FDataflowNode> SelectedDataflowNode = GetSelectedDataflowNode();
 									if (SelectedDataflowNode.Get() == InDataflowNode)
 									{
-										GetClothCollectionIfPossible(SelectedDataflowNode, DataflowContext);
+										// Node was invalidated, update the Construction viewport
+										const TSharedPtr<FManagedArrayCollection> InputCollection = GetInputClothCollectionIfPossible(SelectedDataflowNode, DataflowContext);
+										const TSharedPtr<FManagedArrayCollection> Collection = GetClothCollectionIfPossible(SelectedDataflowNode, DataflowContext);
+										if (UChaosClothAssetEditorMode* const ClothMode =
+											CastChecked<UChaosClothAssetEditorMode>(EditorModeManager->GetActiveScriptableMode(UChaosClothAssetEditorMode::EM_ChaosClothAssetEditorModeId)))
+										{
+											constexpr bool bDeferDynamicMeshInit = false;
+											ClothMode->SetSelectedClothCollection(Collection, InputCollection, bDeferDynamicMeshInit);
+										}
+
 										if (DataflowContext.IsValid())  // TODO: The context shouldn't be nullptr on the first node creation
 										{
 											SelectedDataflowNode->OnSelected(*DataflowContext);
