Cloth Asset - Removed StringValue from the Connectable customization display.

[RN] Cloth Asset - Removed the StringValue child property from the node's UI display of the Selection and Weight Map names.

diff --git a/Engine/Plugins/ChaosClothAssetEditor/Source/ChaosClothAssetDataflowNodes/Private/ChaosClothAsset/ConnectableValueCustomization.cpp b/Engine/Plugins/ChaosClothAssetEditor/Source/ChaosClothAssetDataflowNodes/Private/ChaosClothAsset/ConnectableValueCustomization.cpp
--- a/Engine/Plugins/ChaosClothAssetEditor/Source/ChaosClothAssetDataflowNodes/Private/ChaosClothAsset/ConnectableValueCustomization.cpp
+++ b/Engine/Plugins/ChaosClothAssetEditor/Source/ChaosClothAssetDataflowNodes/Private/ChaosClothAsset/ConnectableValueCustomization.cpp
@@ -61,6 +61,27 @@ PRAGMA_ENABLE_DEPRECATION_WARNINGS
 	
 	FConnectableValueCustomization::~FConnectableValueCustomization() = default;
 
+	void FConnectableValueCustomization::CustomizeChildren(
+		TSharedRef<IPropertyHandle> PropertyHandle,
+		IDetailChildrenBuilder& ChildBuilder,
+		IPropertyTypeCustomizationUtils& CustomizationUtils)
+	{
+		for (int32 ChildIndex = 0; ChildIndex < SortedChildHandles.Num(); ++ChildIndex)
+		{
+			TSharedRef<IPropertyHandle> ChildHandle = SortedChildHandles[ChildIndex];
+
+			if (CouldUseFabricsProperty(ChildHandle))
+			{
+				bool bCouldUseFabrics = false;
+				ChildHandle->GetValue(bCouldUseFabrics);
+				if (bCouldUseFabrics)
+				{
+					FImportedValueCustomization::CustomizeChildren(PropertyHandle, ChildBuilder, CustomizationUtils);
+				}
+				return;
+			}
+		}
+	}
 
 	void FConnectableValueCustomization::MakeHeaderRow(TSharedRef<class IPropertyHandle>& StructPropertyHandle, FDetailWidgetRow& Row)
 	{
diff --git a/Engine/Plugins/ChaosClothAssetEditor/Source/ChaosClothAssetDataflowNodes/Private/ChaosClothAsset/ConnectableValueCustomization.h b/Engine/Plugins/ChaosClothAssetEditor/Source/ChaosClothAssetDataflowNodes/Private/ChaosClothAsset/ConnectableValueCustomization.h
--- a/Engine/Plugins/ChaosClothAssetEditor/Source/ChaosClothAssetDataflowNodes/Private/ChaosClothAsset/ConnectableValueCustomization.h
+++ b/Engine/Plugins/ChaosClothAssetEditor/Source/ChaosClothAssetDataflowNodes/Private/ChaosClothAsset/ConnectableValueCustomization.h
@@ -19,15 +19,23 @@ namespace UE::Chaos::ClothAsset
 		virtual ~FConnectableValueCustomization() override;
 
 	protected:
+		//~ Begin IPropertyTypeCustomization implementation
+		virtual void CustomizeChildren(
+			TSharedRef<IPropertyHandle> PropertyHandle,
+			IDetailChildrenBuilder& ChildBuilder,
+			IPropertyTypeCustomizationUtils& CustomizationUtils) override;
+		//~ End IPropertyTypeCustomization implementation
+
+		//~ Begin FMathStructCustomization implementation
 		virtual void MakeHeaderRow(TSharedRef<class IPropertyHandle>& StructPropertyHandle, FDetailWidgetRow& Row) override;
 		virtual TSharedRef<SWidget> MakeChildWidget(TSharedRef<IPropertyHandle>& StructurePropertyHandle, TSharedRef<IPropertyHandle>& PropertyHandle) override;
-		
+		//~ End FMathStructCustomization implementation
+
 		UE_DEPRECATED(5.5, "Override properties are no longer used.")
 		static bool IsOverrideProperty(const TSharedPtr<IPropertyHandle>& Property);
 		UE_DEPRECATED(5.5, "Override properties are no longer used.")
 		static bool IsOverridePropertyOf(const TSharedPtr<IPropertyHandle>& OverrideProperty, const TSharedPtr<IPropertyHandle>& Property);
 		static bool BuildFabricMapsProperty(const TSharedPtr<IPropertyHandle>& Property);
 		static bool CouldUseFabricsProperty(const TSharedPtr<IPropertyHandle>& Property);
-		
 	};
 }
diff --git a/Engine/Plugins/ChaosClothAssetEditor/Source/ChaosClothAssetDataflowNodes/Private/ChaosClothAsset/WeightedValueCustomization.h b/Engine/Plugins/ChaosClothAssetEditor/Source/ChaosClothAssetDataflowNodes/Private/ChaosClothAsset/WeightedValueCustomization.h
--- a/Engine/Plugins/ChaosClothAssetEditor/Source/ChaosClothAssetDataflowNodes/Private/ChaosClothAsset/WeightedValueCustomization.h
+++ b/Engine/Plugins/ChaosClothAssetEditor/Source/ChaosClothAssetDataflowNodes/Private/ChaosClothAsset/WeightedValueCustomization.h
@@ -23,8 +23,21 @@ namespace UE::Chaos::ClothAsset
 		virtual ~FWeightedValueCustomization() override;
 
 	protected:
+		//~ Begin IPropertyTypeCustomization implementation
+		virtual void CustomizeChildren(
+			TSharedRef<IPropertyHandle> PropertyHandle,
+			IDetailChildrenBuilder& ChildBuilder,
+			IPropertyTypeCustomizationUtils& CustomizationUtils)
+		{
+			// Skip the FConnectableValueCustomization implementation that doesn't display children
+			FImportedValueCustomization::CustomizeChildren(PropertyHandle, ChildBuilder, CustomizationUtils);
+		}
+		//~ End IPropertyTypeCustomization implementation
+
+		//~ Begin FMathStructCustomization implementation
 		virtual void MakeHeaderRow(TSharedRef<class IPropertyHandle>& StructPropertyHandle, FDetailWidgetRow& Row) override;
 		virtual TSharedRef<SWidget> MakeChildWidget(TSharedRef<IPropertyHandle>& StructurePropertyHandle, TSharedRef<IPropertyHandle>& PropertyHandle) override;
+		//~ End FMathStructCustomization implementation
 
 	private:
 		static void ExtractFloatMetadata(
