Cloth Editor Weight Map Paint Tool: Smooth brush improvements:
- consider a wider neighborhood of vertices when computing average values
- added a falloff function

diff --git a/Engine/Plugins/ChaosClothAssetEditor/Source/ChaosClothAssetEditorTools/Private/ChaosClothAsset/ClothWeightMapPaintBrushOps.cpp b/Engine/Plugins/ChaosClothAssetEditor/Source/ChaosClothAssetEditorTools/Private/ChaosClothAsset/ClothWeightMapPaintBrushOps.cpp
--- a/Engine/Plugins/ChaosClothAssetEditor/Source/ChaosClothAssetEditorTools/Private/ChaosClothAsset/ClothWeightMapPaintBrushOps.cpp
+++ b/Engine/Plugins/ChaosClothAssetEditor/Source/ChaosClothAssetEditorTools/Private/ChaosClothAsset/ClothWeightMapPaintBrushOps.cpp
@@ -131,18 +131,27 @@ void FWeightMapSmoothBrushOp::ApplyStampByVertices(
 	for (int32 BufferIndex = 0; BufferIndex < NumVertices; ++BufferIndex)
 	{
 		const double Diff = OneRingAverages[BufferIndex] - VertexWeightValues[BufferIndex];
+		
+		const int32 VertexIndex = Vertices[BufferIndex];
+		const FVector3d VertexPos = Mesh->GetVertex(VertexIndex);
+
 		if (bApplyRadiusLimit)
 		{
 			const FVector3d& StampPos = Stamp.LocalFrame.Origin;
-			const int32 VertexIndex = Vertices[BufferIndex];
-			const FVector3d VertexPos = Mesh->GetVertex(VertexIndex);
 			const double DistanceSquared = (VertexPos - StampPos).SquaredLength();
 			if (DistanceSquared >= Stamp.Radius * Stamp.Radius)
 			{
 				continue;
 			}
 		}
-		VertexWeightValues[BufferIndex] += Stamp.Power * Diff;
+
+		double FalloffScalar = 1.0;
+		if (this->Falloff)
+		{
+			FalloffScalar = GetFalloff().Evaluate(Stamp, VertexPos);
+		}
+
+		VertexWeightValues[BufferIndex] += FalloffScalar * Stamp.Power * Diff;
 	}
 }
 
diff --git a/Engine/Plugins/ChaosClothAssetEditor/Source/ChaosClothAssetEditorTools/Private/ChaosClothAsset/ClothWeightMapPaintBrushOps.h b/Engine/Plugins/ChaosClothAssetEditor/Source/ChaosClothAssetEditorTools/Private/ChaosClothAsset/ClothWeightMapPaintBrushOps.h
--- a/Engine/Plugins/ChaosClothAssetEditor/Source/ChaosClothAssetEditorTools/Private/ChaosClothAsset/ClothWeightMapPaintBrushOps.h
+++ b/Engine/Plugins/ChaosClothAssetEditor/Source/ChaosClothAssetEditorTools/Private/ChaosClothAsset/ClothWeightMapPaintBrushOps.h
@@ -132,7 +132,11 @@ public:
 	virtual float GetStrength() override { return Strength; }
 	virtual void SetStrength(float NewStrength) override { Strength = NewStrength; }
 
-	virtual float GetFalloff() override { return 0.0f; }
+	UPROPERTY(EditAnywhere, Category = PaintBrush, meta = (DisplayName = "Falloff", UIMin = "0.0", UIMax = "1.", ClampMin = "0.0", ClampMax = "1."))
+	float Falloff = 0.5;
+
+	virtual float GetFalloff() override { return Falloff; }
+	virtual void SetFalloff(float NewFalloff) override { Falloff = NewFalloff; }
 };
 
 
diff --git a/Engine/Plugins/ChaosClothAssetEditor/Source/ChaosClothAssetEditorTools/Private/ChaosClothAsset/ClothWeightMapPaintTool.cpp b/Engine/Plugins/ChaosClothAssetEditor/Source/ChaosClothAssetEditorTools/Private/ChaosClothAsset/ClothWeightMapPaintTool.cpp
--- a/Engine/Plugins/ChaosClothAssetEditor/Source/ChaosClothAssetEditorTools/Private/ChaosClothAsset/ClothWeightMapPaintTool.cpp
+++ b/Engine/Plugins/ChaosClothAssetEditor/Source/ChaosClothAssetEditorTools/Private/ChaosClothAsset/ClothWeightMapPaintTool.cpp
@@ -228,6 +228,14 @@ void UClothEditorWeightMapPaintTool::Setup()
 			UMeshSculptToolBase::BrushProperties->BrushSize.AdaptiveSize = NewSize; 
 			CalculateBrushRadius();
 		});
+
+	FilterProperties->WatchProperty(FilterProperties->Falloff, [this](double NewFalloff)
+		{
+			// Brush indicator rendering uses this value
+			GetActiveBrushOp()->PropertySet->SetFalloff(NewFalloff);
+		});
+
+
 	FilterProperties->BrushSize = UMeshSculptToolBase::BrushProperties->BrushSize.AdaptiveSize;
 	FilterProperties->RestoreProperties(this);
 	AddToolPropertySource(FilterProperties);
@@ -513,6 +521,7 @@ void UClothEditorWeightMapPaintTool::Setup()
 	DynamicMeshComponent->FastNotifyVertexAttributesUpdated(EMeshRenderAttributeFlags::VertexColors);
 	GetToolManager()->PostInvalidation();
 
+	SetPrimaryFalloffType(EMeshSculptFalloffType::Smooth);
 }
 
 void UClothEditorWeightMapPaintTool::InitializeBrushSizeRange(const UE::Geometry::FAxisAlignedBox3d& TargetBounds)
@@ -713,6 +722,7 @@ void UClothEditorWeightMapPaintTool::OnBeginStroke(const FRay& WorldRay)
 	if (SmoothBrushOpProperties)
 	{
 		SmoothBrushOpProperties->Strength = FilterProperties->Strength * FilterProperties->Strength;
+		SmoothBrushOpProperties->Falloff = FilterProperties->Falloff;
 	}
 
 	// initialize first "Last Stamp", so that we can assume all stamps in stroke have a valid previous stamp
@@ -872,6 +882,21 @@ void UClothEditorWeightMapPaintTool::UpdateROI(const FSculptBrushStamp& BrushSta
 		}
 	}
 
+	// If we are Smoothing, expand the set of vertices to consider. Otherwise vertices near the brush bounds will not use the expected neighborhood to get an average weight.
+	const bool bExpandVertexROI = (FilterProperties->SubToolType == EClothEditorWeightMapPaintInteractionType::Brush && FilterProperties->PrimaryBrushType == EClothEditorWeightMapPaintBrushType::Smooth);
+	if (bExpandVertexROI)
+	{
+		TSet<int32> NewVertexSetBuffer = VertexSetBuffer;
+		for (const int32 Vert : VertexSetBuffer)
+		{
+			for (const int32 NeighborVert : Mesh->VtxVerticesItr(Vert))
+			{
+				NewVertexSetBuffer.Add(NeighborVert);
+			}
+		}
+		VertexSetBuffer = MoveTemp(NewVertexSetBuffer);
+	}
+
 
 	VertexROI.SetNum(0, EAllowShrinking::No);
 	//TODO: If we paint a 2D projection of UVs, these will need to be the 2D vertices not the 3D original mesh vertices
diff --git a/Engine/Plugins/ChaosClothAssetEditor/Source/ChaosClothAssetEditorTools/Private/ChaosClothAsset/ClothWeightMapPaintTool.h b/Engine/Plugins/ChaosClothAssetEditor/Source/ChaosClothAssetEditorTools/Private/ChaosClothAsset/ClothWeightMapPaintTool.h
--- a/Engine/Plugins/ChaosClothAssetEditor/Source/ChaosClothAssetEditorTools/Private/ChaosClothAsset/ClothWeightMapPaintTool.h
+++ b/Engine/Plugins/ChaosClothAssetEditor/Source/ChaosClothAssetEditorTools/Private/ChaosClothAsset/ClothWeightMapPaintTool.h
@@ -136,6 +136,11 @@ public:
 		HideEditConditionToggle, EditConditionHides, EditCondition = "SubToolType == EClothEditorWeightMapPaintInteractionType::Brush"))
 	float BrushSize = 0.25f;
 
+	/** Relative size of falloff region inside the brush */
+	UPROPERTY(EditAnywhere, Category = Brush, meta = (UIMin = 0, ClampMin = 0, UIMax = 1, ClampMax = 1,
+		HideEditConditionToggle, EditCondition = "SubToolType == EClothEditorWeightMapPaintInteractionType::Brush && PrimaryBrushType == EClothEditorWeightMapPaintBrushType::Smooth"))
+	float Falloff = 0.5;
+
 	/** Allow the Brush to hit the back-side of the mesh */
 	UPROPERTY(EditAnywhere, Category = Brush, meta = (HideEditConditionToggle, EditConditionHides, EditCondition = "SubToolType == EClothEditorWeightMapPaintInteractionType::Brush"))
 	bool bHitBackFaces = true;
