Cloth Mesh Selection Tool:
- Add a "Clear Selection" action
- Replace the Grow/Shrink/Flood actions with SelectionMechanic calls

diff --git a/Engine/Plugins/ChaosClothAssetEditor/Source/ChaosClothAssetEditorTools/Private/ChaosClothAsset/ClothMeshSelectionTool.cpp b/Engine/Plugins/ChaosClothAssetEditor/Source/ChaosClothAssetEditorTools/Private/ChaosClothAsset/ClothMeshSelectionTool.cpp
--- a/Engine/Plugins/ChaosClothAssetEditor/Source/ChaosClothAssetEditorTools/Private/ChaosClothAsset/ClothMeshSelectionTool.cpp
+++ b/Engine/Plugins/ChaosClothAssetEditor/Source/ChaosClothAssetEditorTools/Private/ChaosClothAsset/ClothMeshSelectionTool.cpp
@@ -622,13 +622,16 @@ void UClothMeshSelectionTool::ApplyAction(EClothMeshSelectionToolActions ActionT
 		ImportFromCollection(/*bImportFromSecondarySet = */ true);
 		break;
 	case EClothMeshSelectionToolActions::GrowSelection:
-		GrowSelection();
+		SelectionMechanic->GrowSelection();
 		break;
 	case EClothMeshSelectionToolActions::ShrinkSelection:
-		ShrinkSelection();
+		SelectionMechanic->ShrinkSelection();
 		break;
 	case EClothMeshSelectionToolActions::FloodSelection:
-		FloodSelection();
+		SelectionMechanic->FloodSelection();
+		break;
+	case EClothMeshSelectionToolActions::ClearSelection:
+		SelectionMechanic->ClearSelection();
 		break;
 	}
 }
@@ -746,87 +749,5 @@ void UClothMeshSelectionTool::ImportFromCollection(bool bImportFromSecondarySet)
 	}
 }
 
-void UClothMeshSelectionTool::GrowSelection()
-{
-	using namespace UE::Geometry;
-
-	PreviewMesh->ProcessMesh([this](const FDynamicMesh3& Mesh)
-		{
-			FGeometrySelection Selection;
-			Selection.ElementType = SelectionMechanic->Properties->bSelectFaces ? EGeometryElementType::Face : EGeometryElementType::Vertex;
-			SelectionMechanic->GetSelection_AsTriangleTopology(Selection);
-
-			FGeometrySelection BoundarySelection;
-			BoundarySelection.ElementType = SelectionMechanic->Properties->bSelectFaces ? EGeometryElementType::Face : EGeometryElementType::Vertex;
-			MakeBoundaryConnectedSelection(Mesh,
-				Topology.Get(),
-				Selection,
-				[](FGeoSelectionID) { return true; },
-				BoundarySelection
-			);
-			CombineSelectionInPlace(Selection, BoundarySelection, EGeometrySelectionCombineModes::Add);
-			SelectionMechanic->SetSelection_AsTriangleTopology(Selection);
-		});
-
-	PreviewMesh->FastNotifySecondaryTrianglesChanged();
-}
-
-
-void UClothMeshSelectionTool::ShrinkSelection()
-{
-	using namespace UE::Geometry;
-
-	PreviewMesh->ProcessMesh([this](const FDynamicMesh3& Mesh)
-		{
-			FGeometrySelection Selection;
-			Selection.ElementType = SelectionMechanic->Properties->bSelectFaces ? EGeometryElementType::Face : EGeometryElementType::Vertex;
-			SelectionMechanic->GetSelection_AsTriangleTopology(Selection);
-
-			FGeometrySelection BoundarySelection;
-			BoundarySelection.ElementType = SelectionMechanic->Properties->bSelectFaces ? EGeometryElementType::Face : EGeometryElementType::Vertex;
-			MakeBoundaryConnectedSelection(Mesh,
-				Topology.Get(),
-				Selection,
-				[](FGeoSelectionID) { return true; },
-				BoundarySelection
-			);
-
-			CombineSelectionInPlace(Selection, BoundarySelection, EGeometrySelectionCombineModes::Subtract);
-			
-			// TODO: SetSelection_AsTriangleTopology doesn't overwrite the selection, it only adds to it. Fix that.
-			FGroupTopologySelection ClearSelection;
-			SelectionMechanic->SetSelection(ClearSelection, false);
-
-			SelectionMechanic->SetSelection_AsTriangleTopology(Selection);
-		});
-
-	PreviewMesh->FastNotifySecondaryTrianglesChanged();
-}
-
-
-void UClothMeshSelectionTool::FloodSelection()
-{
-	using namespace UE::Geometry;
-
-	PreviewMesh->ProcessMesh([this](const FDynamicMesh3& Mesh)
-		{
-			FGeometrySelection Selection;
-			Selection.ElementType = SelectionMechanic->Properties->bSelectFaces ? EGeometryElementType::Face : EGeometryElementType::Vertex;
-			SelectionMechanic->GetSelection_AsTriangleTopology(Selection);
-
-			FGeometrySelection ConnectedSelection;
-			ConnectedSelection.ElementType = SelectionMechanic->Properties->bSelectFaces ? EGeometryElementType::Face : EGeometryElementType::Vertex;
-			MakeSelectAllConnectedSelection(Mesh,
-				Topology.Get(),
-				Selection,
-				[](FGeoSelectionID) { return true; },
-				[](FGeoSelectionID, FGeoSelectionID) { return true; },
-				ConnectedSelection
-			);
-			SelectionMechanic->SetSelection_AsTriangleTopology(ConnectedSelection);
-		});
-
-	PreviewMesh->FastNotifySecondaryTrianglesChanged();
-}
 
 #undef LOCTEXT_NAMESPACE
diff --git a/Engine/Plugins/ChaosClothAssetEditor/Source/ChaosClothAssetEditorTools/Private/ChaosClothAsset/ClothMeshSelectionTool.h b/Engine/Plugins/ChaosClothAssetEditor/Source/ChaosClothAssetEditorTools/Private/ChaosClothAsset/ClothMeshSelectionTool.h
--- a/Engine/Plugins/ChaosClothAssetEditor/Source/ChaosClothAssetEditorTools/Private/ChaosClothAsset/ClothMeshSelectionTool.h
+++ b/Engine/Plugins/ChaosClothAssetEditor/Source/ChaosClothAssetEditorTools/Private/ChaosClothAsset/ClothMeshSelectionTool.h
@@ -28,7 +28,8 @@ enum class EClothMeshSelectionToolActions
 
 	GrowSelection,
 	ShrinkSelection,
-	FloodSelection
+	FloodSelection,
+	ClearSelection
 };
 
 
@@ -86,6 +87,12 @@ public:
 		PostAction(EClothMeshSelectionToolActions::FloodSelection);
 	}
 
+	UFUNCTION(CallInEditor, Category = Selection)
+	void ClearSelection()
+	{
+		PostAction(EClothMeshSelectionToolActions::ClearSelection);
+	}
+
 };
 
 UCLASS()
@@ -179,9 +186,5 @@ private:
 	virtual void ApplyAction(EClothMeshSelectionToolActions ActionType);
 
 	void ImportFromCollection(bool bImportFromSecondarySet);
-
-	void GrowSelection();
-	void ShrinkSelection();
-	void FloodSelection();
 };
 
