Cloth Editor Weight Map Paint Tool: replace the change tracker for undo/redo. Rather than using the generic FDynamicMeshChangeTracker, we instead store changes to the weight map values as they would be stored on the Weight Map Node. This should allow the same change to be applied to different mesh configurations (Sim2D/Sim3D). This is a prerequisite for being able to change view modes while keeping the tool running and maintaining the undo history.

diff --git a/Engine/Plugins/ChaosClothAssetEditor/Source/ChaosClothAssetEditorTools/Private/ChaosClothAsset/ClothWeightMapPaintTool.cpp b/Engine/Plugins/ChaosClothAssetEditor/Source/ChaosClothAssetEditorTools/Private/ChaosClothAsset/ClothWeightMapPaintTool.cpp
--- a/Engine/Plugins/ChaosClothAssetEditor/Source/ChaosClothAssetEditorTools/Private/ChaosClothAsset/ClothWeightMapPaintTool.cpp
+++ b/Engine/Plugins/ChaosClothAssetEditor/Source/ChaosClothAssetEditorTools/Private/ChaosClothAsset/ClothWeightMapPaintTool.cpp
@@ -772,6 +772,11 @@ void UClothEditorWeightMapPaintTool::OnEndStroke()
 	EndChange();
 }
 
+void UClothEditorWeightMapPaintTool::OnCancelStroke()
+{
+	GetActiveBrushOp()->CancelStroke();
+	ActiveChangeBuilder.Reset();
+}
 
 
 
@@ -1022,18 +1027,30 @@ bool UClothEditorWeightMapPaintTool::SyncMeshWithWeightBuffer(FDynamicMesh3* Mes
 
 			if (ROIWeightValueBuffer[k] != CurWeight)
 			{
+				const float NewValue = ROIWeightValueBuffer[k];
+
 				if (bHaveDynamicMeshToWeightConversion)
 				{
 					for (const int32 Idx : WeightToDynamicMesh[DynamicMeshToWeight[VertIdx]])
 					{
-						ActiveWeightEditChangeTracker->SaveVertexOneRingTriangles(Idx, true);
-						ActiveWeightMap->SetValue(Idx, &ROIWeightValueBuffer[k]);
+						float PrevValue;
+						ActiveWeightMap->GetValue(Idx, &PrevValue);
+
+						ensure(DynamicMeshToWeight[VertIdx] == MeshIndexToNodeIndex(Idx));
+						ActiveChangeBuilder->UpdateValue(DynamicMeshToWeight[VertIdx], PrevValue, NewValue);
+						
+						ActiveWeightMap->SetValue(Idx, &NewValue);
 					}
 				}
 				else
 				{
-				ActiveWeightEditChangeTracker->SaveVertexOneRingTriangles(VertIdx, true);
-				ActiveWeightMap->SetValue(VertIdx, &ROIWeightValueBuffer[k]);
+					float PrevValue;
+					ActiveWeightMap->GetValue(VertIdx, &PrevValue);
+
+					ensure(VertIdx == MeshIndexToNodeIndex(VertIdx));
+					ActiveChangeBuilder->UpdateValue(VertIdx, PrevValue, NewValue);
+
+					ActiveWeightMap->SetValue(VertIdx, &NewValue);
 				}
 				NumModified++;
 			}
@@ -1288,25 +1305,6 @@ void UClothEditorWeightMapPaintTool::ComputeGradient()
 		TempROIBuffer.Add(vid);
 	}
 
-	if (bHaveDynamicMeshToWeightConversion)
-	{
-		for (int32 vid : TempROIBuffer)
-		{
-			for (const int32 Idx : WeightToDynamicMesh[DynamicMeshToWeight[vid]])
-			{
-				ActiveWeightEditChangeTracker->SaveVertexOneRingTriangles(Idx, true);
-			}
-		}
-	}
-	else
-	{
-		for (int32 vid : TempROIBuffer)
-		{
-			ActiveWeightEditChangeTracker->SaveVertexOneRingTriangles(vid, true);
-		}
-	}
-
-
 	for (const int32 VertexIndex : TempROIBuffer)
 	{
 		const FVector3d Vert = Mesh->GetVertex(VertexIndex);
@@ -1337,18 +1335,30 @@ void UClothEditorWeightMapPaintTool::ComputeGradient()
 			}
 		}
 
-		const float Value = FMath::LerpStable(FilterProperties->GradientLowValue, FilterProperties->GradientHighValue, DistanceToLowSq / (DistanceToLowSq + DistanceToHighSq));
+		const float NewValue = FMath::LerpStable(FilterProperties->GradientLowValue, FilterProperties->GradientHighValue, DistanceToLowSq / (DistanceToLowSq + DistanceToHighSq));
 		if (bHaveDynamicMeshToWeightConversion)
 		{
 			for (const int32 Idx : WeightToDynamicMesh[DynamicMeshToWeight[VertexIndex]])
 			{
-				ActiveWeightMap->SetValue(Idx, &Value);
+				float PreviousValue;
+				ActiveWeightMap->GetValue(Idx, &PreviousValue);
+				
+				ensure(MeshIndexToNodeIndex(Idx) == DynamicMeshToWeight[VertexIndex]);
+				ActiveChangeBuilder->UpdateValue(DynamicMeshToWeight[VertexIndex], PreviousValue, NewValue);
+
+				ActiveWeightMap->SetValue(Idx, &NewValue);
 			}
 		}
 		else
 		{
-		ActiveWeightMap->SetValue(VertexIndex, &Value);
-	}
+			float PreviousValue;
+			ActiveWeightMap->GetValue(VertexIndex, &PreviousValue);
+
+			ensure(MeshIndexToNodeIndex(VertexIndex) == VertexIndex);
+			ActiveChangeBuilder->UpdateValue(VertexIndex, PreviousValue, NewValue);
+
+			ActiveWeightMap->SetValue(VertexIndex, &NewValue);
+		}
 	}
 
 	// update colors
@@ -1412,7 +1422,12 @@ void UClothEditorWeightMapPaintTool::SetVerticesToWeightMap(const TSet<int32>& V
 		{
 			for (const int32 Idx : WeightToDynamicMesh[DynamicMeshToWeight[vid]])
 			{
-				ActiveWeightEditChangeTracker->SaveVertexOneRingTriangles(Idx, true);
+				float PreviousValue;
+				ActiveWeightMap->GetValue(Idx, &PreviousValue);
+
+				ensure(MeshIndexToNodeIndex(Idx) == DynamicMeshToWeight[vid]);
+				ActiveChangeBuilder->UpdateValue(DynamicMeshToWeight[vid], PreviousValue, WeightValue);
+
 				ActiveWeightMap->SetValue(Idx, &WeightValue);
 			}
 		}
@@ -1421,10 +1436,12 @@ void UClothEditorWeightMapPaintTool::SetVerticesToWeightMap(const TSet<int32>& V
 	{
 		for (int32 vid : TempROIBuffer)
 		{
-			ActiveWeightEditChangeTracker->SaveVertexOneRingTriangles(vid, true);
-		}
-		for (int32 vid : TempROIBuffer)
-		{		
+			float PreviousValue;
+			ActiveWeightMap->GetValue(vid, &PreviousValue);
+			
+			ensure(MeshIndexToNodeIndex(vid) == vid);
+			ActiveChangeBuilder->UpdateValue(vid, PreviousValue, WeightValue);
+
 			ActiveWeightMap->SetValue(vid, &WeightValue);
 		}
 	}
@@ -1848,7 +1865,12 @@ void UClothEditorWeightMapPaintTool::FloodFillCurrentWeightAction()
 		{
 			for (const int32 Idx : WeightToDynamicMesh[DynamicMeshToWeight[vid]])
 			{
-				ActiveWeightEditChangeTracker->SaveVertexOneRingTriangles(Idx, true);
+				float PreviousValue;
+				ActiveWeightMap->GetValue(Idx, &PreviousValue);
+				
+				ensure(MeshIndexToNodeIndex(Idx) == DynamicMeshToWeight[vid]);
+				ActiveChangeBuilder->UpdateValue(DynamicMeshToWeight[vid], PreviousValue, SetWeightValue);
+
 				ActiveWeightMap->SetValue(Idx, &SetWeightValue);
 			}
 		}
@@ -1857,11 +1879,13 @@ void UClothEditorWeightMapPaintTool::FloodFillCurrentWeightAction()
 	{
 		for (int32 vid : TempROIBuffer)
 		{
-			ActiveWeightEditChangeTracker->SaveVertexOneRingTriangles(vid, true);
-		}
-		for (int32 vid : TempROIBuffer)
-		{
-			ActiveWeightMap->SetValue(vid, &SetWeightValue);
+			float PreviousValue;
+			ActiveWeightMap->GetValue(vid, &PreviousValue);
+
+			ensure(MeshIndexToNodeIndex(vid) == vid);
+			ActiveChangeBuilder->UpdateValue(vid, PreviousValue, SetWeightValue);
+
+			ActiveWeightMap->SetValue(vid, &SetWeightValue); 
 		}
 	}
 
@@ -1896,7 +1920,12 @@ void UClothEditorWeightMapPaintTool::ClearAllWeightsAction()
 		{
 			for (const int32 Idx : WeightToDynamicMesh[DynamicMeshToWeight[vid]])
 			{
-				ActiveWeightEditChangeTracker->SaveVertexOneRingTriangles(Idx, true);
+				float PreviousValue;
+				ActiveWeightMap->GetValue(Idx, &PreviousValue);
+
+				ensure(MeshIndexToNodeIndex(Idx) == DynamicMeshToWeight[vid]);
+				ActiveChangeBuilder->UpdateValue(DynamicMeshToWeight[vid], PreviousValue, SetWeightValue);
+
 				ActiveWeightMap->SetValue(Idx, &SetWeightValue);
 			}
 		}
@@ -1905,10 +1934,12 @@ void UClothEditorWeightMapPaintTool::ClearAllWeightsAction()
 	{
 		for (int32 vid : TempROIBuffer)
 		{
-			ActiveWeightEditChangeTracker->SaveVertexOneRingTriangles(vid, true);
-		}
-		for (int32 vid : TempROIBuffer)
-		{
+			float PreviousValue;
+			ActiveWeightMap->GetValue(vid, &PreviousValue);
+
+			ensure(MeshIndexToNodeIndex(vid) == vid);
+			ActiveChangeBuilder->UpdateValue(vid, PreviousValue, SetWeightValue);
+
 			ActiveWeightMap->SetValue(vid, &SetWeightValue);
 		}
 	}
@@ -1933,12 +1964,13 @@ void UClothEditorWeightMapPaintTool::InvertWeightsAction()
 
 	for (const int32 VertexID : Mesh->VertexIndicesItr())
 	{
-		ActiveWeightEditChangeTracker->SaveVertexOneRingTriangles(VertexID, true);
+		float PreviousValue;
+		ActiveWeightMap->GetValue(VertexID, &PreviousValue);
+		const float NewWeightValue = 1.0f - PreviousValue;
+		ActiveWeightMap->SetValue(VertexID, &NewWeightValue);
 
-		float WeightValue;
-		ActiveWeightMap->GetValue(VertexID, &WeightValue);
-		WeightValue = 1.0f - WeightValue;
-		ActiveWeightMap->SetValue(VertexID, &WeightValue);
+		const int32 NodeIndex = MeshIndexToNodeIndex(VertexID);
+		ActiveChangeBuilder->UpdateValue(NodeIndex, PreviousValue, NewWeightValue);
 	}
 
 	// update colors
@@ -1964,12 +1996,13 @@ void UClothEditorWeightMapPaintTool::MultiplyWeightsAction()
 
 	for (const int32 VertexID : Mesh->VertexIndicesItr())
 	{
-		ActiveWeightEditChangeTracker->SaveVertexOneRingTriangles(VertexID, true);
+		float PreviousValue;
+		ActiveWeightMap->GetValue(VertexID, &PreviousValue);
+		const float NewWeightValue = FMath::Clamp(WeightMultiplierValue * PreviousValue, 0.f, 1.f);
+		ActiveWeightMap->SetValue(VertexID, &NewWeightValue);
 
-		float WeightValue;
-		ActiveWeightMap->GetValue(VertexID, &WeightValue);
-		WeightValue = FMath::Clamp(WeightMultiplierValue * WeightValue, 0.f, 1.f);
-		ActiveWeightMap->SetValue(VertexID, &WeightValue);
+		const int32 NodeIndex = MeshIndexToNodeIndex(VertexID);
+		ActiveChangeBuilder->UpdateValue(NodeIndex, PreviousValue, NewWeightValue);
 	}
 
 	// update colors
@@ -2036,85 +2069,92 @@ void UClothEditorWeightMapPaintTool::UpdateSelectedNode()
 }
 
 
-//
-// Change Tracking
-//
+int32 UClothEditorWeightMapPaintTool::MeshIndexToNodeIndex(int32 MeshVertexIndex) const
+{
+	if (bHaveDynamicMeshToWeightConversion)
+	{
+		return DynamicMeshToWeight[MeshVertexIndex];
+	}
+	else
+	{
+		return MeshVertexIndex;
+	}
+}
 
-namespace ClothWeightPaintLocals
+void UClothEditorWeightMapPaintTool::UpdateMapValuesFromNodeValues(const TArray<int32>& Indices, const TArray<float>& Values)
 {
+	check(Indices.Num() == Values.Num());
 
-	/**
-	 * A wrapper change that applies a given change to the unwrap canonical mesh of an input, and uses that
-	 * to update the other views. Causes a broadcast of OnCanonicalModified.
-	 */
-	class  FClothWeightPaintMeshChange : public FToolCommandChange
+	for (int32 PairIndex = 0; PairIndex < Indices.Num(); ++PairIndex)
 	{
-	public:
-		FClothWeightPaintMeshChange(UDynamicMeshComponent* DynamicMeshComponentIn, TUniquePtr<FDynamicMeshChange> DynamicMeshChangeIn)
-			: DynamicMeshComponent(DynamicMeshComponentIn)
-			, DynamicMeshChange(MoveTemp(DynamicMeshChangeIn))
-		{
-			ensure(DynamicMeshComponentIn);
-			ensure(DynamicMeshChange);
-		};
+		const int32 BufferIndex = Indices[PairIndex];
+		const float Value = Values[PairIndex];
 
-		virtual void Apply(UObject* Object) override
+		if (bHaveDynamicMeshToWeightConversion)
 		{
-			DynamicMeshChange->Apply(DynamicMeshComponent->GetMesh(), false);			
+			for (const int32 MeshIndex : WeightToDynamicMesh[BufferIndex])
+			{
+				ActiveWeightMap->SetValue(MeshIndex, &Value);
+			}
 		}
-
-		virtual void Revert(UObject* Object) override
+		else
 		{
-			DynamicMeshChange->Apply(DynamicMeshComponent->GetMesh(), true);			
+			ActiveWeightMap->SetValue(BufferIndex, &Value);
 		}
+	}
 
-		virtual bool HasExpired(UObject* Object) const override
-		{
-			return !(DynamicMeshComponent.IsValid() && DynamicMeshChange);
-		}
+	if (Indices.Num() > 0)
+	{
+		UpdateVertexColorOverlay();
+		DynamicMeshComponent->FastNotifyVertexAttributesUpdated(EMeshRenderAttributeFlags::VertexColors);
+	}
+}
 
 
-		virtual FString ToString() const override
-		{
-			return TEXT("FClothWeightPaintMeshChange");
-		}
+//
+// Change Tracking
+//
 
-	protected:
-		TWeakObjectPtr<UDynamicMeshComponent> DynamicMeshComponent;
-		TUniquePtr<FDynamicMeshChange> DynamicMeshChange;
-	};
-}
 
 void UClothEditorWeightMapPaintTool::BeginChange()
 {
+	check(ActiveChangeBuilder == nullptr);
+	ActiveChangeBuilder = MakeUnique<TIndexedValuesChangeBuilder<float, FNodeBufferWeightChange>>();
+	ActiveChangeBuilder->BeginNewChange();
 
-	check(ActiveWeightEditChangeTracker == nullptr);
-	
-	ActiveWeightEditChangeTracker = MakeUnique<FDynamicMeshChangeTracker>(GetSculptMesh());
-	ActiveWeightEditChangeTracker->BeginChange();
 	LongTransactions.Open(LOCTEXT("WeightPaintChange", "Weight Stroke"), GetToolManager());
 }
 
 void UClothEditorWeightMapPaintTool::EndChange()
 {
-	check(ActiveWeightEditChangeTracker);
+	check(ActiveChangeBuilder);
 
 	bAnyChangeMade = true;
 
-	TUniquePtr<FDynamicMeshChange> EditResult = ActiveWeightEditChangeTracker->EndChange();
+	TUniquePtr<FNodeBufferWeightChange> EditResult = ActiveChangeBuilder->ExtractResult();
+	ActiveChangeBuilder.Reset();
+
+	EditResult->ApplyFunction = [](UObject* Object, const int32& AttribIndex, const TArray<int32>& Indices, const TArray<float>& Values)
+	{
+		UClothEditorWeightMapPaintTool* const Tool = CastChecked<UClothEditorWeightMapPaintTool>(Object);
+		Tool->UpdateMapValuesFromNodeValues(Indices, Values);
+	};
 
-	TUniquePtr<ClothWeightPaintLocals::FClothWeightPaintMeshChange> ClothWeightPaintMeshChange =
-		MakeUnique<ClothWeightPaintLocals::FClothWeightPaintMeshChange>(DynamicMeshComponent.Get(), MoveTemp(EditResult));
-	ActiveWeightEditChangeTracker = nullptr;
+	EditResult->RevertFunction = [](UObject* Object, const int32& AttribIndex, const TArray<int32>& Indices, const TArray<float>& Values)
+	{
+		UClothEditorWeightMapPaintTool* const Tool = CastChecked<UClothEditorWeightMapPaintTool>(Object);
+		Tool->UpdateMapValuesFromNodeValues(Indices, Values);
+	};
 
-	TUniquePtr<TWrappedToolCommandChange<ClothWeightPaintLocals::FClothWeightPaintMeshChange>> NewChange = MakeUnique<TWrappedToolCommandChange<ClothWeightPaintLocals::FClothWeightPaintMeshChange>>();
-	NewChange->WrappedChange = MoveTemp(ClothWeightPaintMeshChange);
+	TUniquePtr<TWrappedToolCommandChange<FNodeBufferWeightChange>> NewChange = MakeUnique<TWrappedToolCommandChange<FNodeBufferWeightChange>>();
+	NewChange->WrappedChange = MoveTemp(EditResult);
 	NewChange->BeforeModify = [this](bool bRevert)
 	{
 		this->WaitForPendingUndoRedo();
 	};
 
-	GetToolManager()->EmitObjectChange(DynamicMeshComponent, MoveTemp(NewChange), LOCTEXT("WeightPaintChange", "Weight Stroke"));
+	GetToolManager()->EmitObjectChange(this, MoveTemp(NewChange), LOCTEXT("VertexWeightChange", "Weight Stroke"));
+
 	LongTransactions.Close(GetToolManager());
 }
 
diff --git a/Engine/Plugins/ChaosClothAssetEditor/Source/ChaosClothAssetEditorTools/Private/ChaosClothAsset/ClothWeightMapPaintTool.h b/Engine/Plugins/ChaosClothAssetEditor/Source/ChaosClothAssetEditorTools/Private/ChaosClothAsset/ClothWeightMapPaintTool.h
--- a/Engine/Plugins/ChaosClothAssetEditor/Source/ChaosClothAssetEditorTools/Private/ChaosClothAsset/ClothWeightMapPaintTool.h
+++ b/Engine/Plugins/ChaosClothAssetEditor/Source/ChaosClothAssetEditorTools/Private/ChaosClothAsset/ClothWeightMapPaintTool.h
@@ -17,12 +17,13 @@
 #include "DynamicMesh/DynamicMeshAABBTree3.h"
 #include "DynamicMesh/DynamicMeshOctree3.h"
 #include "DynamicMesh/DynamicMeshAttributeSet.h"
-#include "DynamicMesh/DynamicMeshChangeTracker.h"
 #include "DynamicMesh/MeshNormals.h"
 #include "TransformTypes.h"
 #include "ToolDataVisualizer.h"
 #include "GroupTopology.h"
 #include "Dataflow/DataflowObjectInterface.h"
+#include "Changes/IndexedAttributeChange.h"
+
 #include "ClothWeightMapPaintTool.generated.h"
 
 class UMeshElementsVisualizer;
@@ -377,6 +378,7 @@ protected:
 	virtual int32 FindHitTargetMeshTriangle(const FRay3d& LocalRay) override;
 
 	virtual void OnBeginStroke(const FRay& WorldRay) override;
+	virtual void OnCancelStroke() override;
 	virtual void OnEndStroke() override;
 
 	virtual TUniquePtr<FMeshSculptBrushOp>& GetActiveBrushOp();
@@ -539,10 +541,30 @@ protected:
 	bool SyncMeshWithWeightBuffer(FDynamicMesh3* Mesh);
 	bool SyncWeightBufferWithMesh(const FDynamicMesh3* Mesh);
 
-	TUniquePtr<UE::Geometry::FDynamicMeshChangeTracker> ActiveWeightEditChangeTracker;
+
+	// Undo/Redo change support
+
+	class FNodeBufferWeightChange : public TCustomIndexedValuesChange<float, int32>
+	{
+	public:
+		virtual FString ToString() const override
+		{
+			return FString(TEXT("Cloth Vertex Scalar Weight Change"));
+		}
+	};
+	TUniquePtr<TIndexedValuesChangeBuilder<float, FNodeBufferWeightChange>> ActiveChangeBuilder;
 	void BeginChange();
 	void EndChange();
 
+	// The corresponding FChaosClothAssetWeightMapNode has a buffer of scalar weights. Depending on bHaveDynamicMeshToWeightConversion, the 
+	// indexing of this buffer might be different than the mesh vertex indexing. This function returns the corresponding index in the node buffer for
+	// a given mesh vertex index.
+	int32 MeshIndexToNodeIndex(int32 MeshVertexIndex) const;
+
+	// Update the active weight map attribute from the map of indices/values. The inputs are given in "node buffer format" (see function above).
+	void UpdateMapValuesFromNodeValues(const TArray<int32>& Indices, const TArray<float>& Values);
+
+
 	FColor GetColorForWeightValue(double WeightValue);
 
 	TArray<FVector3d> TriNormals;
