Chaos Cloth Asset
- Removed the terminal node inheritance from the selection node.
- Moved the selection function to the node detail view.
- Removed experimental from the proxy deformer.
- Made all deprecated node's structs also deprecated.
- Renamed the attribute node V2 into v2 (with a lower case v) to match the naming convention used on the other versioned node names.

[RN] Cloth Asset - New Proxy Deformer node version using two single selections instead of one dual selection per filter. Dual selections are now deprecated but can be still be imported into single selections usign the Selection node UI.

diff --git a/Engine/Plugins/ChaosClothAssetEditor/Source/ChaosClothAssetDataflowNodes/Private/ChaosClothAsset/AddWeightMapNode.cpp b/Engine/Plugins/ChaosClothAssetEditor/Source/ChaosClothAssetDataflowNodes/Private/ChaosClothAsset/AddWeightMapNode.cpp
--- a/Engine/Plugins/ChaosClothAssetEditor/Source/ChaosClothAssetDataflowNodes/Private/ChaosClothAsset/AddWeightMapNode.cpp
+++ b/Engine/Plugins/ChaosClothAssetEditor/Source/ChaosClothAssetDataflowNodes/Private/ChaosClothAsset/AddWeightMapNode.cpp
@@ -81,7 +81,9 @@ namespace UE::Chaos::ClothAsset::Private
 			case EChaosClothAssetWeightMapOverrideType::ReplaceChanged:
 				if (InputMap[Index] == FinalValues[Index])
 				{
+PRAGMA_DISABLE_DEPRECATION_WARNINGS
 					SourceVertexWeights[Index] = FChaosClothAssetAddWeightMapNode::ReplaceChangedPassthroughValue;
+PRAGMA_ENABLE_DEPRECATION_WARNINGS
 				}
 				else
 				{
@@ -103,8 +105,10 @@ namespace UE::Chaos::ClothAsset::Private
 		{
 			for (int32 Index = 0; Index < EndWeightIndex; ++Index)
 			{
+PRAGMA_DISABLE_DEPRECATION_WARNINGS
 				FinalOutputMap[Index] = FMath::Clamp(SourceVertexWeights[Index] == FChaosClothAssetAddWeightMapNode::ReplaceChangedPassthroughValue ?
 					0.f : SourceVertexWeights[Index], 0.f, 1.f);
+PRAGMA_ENABLE_DEPRECATION_WARNINGS
 			}
 			return;
 		}
@@ -112,7 +116,9 @@ namespace UE::Chaos::ClothAsset::Private
 		check(InputMap.Num() == FinalOutputMap.Num());
 		for (int32 Index = 0; Index < EndWeightIndex; ++Index)
 		{
+PRAGMA_DISABLE_DEPRECATION_WARNINGS
 			if(SourceVertexWeights[Index] == FChaosClothAssetAddWeightMapNode::ReplaceChangedPassthroughValue)
+PRAGMA_ENABLE_DEPRECATION_WARNINGS
 			{
 				// This value is only set when OverrideType == ReplaceChanged, but it's possible the override type changed.
 				FinalOutputMap[Index] = FMath::Clamp(InputMap[Index], 0.f, 1.f);
@@ -359,7 +365,9 @@ class FChaosClothAssetAddWeightMapNode::FWeightMapNodeChange final : public FToo
 
 public: 
 
+PRAGMA_DISABLE_DEPRECATION_WARNINGS
 	FWeightMapNodeChange(const FChaosClothAssetAddWeightMapNode& Node) :
+PRAGMA_ENABLE_DEPRECATION_WARNINGS
 		NodeGuid(Node.GetGuid()),
 		SavedWeights(Node.GetVertexWeights()),
 		SavedRenderWeights(Node.GetRenderVertexWeights()),
@@ -400,7 +408,9 @@ private:
 		{
 			if (const TSharedPtr<FDataflowNode> BaseNode = Dataflow->GetDataflow()->FindBaseNode(NodeGuid))
 			{
+PRAGMA_DISABLE_DEPRECATION_WARNINGS
 				if (FChaosClothAssetAddWeightMapNode* const Node = BaseNode->AsType<FChaosClothAssetAddWeightMapNode>())
+PRAGMA_ENABLE_DEPRECATION_WARNINGS
 				{
 					Swap(Node->GetVertexWeights(), SavedWeights);
 					Swap(Node->GetRenderVertexWeights(), SavedRenderWeights);
@@ -414,9 +424,11 @@ private:
 	}
 };
 
+PRAGMA_DISABLE_DEPRECATION_WARNINGS
 TUniquePtr<FToolCommandChange> FChaosClothAssetAddWeightMapNode::MakeWeightMapNodeChange(const FChaosClothAssetAddWeightMapNode& Node)
 {
 	return MakeUnique<FChaosClothAssetAddWeightMapNode::FWeightMapNodeChange>(Node);
 }
+PRAGMA_ENABLE_DEPRECATION_WARNINGS
 
 #undef LOCTEXT_NAMESPACE
diff --git a/Engine/Plugins/ChaosClothAssetEditor/Source/ChaosClothAssetDataflowNodes/Private/ChaosClothAsset/AttributeNode.cpp b/Engine/Plugins/ChaosClothAssetEditor/Source/ChaosClothAssetDataflowNodes/Private/ChaosClothAsset/AttributeNode.cpp
--- a/Engine/Plugins/ChaosClothAssetEditor/Source/ChaosClothAssetDataflowNodes/Private/ChaosClothAsset/AttributeNode.cpp
+++ b/Engine/Plugins/ChaosClothAssetEditor/Source/ChaosClothAssetDataflowNodes/Private/ChaosClothAsset/AttributeNode.cpp
@@ -11,7 +11,7 @@
 
 #define LOCTEXT_NAMESPACE "ChaosClothAssetAttributeNode"
 
-FChaosClothAssetAttributeNode_V2::FChaosClothAssetAttributeNode_V2(const Dataflow::FNodeParameters& InParam, FGuid InGuid)
+FChaosClothAssetAttributeNode_v2::FChaosClothAssetAttributeNode_v2(const Dataflow::FNodeParameters& InParam, FGuid InGuid)
 	: FDataflowNode(InParam, InGuid)
 {
 	RegisterInputConnection(&Collection);
@@ -20,7 +20,7 @@ FChaosClothAssetAttributeNode_V2::FChaosClothAssetAttributeNode_V2(const Dataflo
 	RegisterOutputConnection(&Name.StringValue, &Name.StringValue, GET_MEMBER_NAME_CHECKED(FChaosClothAssetConnectableIOStringValue, StringValue));
 }
 
-void FChaosClothAssetAttributeNode_V2::Evaluate(Dataflow::FContext& Context, const FDataflowOutput* Out) const
+void FChaosClothAssetAttributeNode_v2::Evaluate(Dataflow::FContext& Context, const FDataflowOutput* Out) const
 {
 	if (Out->IsA<FManagedArrayCollection>(&Collection))
 	{
@@ -90,7 +90,7 @@ void FChaosClothAssetAttributeNode_V2::Evaluate(Dataflow::FContext& Context, con
 	}
 }
 
-void FChaosClothAssetAttributeNode_V2::OnSelected(Dataflow::FContext& Context)
+void FChaosClothAssetAttributeNode_v2::OnSelected(Dataflow::FContext& Context)
 {
 	using namespace UE::Chaos::ClothAsset;
 
@@ -111,7 +111,7 @@ void FChaosClothAssetAttributeNode_V2::OnSelected(Dataflow::FContext& Context)
 	}
 }
 
-void FChaosClothAssetAttributeNode_V2::OnDeselected()
+void FChaosClothAssetAttributeNode_v2::OnDeselected()
 {
 	// Clean up, to avoid another toolkit picking up the wrong context evaluation
 	CachedCollectionGroupNames.Reset();
diff --git a/Engine/Plugins/ChaosClothAssetEditor/Source/ChaosClothAssetDataflowNodes/Private/ChaosClothAsset/ChaosClothAssetDataflowNodesModule.cpp b/Engine/Plugins/ChaosClothAssetEditor/Source/ChaosClothAssetDataflowNodes/Private/ChaosClothAsset/ChaosClothAssetDataflowNodesModule.cpp
--- a/Engine/Plugins/ChaosClothAssetEditor/Source/ChaosClothAssetDataflowNodes/Private/ChaosClothAsset/ChaosClothAssetDataflowNodesModule.cpp
+++ b/Engine/Plugins/ChaosClothAssetEditor/Source/ChaosClothAssetDataflowNodes/Private/ChaosClothAsset/ChaosClothAssetDataflowNodesModule.cpp
@@ -80,23 +80,20 @@ namespace UE::Chaos::ClothAsset
 	{
 		static void RegisterDataflowNodes()
 		{
+			DATAFLOW_NODE_REGISTER_CREATION_FACTORY_NODE_COLORS_BY_CATEGORY("Cloth", FColorScheme::NodeHeader, FColorScheme::NodeBody);
 			DATAFLOW_NODE_REGISTER_CREATION_FACTORY(FChaosClothAssetAddStitchNode);
-			DATAFLOW_NODE_REGISTER_CREATION_FACTORY(FChaosClothAssetAddWeightMapNode);
 			DATAFLOW_NODE_REGISTER_CREATION_FACTORY(FChaosClothAssetWeightMapNode);
-			DATAFLOW_NODE_REGISTER_CREATION_FACTORY(FChaosClothAssetAttributeNode);
-			DATAFLOW_NODE_REGISTER_CREATION_FACTORY(FChaosClothAssetAttributeNode_V2);
+			DATAFLOW_NODE_REGISTER_CREATION_FACTORY(FChaosClothAssetAttributeNode_v2);
 			DATAFLOW_NODE_REGISTER_CREATION_FACTORY(FChaosClothAssetBindToRootBoneNode);
 			DATAFLOW_NODE_REGISTER_CREATION_FACTORY(FChaosClothAssetCopySimulationToRenderMeshNode);
 			DATAFLOW_NODE_REGISTER_CREATION_FACTORY(FChaosClothAssetDatasmithImportNode);
 			DATAFLOW_NODE_REGISTER_CREATION_FACTORY(FChaosClothAssetDeleteElementNode);
 			DATAFLOW_NODE_REGISTER_CREATION_FACTORY(FChaosClothAssetImportNode);
-			DATAFLOW_NODE_REGISTER_CREATION_FACTORY(FChaosClothAssetMergeClothCollectionsNode);
 			DATAFLOW_NODE_REGISTER_CREATION_FACTORY(FChaosClothAssetMergeClothCollectionsNode_v2);
-			DATAFLOW_NODE_REGISTER_CREATION_FACTORY(FChaosClothAssetProxyDeformerNode);
 			DATAFLOW_NODE_REGISTER_CREATION_FACTORY(FChaosClothAssetProxyDeformerNode_v2);
 			DATAFLOW_NODE_REGISTER_CREATION_FACTORY(FChaosClothAssetRemeshNode);
 			DATAFLOW_NODE_REGISTER_CREATION_FACTORY(FChaosClothAssetReverseNormalsNode);
-			DATAFLOW_NODE_REGISTER_CREATION_FACTORY(FChaosClothAssetSelectionNode);
+			DATAFLOW_NODE_REGISTER_CREATION_FACTORY(FChaosClothAssetSelectionNode_v2);
 			DATAFLOW_NODE_REGISTER_CREATION_FACTORY(FChaosClothAssetSelectionToIntMapNode);
 			DATAFLOW_NODE_REGISTER_CREATION_FACTORY(FChaosClothAssetSelectionToWeightMapNode);
 			DATAFLOW_NODE_REGISTER_CREATION_FACTORY(FChaosClothAssetSetPhysicsAssetNode);
@@ -111,7 +108,6 @@ namespace UE::Chaos::ClothAsset
 			DATAFLOW_NODE_REGISTER_CREATION_FACTORY(FChaosClothAssetSimulationDampingConfigNode);
 			DATAFLOW_NODE_REGISTER_CREATION_FACTORY(FChaosClothAssetSimulationDefaultConfigNode);
 			DATAFLOW_NODE_REGISTER_CREATION_FACTORY(FChaosClothAssetSimulationGravityConfigNode);
-			DATAFLOW_NODE_REGISTER_CREATION_FACTORY(FChaosClothAssetSimulationLongRangeAttachmentConfigNode);
 			DATAFLOW_NODE_REGISTER_CREATION_FACTORY(FChaosClothAssetSimulationLongRangeAttachmentConfigNode_v2);
 			DATAFLOW_NODE_REGISTER_CREATION_FACTORY(FChaosClothAssetSimulationMassConfigNode);
 			DATAFLOW_NODE_REGISTER_CREATION_FACTORY(FChaosClothAssetSimulationMaxDistanceConfigNode);
@@ -132,19 +128,27 @@ namespace UE::Chaos::ClothAsset
 			DATAFLOW_NODE_REGISTER_CREATION_FACTORY(FChaosClothAssetSimulationXPBDBendingElementConfigNode);
 			DATAFLOW_NODE_REGISTER_CREATION_FACTORY(FChaosClothAssetSimulationXPBDBendingSpringConfigNode);
 			DATAFLOW_NODE_REGISTER_CREATION_FACTORY(FChaosClothAssetSimulationXPBDEdgeSpringConfigNode);
-			DATAFLOW_NODE_REGISTER_CREATION_FACTORY(FChaosClothAssetSkeletalMeshImportNode);
 			DATAFLOW_NODE_REGISTER_CREATION_FACTORY(FChaosClothAssetSkeletalMeshImportNode_v2);
 			DATAFLOW_NODE_REGISTER_CREATION_FACTORY(FChaosClothAssetSkinningBlendNode);
-			DATAFLOW_NODE_REGISTER_CREATION_FACTORY(FChaosClothAssetStaticMeshImportNode);
 			DATAFLOW_NODE_REGISTER_CREATION_FACTORY(FChaosClothAssetStaticMeshImportNode_v2);
-			DATAFLOW_NODE_REGISTER_CREATION_FACTORY(FChaosClothAssetTerminalNode);
 			DATAFLOW_NODE_REGISTER_CREATION_FACTORY(FChaosClothAssetTerminalNode_v2);
 			DATAFLOW_NODE_REGISTER_CREATION_FACTORY(FChaosClothAssetTransferSkinWeightsNode);
 			DATAFLOW_NODE_REGISTER_CREATION_FACTORY(FChaosClothAssetTransformPositionsNode);
 			DATAFLOW_NODE_REGISTER_CREATION_FACTORY(FChaosClothAssetTransformUVsNode);
 			DATAFLOW_NODE_REGISTER_CREATION_FACTORY(FChaosClothAssetUSDImportNode);
 			DATAFLOW_NODE_REGISTER_CREATION_FACTORY(FChaosClothAssetWeightMapToSelectionNode);
-			DATAFLOW_NODE_REGISTER_CREATION_FACTORY_NODE_COLORS_BY_CATEGORY("Cloth", FColorScheme::NodeHeader, FColorScheme::NodeBody);
+			// Deprecated nodes
+PRAGMA_DISABLE_DEPRECATION_WARNINGS
+			DATAFLOW_NODE_REGISTER_CREATION_FACTORY(FChaosClothAssetAttributeNode);
+			DATAFLOW_NODE_REGISTER_CREATION_FACTORY(FChaosClothAssetAddWeightMapNode);
+			DATAFLOW_NODE_REGISTER_CREATION_FACTORY(FChaosClothAssetMergeClothCollectionsNode);
+			DATAFLOW_NODE_REGISTER_CREATION_FACTORY(FChaosClothAssetProxyDeformerNode);
+			DATAFLOW_NODE_REGISTER_CREATION_FACTORY(FChaosClothAssetSelectionNode);
+			DATAFLOW_NODE_REGISTER_CREATION_FACTORY(FChaosClothAssetSimulationLongRangeAttachmentConfigNode);
+			DATAFLOW_NODE_REGISTER_CREATION_FACTORY(FChaosClothAssetSkeletalMeshImportNode);
+			DATAFLOW_NODE_REGISTER_CREATION_FACTORY(FChaosClothAssetStaticMeshImportNode);
+			DATAFLOW_NODE_REGISTER_CREATION_FACTORY(FChaosClothAssetTerminalNode);
+PRAGMA_ENABLE_DEPRECATION_WARNINGS
 		}
 
 
@@ -300,10 +304,10 @@ namespace UE::Chaos::ClothAsset
 				PropertyModule->RegisterCustomPropertyTypeLayout(FChaosClothAssetWeightedValueNonAnimatable::StaticStruct()->GetFName(), FOnGetPropertyTypeCustomizationInstance::CreateStatic(&FWeightedValueCustomization::MakeInstance));
 				PropertyModule->RegisterCustomPropertyTypeLayout(FChaosClothAssetWeightedValueNonAnimatableNoLowHighRange::StaticStruct()->GetFName(), FOnGetPropertyTypeCustomizationInstance::CreateStatic(&FWeightedValueCustomization::MakeInstance));
 				PropertyModule->RegisterCustomPropertyTypeLayout(FChaosClothAssetWeightedValueOverride::StaticStruct()->GetFName(), FOnGetPropertyTypeCustomizationInstance::CreateStatic(&FWeightedValueCustomization::MakeInstance));
-				PropertyModule->RegisterCustomPropertyTypeLayout(FChaosClothAssetConnectableStringValue::StaticStruct()->GetFName(), FOnGetPropertyTypeCustomizationInstance::CreateStatic(&FWeightedValueCustomization::MakeInstance));
+				PropertyModule->RegisterCustomPropertyTypeLayout(FChaosClothAssetConnectableStringValue::StaticStruct()->GetFName(), FOnGetPropertyTypeCustomizationInstance::CreateStatic(&FConnectableValueCustomization::MakeInstance));
 				PropertyModule->RegisterCustomPropertyTypeLayout(FChaosClothAssetConnectableIStringValue::StaticStruct()->GetFName(), FOnGetPropertyTypeCustomizationInstance::CreateStatic(&FConnectableValueCustomization::MakeInstance));
-				PropertyModule->RegisterCustomPropertyTypeLayout(FChaosClothAssetConnectableOStringValue::StaticStruct()->GetFName(), FOnGetPropertyTypeCustomizationInstance::CreateStatic(&FWeightedValueCustomization::MakeInstance));
-				PropertyModule->RegisterCustomPropertyTypeLayout(FChaosClothAssetConnectableIOStringValue::StaticStruct()->GetFName(), FOnGetPropertyTypeCustomizationInstance::CreateStatic(&FWeightedValueCustomization::MakeInstance));
+				PropertyModule->RegisterCustomPropertyTypeLayout(FChaosClothAssetConnectableOStringValue::StaticStruct()->GetFName(), FOnGetPropertyTypeCustomizationInstance::CreateStatic(&FConnectableValueCustomization::MakeInstance));
+				PropertyModule->RegisterCustomPropertyTypeLayout(FChaosClothAssetConnectableIOStringValue::StaticStruct()->GetFName(), FOnGetPropertyTypeCustomizationInstance::CreateStatic(&FConnectableValueCustomization::MakeInstance));
 				PropertyModule->RegisterCustomPropertyTypeLayout(FChaosClothAssetNodeSelectionGroup::StaticStruct()->GetFName(), FOnGetPropertyTypeCustomizationInstance::CreateStatic(&FSelectionGroupCustomization::MakeInstance));
 				PropertyModule->RegisterCustomPropertyTypeLayout(FChaosClothAssetImportFilePath::StaticStruct()->GetFName(), FOnGetPropertyTypeCustomizationInstance::CreateStatic(&FImportFilePathCustomization::MakeInstance));
 				PropertyModule->RegisterCustomPropertyTypeLayout(FChaosClothAssetTerminalNodeRefreshAsset::StaticStruct()->GetFName(), FOnGetPropertyTypeCustomizationInstance::CreateStatic(&FTerminalNodeRefreshAssetCustomization::MakeInstance));
diff --git a/Engine/Plugins/ChaosClothAssetEditor/Source/ChaosClothAssetDataflowNodes/Private/ChaosClothAsset/MergeClothCollectionsNode.cpp b/Engine/Plugins/ChaosClothAssetEditor/Source/ChaosClothAssetDataflowNodes/Private/ChaosClothAsset/MergeClothCollectionsNode.cpp
--- a/Engine/Plugins/ChaosClothAssetEditor/Source/ChaosClothAssetDataflowNodes/Private/ChaosClothAsset/MergeClothCollectionsNode.cpp
+++ b/Engine/Plugins/ChaosClothAssetEditor/Source/ChaosClothAssetDataflowNodes/Private/ChaosClothAsset/MergeClothCollectionsNode.cpp
@@ -739,8 +739,9 @@ TArray<const FManagedArrayCollection*> FChaosClothAssetMergeClothCollectionsNode
 	return Collections;
 }
 
-
+PRAGMA_DISABLE_DEPRECATION_WARNINGS  // Unexpected deprecation message on some platforms otherwise
 const FManagedArrayCollection* FChaosClothAssetMergeClothCollectionsNode::GetCollection(int32 Index) const
+PRAGMA_ENABLE_DEPRECATION_WARNINGS
 {
 	switch (Index)
 	{
diff --git a/Engine/Plugins/ChaosClothAssetEditor/Source/ChaosClothAssetDataflowNodes/Private/ChaosClothAsset/RemeshNode.cpp b/Engine/Plugins/ChaosClothAssetEditor/Source/ChaosClothAssetDataflowNodes/Private/ChaosClothAsset/RemeshNode.cpp
--- a/Engine/Plugins/ChaosClothAssetEditor/Source/ChaosClothAssetDataflowNodes/Private/ChaosClothAsset/RemeshNode.cpp
+++ b/Engine/Plugins/ChaosClothAssetEditor/Source/ChaosClothAssetDataflowNodes/Private/ChaosClothAsset/RemeshNode.cpp
@@ -1261,6 +1261,7 @@ void FChaosClothAssetRemeshNode::RebuildTopologyDependentSimData(const TSharedRe
 	// Reconstruct long-range attachments
 
 	// v1 (weight map)
+PRAGMA_DISABLE_DEPRECATION_WARNINGS
 	const FString FixedEndWeightMapString = GET_MEMBER_NAME_STRING_CHECKED(FChaosClothAssetSimulationLongRangeAttachmentConfigNode, FixedEndWeightMap);
 	if (InProperties.GetKeyIndex(FixedEndWeightMapString) != INDEX_NONE)
 	{
@@ -1271,6 +1272,7 @@ void FChaosClothAssetRemeshNode::RebuildTopologyDependentSimData(const TSharedRe
 
 		UE::Chaos::ClothAsset::FClothEngineTools::GenerateTethers(OutClothCollection, FixedEndWeightMap, bUseGeodesicTethers);
 	}
+PRAGMA_ENABLE_DEPRECATION_WARNINGS
 
 	// v2 (vertex set)
 	const FString FixedEndSetString = GET_MEMBER_NAME_STRING_CHECKED(FChaosClothAssetSimulationLongRangeAttachmentConfigNode_v2, FixedEndSet);
diff --git a/Engine/Plugins/ChaosClothAssetEditor/Source/ChaosClothAssetDataflowNodes/Private/ChaosClothAsset/SelectionGroupCustomization.cpp b/Engine/Plugins/ChaosClothAssetEditor/Source/ChaosClothAssetDataflowNodes/Private/ChaosClothAsset/SelectionGroupCustomization.cpp
--- a/Engine/Plugins/ChaosClothAssetEditor/Source/ChaosClothAssetDataflowNodes/Private/ChaosClothAsset/SelectionGroupCustomization.cpp
+++ b/Engine/Plugins/ChaosClothAssetEditor/Source/ChaosClothAssetDataflowNodes/Private/ChaosClothAsset/SelectionGroupCustomization.cpp
@@ -3,6 +3,7 @@
 #include "ChaosClothAsset/SelectionGroupCustomization.h"
 #include "ChaosClothAsset/AttributeNode.h"
 #include "ChaosClothAsset/ClothDataflowTools.h"  // For MakeCollectionName
+#include "ChaosClothAsset/CollectionClothFacade.h"
 #include "ChaosClothAsset/DeleteElementNode.h"
 #include "ChaosClothAsset/SelectionNode.h"
 #include "Widgets/Input/SEditableTextBox.h"
@@ -110,21 +111,28 @@ namespace UE::Chaos::ClothAsset
 		GroupNames.Reset();
 
 		TArray<FName> CollectionGroupNames;
-		if (const FChaosClothAssetSelectionNode* const SelectionNode = GetOwnerStruct<FChaosClothAssetSelectionNode>())
+		if (const FChaosClothAssetSelectionNode_v2* const SelectionNode_v2 = GetOwnerStruct<FChaosClothAssetSelectionNode_v2>())
 		{
-			CollectionGroupNames = SelectionNode->GetCachedCollectionGroupNames();
-
+			CollectionGroupNames = SelectionNode_v2->GetCachedCollectionGroupNames();
 		}
-		else if (const FChaosClothAssetDeleteElementNode* const DeleteNode =
-			GetOwnerStruct<FChaosClothAssetDeleteElementNode>())
+		else if (const FChaosClothAssetDeleteElementNode* const DeleteNode = GetOwnerStruct<FChaosClothAssetDeleteElementNode>())
 		{
 			CollectionGroupNames = DeleteNode->GetCachedCollectionGroupNames();
 		}
-		else if (const FChaosClothAssetAttributeNode* const AttributeNode =
-			GetOwnerStruct<FChaosClothAssetAttributeNode>())
+		else if (const FChaosClothAssetAttributeNode_v2* const AttributeNode_v2 = GetOwnerStruct<FChaosClothAssetAttributeNode_v2>())
+		{
+			CollectionGroupNames = AttributeNode_v2->GetCachedCollectionGroupNames();
+		}
+PRAGMA_DISABLE_DEPRECATION_WARNINGS
+		else if (const FChaosClothAssetSelectionNode* const SelectionNode = GetOwnerStruct<FChaosClothAssetSelectionNode>())
+		{
+			CollectionGroupNames = SelectionNode->GetCachedCollectionGroupNames();
+		}
+		else if (const FChaosClothAssetAttributeNode* const AttributeNode = GetOwnerStruct<FChaosClothAssetAttributeNode>())
 		{
 			CollectionGroupNames = AttributeNode->GetCachedCollectionGroupNames();
 		}
+PRAGMA_ENABLE_DEPRECATION_WARNINGS
 
 		for (const FName& CollectionGroupName : CollectionGroupNames)
 		{
diff --git a/Engine/Plugins/ChaosClothAssetEditor/Source/ChaosClothAssetDataflowNodes/Private/ChaosClothAsset/SelectionNode.cpp b/Engine/Plugins/ChaosClothAssetEditor/Source/ChaosClothAssetDataflowNodes/Private/ChaosClothAsset/SelectionNode.cpp
--- a/Engine/Plugins/ChaosClothAssetEditor/Source/ChaosClothAssetDataflowNodes/Private/ChaosClothAsset/SelectionNode.cpp
+++ b/Engine/Plugins/ChaosClothAssetEditor/Source/ChaosClothAssetDataflowNodes/Private/ChaosClothAsset/SelectionNode.cpp
@@ -18,214 +18,496 @@
 
 namespace UE::Chaos::ClothAsset::Private
 {
-void ConvertWeightMapToVertexSelection(const TArray<float>& WeightMap, const float TransferSelectionThreshold, TSet<int32>& OutSelection)
-{
-	OutSelection.Reset();
-	for (int32 Index = 0; Index < WeightMap.Num(); ++Index)
+	void ConvertWeightMapToVertexSelection(const TArray<float>& WeightMap, const float TransferSelectionThreshold, TSet<int32>& OutSelection)
 	{
-		if (WeightMap[Index] >= TransferSelectionThreshold)
+		OutSelection.Reset();
+		for (int32 Index = 0; Index < WeightMap.Num(); ++Index)
 		{
-			OutSelection.Add(Index);
+			if (WeightMap[Index] >= TransferSelectionThreshold)
+			{
+				OutSelection.Add(Index);
+			}
 		}
 	}
-}
 
-void ConvertWeightMapToFaceSelection(const TArray<float>& WeightMap, const float TransferSelectionThreshold, const TConstArrayView<FIntVector3>& Indices, TSet<int32>& OutSelection)
-{
-	OutSelection.Reset();
-	for (int32 FaceIndex = 0; FaceIndex < Indices.Num(); ++FaceIndex)
+	void ConvertWeightMapToFaceSelection(const TArray<float>& WeightMap, const float TransferSelectionThreshold, const TConstArrayView<FIntVector3>& Indices, TSet<int32>& OutSelection)
 	{
-		const FIntVector& Element = Indices[FaceIndex];
-		if (WeightMap[Element[0]] >= TransferSelectionThreshold &&
-			WeightMap[Element[1]] >= TransferSelectionThreshold &&
-			WeightMap[Element[2]] >= TransferSelectionThreshold)
+		OutSelection.Reset();
+		for (int32 FaceIndex = 0; FaceIndex < Indices.Num(); ++FaceIndex)
 		{
-			OutSelection.Add(FaceIndex);
+			const FIntVector& Element = Indices[FaceIndex];
+			if (WeightMap[Element[0]] >= TransferSelectionThreshold &&
+				WeightMap[Element[1]] >= TransferSelectionThreshold &&
+				WeightMap[Element[2]] >= TransferSelectionThreshold)
+			{
+				OutSelection.Add(FaceIndex);
+			}
 		}
 	}
-}
 
-template<bool bIsSecondarySelection>
-bool TransferSelectionSet(const TSharedRef<const FManagedArrayCollection>& TransferClothCollection, const TSharedRef<FManagedArrayCollection>& ClothCollection,
-	const FName& InInputName, const FName& SelectionGroupName, const EChaosClothAssetWeightMapTransferType SimTransferType, const float TransferSelectionThreshold, TSet<int32>& OutSelection)
-{
-	FCollectionClothConstFacade ClothFacade(ClothCollection);
-	FCollectionClothConstFacade TransferClothFacade(TransferClothCollection);
-	FCollectionClothSelectionConstFacade TransferSelectionFacade(TransferClothCollection);
+	template<bool bIsSecondarySelection>
+	bool TransferSelectionSet(const TSharedRef<const FManagedArrayCollection>& TransferClothCollection, const TSharedRef<FManagedArrayCollection>& ClothCollection,
+		const FName& InInputName, const FName& SelectionGroupName, const EChaosClothAssetWeightMapTransferType SimTransferType, const float TransferSelectionThreshold, TSet<int32>& OutSelection)
+	{
+		FCollectionClothConstFacade ClothFacade(ClothCollection);
+		FCollectionClothConstFacade TransferClothFacade(TransferClothCollection);
+		FCollectionClothSelectionConstFacade TransferSelectionFacade(TransferClothCollection);
 
-	const bool bIsValidRenderSelection = SelectionGroupName == ClothCollectionGroup::RenderFaces || SelectionGroupName == ClothCollectionGroup::RenderVertices; 
-	const bool bIsValidSimSelection = SelectionGroupName == ClothCollectionGroup::SimFaces || SelectionGroupName == ClothCollectionGroup::SimVertices2D || SelectionGroupName == ClothCollectionGroup::SimVertices3D;
+		const bool bIsValidRenderSelection = SelectionGroupName == ClothCollectionGroup::RenderFaces || SelectionGroupName == ClothCollectionGroup::RenderVertices; 
+		const bool bIsValidSimSelection = SelectionGroupName == ClothCollectionGroup::SimFaces || SelectionGroupName == ClothCollectionGroup::SimVertices2D || SelectionGroupName == ClothCollectionGroup::SimVertices3D;
 
-	if (!bIsValidRenderSelection && !bIsValidSimSelection)
-	{
-		return false;
-	}
+		if (!bIsValidRenderSelection && !bIsValidSimSelection)
+		{
+			return false;
+		}
 
-	// Get the selection as a vertex set.
-	TSet<int32> TransferSet;
-	const FName DesiredTransferGroup = bIsValidRenderSelection ? ClothCollectionGroup::RenderVertices :
-		(SimTransferType == EChaosClothAssetWeightMapTransferType::Use2DSimMesh ? ClothCollectionGroup::SimVertices2D : ClothCollectionGroup::SimVertices3D);
+		// Get the selection as a vertex set.
+		TSet<int32> TransferSet;
+		const FName DesiredTransferGroup = bIsValidRenderSelection ? ClothCollectionGroup::RenderVertices :
+			(SimTransferType == EChaosClothAssetWeightMapTransferType::Use2DSimMesh ? ClothCollectionGroup::SimVertices2D : ClothCollectionGroup::SimVertices3D);
 	
-	PRAGMA_DISABLE_DEPRECATION_WARNINGS
-	if (!FClothGeometryTools::ConvertSelectionToNewGroupType(TransferClothCollection, InInputName, DesiredTransferGroup, bIsSecondarySelection, TransferSet))
-	{
-		return false;
-	}
-	PRAGMA_ENABLE_DEPRECATION_WARNINGS
+PRAGMA_DISABLE_DEPRECATION_WARNINGS
+		if (!FClothGeometryTools::ConvertSelectionToNewGroupType(TransferClothCollection, InInputName, DesiredTransferGroup, bIsSecondarySelection, TransferSet))
+		{
+			return false;
+		}
+PRAGMA_ENABLE_DEPRECATION_WARNINGS
 
-	// Convert to weights that are 0 on unselected vertices and 1 on selected vertices
-	TArray<float> TransferWeights;
-	TransferWeights.SetNumZeroed(TransferClothCollection->NumElements(DesiredTransferGroup));
-	for (const int32 SetIndex : TransferSet)
-	{
-		if (TransferWeights.IsValidIndex(SetIndex))
+		// Convert to weights that are 0 on unselected vertices and 1 on selected vertices
+		TArray<float> TransferWeights;
+		TransferWeights.SetNumZeroed(TransferClothCollection->NumElements(DesiredTransferGroup));
+		for (const int32 SetIndex : TransferSet)
 		{
-			TransferWeights[SetIndex] = 1.f;
+			if (TransferWeights.IsValidIndex(SetIndex))
+			{
+				TransferWeights[SetIndex] = 1.f;
+			}
 		}
-	}
 
-	// Transfer weights
-	TArray<float> RemappedWeights;
-	RemappedWeights.SetNumZeroed(ClothCollection->NumElements(DesiredTransferGroup));
-	if (bIsValidRenderSelection)
-	{
-		FClothGeometryTools::TransferWeightMap(
-			TransferClothFacade.GetRenderPosition(),
-			TransferClothFacade.GetRenderIndices(),
-			TransferWeights,
-			ClothFacade.GetRenderPosition(),
-			ClothFacade.GetRenderNormal(),
-			ClothFacade.GetRenderIndices(),
-			TArrayView<float>(RemappedWeights));
+		// Transfer weights
+		TArray<float> RemappedWeights;
+		RemappedWeights.SetNumZeroed(ClothCollection->NumElements(DesiredTransferGroup));
+		if (bIsValidRenderSelection)
+		{
+			FClothGeometryTools::TransferWeightMap(
+				TransferClothFacade.GetRenderPosition(),
+				TransferClothFacade.GetRenderIndices(),
+				TransferWeights,
+				ClothFacade.GetRenderPosition(),
+				ClothFacade.GetRenderNormal(),
+				ClothFacade.GetRenderIndices(),
+				TArrayView<float>(RemappedWeights));
+
+			OutSelection.Reset();
+
+			if (SelectionGroupName == ClothCollectionGroup::RenderFaces)
+			{
+				ConvertWeightMapToFaceSelection(RemappedWeights, TransferSelectionThreshold, ClothFacade.GetRenderIndices(), OutSelection);
+			}
+			else
+			{
+				check(SelectionGroupName == ClothCollectionGroup::RenderVertices);
+				ConvertWeightMapToVertexSelection(RemappedWeights, TransferSelectionThreshold, OutSelection);
+			}
+		}
+		else if (SimTransferType == EChaosClothAssetWeightMapTransferType::Use2DSimMesh)
+		{
+			TArray<FVector3f> TransferSimPositions2DAs3D;
+			TConstArrayView<FVector2f> TransferPositions2D = TransferClothFacade.GetSimPosition2D();
+			TransferSimPositions2DAs3D.SetNumUninitialized(TransferPositions2D.Num());
+			for (int32 Index = 0; Index < TransferPositions2D.Num(); ++Index)
+			{
+				TransferSimPositions2DAs3D[Index] = FVector3f(TransferPositions2D[Index], 0.f);
+			}
 
-		OutSelection.Reset();
+			TConstArrayView<FVector2f> Positions2D = ClothFacade.GetSimPosition2D();
+			TArray<FVector3f> Positions2DAs3D;
+			TArray<FVector3f> NormalsZAxis;
+			Positions2DAs3D.SetNumUninitialized(Positions2D.Num());
+			NormalsZAxis.Init(FVector3f::ZAxisVector, Positions2D.Num());
+			for (int32 Index = 0; Index < Positions2D.Num(); ++Index)
+			{
+				Positions2DAs3D[Index] = FVector3f(Positions2D[Index], 0.f);
+			}
 
-		if (SelectionGroupName == ClothCollectionGroup::RenderFaces)
-		{
-			ConvertWeightMapToFaceSelection(RemappedWeights, TransferSelectionThreshold, ClothFacade.GetRenderIndices(), OutSelection);
+			FClothGeometryTools::TransferWeightMap(
+				TConstArrayView<FVector3f>(TransferSimPositions2DAs3D),
+				TransferClothFacade.GetSimIndices2D(),
+				TransferWeights,
+				TConstArrayView<FVector3f>(Positions2DAs3D),
+				TConstArrayView<FVector3f>(NormalsZAxis),
+				ClothFacade.GetSimIndices2D(),
+				TArrayView<float>(RemappedWeights));
+
+			if (SelectionGroupName == ClothCollectionGroup::SimFaces)
+			{
+				ConvertWeightMapToFaceSelection(RemappedWeights, TransferSelectionThreshold, ClothFacade.GetSimIndices2D(), OutSelection);
+			}
+			else if (SelectionGroupName == ClothCollectionGroup::SimVertices2D)
+			{
+				ConvertWeightMapToVertexSelection(RemappedWeights, TransferSelectionThreshold, OutSelection);
+			}
+			else
+			{
+				check(SelectionGroupName == ClothCollectionGroup::SimVertices3D);
+				TSet<int32> Selection2D;
+				ConvertWeightMapToVertexSelection(RemappedWeights, TransferSelectionThreshold, Selection2D);
+				const TConstArrayView<int32> SimVertex3DLookup = ClothFacade.GetSimVertex3DLookup();
+				for (const int32 Vertex2D : Selection2D)
+				{
+					OutSelection.Add(SimVertex3DLookup[Vertex2D]);
+				}
+			}
 		}
 		else
 		{
-			check(SelectionGroupName == ClothCollectionGroup::RenderVertices);
-			ConvertWeightMapToVertexSelection(RemappedWeights, TransferSelectionThreshold, OutSelection);
+			check(SimTransferType == EChaosClothAssetWeightMapTransferType::Use3DSimMesh);
+			FClothGeometryTools::TransferWeightMap(
+				TransferClothFacade.GetSimPosition3D(),
+				TransferClothFacade.GetSimIndices3D(),
+				TransferWeights,
+				ClothFacade.GetSimPosition3D(),
+				ClothFacade.GetSimNormal(),
+				ClothFacade.GetSimIndices3D(),
+				TArrayView<float>(RemappedWeights));
+
+			if (SelectionGroupName == ClothCollectionGroup::SimFaces)
+			{
+				ConvertWeightMapToFaceSelection(RemappedWeights, TransferSelectionThreshold, ClothFacade.GetSimIndices3D(), OutSelection);
+			}
+			else if (SelectionGroupName == ClothCollectionGroup::SimVertices3D)
+			{
+				ConvertWeightMapToVertexSelection(RemappedWeights, TransferSelectionThreshold, OutSelection);
+			}
+			else
+			{
+				check(SelectionGroupName == ClothCollectionGroup::SimVertices2D);
+				TSet<int32> Selection3D;
+				ConvertWeightMapToVertexSelection(RemappedWeights, TransferSelectionThreshold, Selection3D);
+				const TConstArrayView<TArray<int32>> SimVertex2DLookup = ClothFacade.GetSimVertex2DLookup();
+				for (const int32 Vertex3D : Selection3D)
+				{
+					for (const int32 Vertex2D : SimVertex2DLookup[Vertex3D])
+					{
+						OutSelection.Add(Vertex2D);
+					}
+				}
+			}
 		}
+		return true;
 	}
-	else if (SimTransferType == EChaosClothAssetWeightMapTransferType::Use2DSimMesh)
-	{
-		TArray<FVector3f> TransferSimPositions2DAs3D;
-		TConstArrayView<FVector2f> TransferPositions2D = TransferClothFacade.GetSimPosition2D();
-		TransferSimPositions2DAs3D.SetNumUninitialized(TransferPositions2D.Num());
-		for (int32 Index = 0; Index < TransferPositions2D.Num(); ++Index)
-		{
-			TransferSimPositions2DAs3D[Index] = FVector3f(TransferPositions2D[Index], 0.f);
-		}
 
-		TConstArrayView<FVector2f> Positions2D = ClothFacade.GetSimPosition2D();
-		TArray<FVector3f> Positions2DAs3D;
-		TArray<FVector3f> NormalsZAxis;
-		Positions2DAs3D.SetNumUninitialized(Positions2D.Num());
-		NormalsZAxis.Init(FVector3f::ZAxisVector, Positions2D.Num());
-		for (int32 Index = 0; Index < Positions2D.Num(); ++Index)
+	void SetIndices(const TSet<int32>& InputSet, const TSet<int32>& FinalSet, EChaosClothAssetSelectionOverrideType OverrideType, TSet<int32>& Indices, TSet<int32>& RemoveIndices)
+	{
+		if (InputSet.IsEmpty() || OverrideType == EChaosClothAssetSelectionOverrideType::ReplaceAll)
 		{
-			Positions2DAs3D[Index] = FVector3f(Positions2D[Index], 0.f);
+			Indices = FinalSet;
+			RemoveIndices.Reset();
+			return;
 		}
 
-		FClothGeometryTools::TransferWeightMap(
-			TConstArrayView<FVector3f>(TransferSimPositions2DAs3D),
-			TransferClothFacade.GetSimIndices2D(),
-			TransferWeights,
-			TConstArrayView<FVector3f>(Positions2DAs3D),
-			TConstArrayView<FVector3f>(NormalsZAxis),
-			ClothFacade.GetSimIndices2D(),
-			TArrayView<float>(RemappedWeights));
+		Indices = FinalSet.Difference(InputSet);
+		RemoveIndices = InputSet.Difference(FinalSet);
+	}
 
-		if (SelectionGroupName == ClothCollectionGroup::SimFaces)
+	void CalculateFinalSet(const TSet<int32>& InputSet, TSet<int32>& FinalSet, EChaosClothAssetSelectionOverrideType OverrideType, const TSet<int32>& Indices, const TSet<int32>& RemoveIndices)
+	{
+		if (InputSet.IsEmpty() || OverrideType == EChaosClothAssetSelectionOverrideType::ReplaceAll)
 		{
-			ConvertWeightMapToFaceSelection(RemappedWeights, TransferSelectionThreshold, ClothFacade.GetSimIndices2D(), OutSelection);
+			FinalSet = Indices;
+			return;
 		}
-		else if (SelectionGroupName == ClothCollectionGroup::SimVertices2D)
+
+		FinalSet = InputSet;
+		FinalSet.Append(Indices);
+		if (!RemoveIndices.IsEmpty())
 		{
-			ConvertWeightMapToVertexSelection(RemappedWeights, TransferSelectionThreshold, OutSelection);
+			FinalSet = FinalSet.Difference(RemoveIndices);
 		}
-		else
+	}
+}  // namespace UE::Chaos::ClothAsset::Private
+
+// Object encapsulating a change to the Selection Node's values. Used for Undo/Redo.
+class FChaosClothAssetSelectionNode_v2::FSelectionNodeChange final : public FToolCommandChange
+{
+public:
+	FSelectionNodeChange(const FChaosClothAssetSelectionNode_v2& Node)
+		: NodeGuid(Node.GetGuid())
+		, SavedName(Node.OutputName.StringValue)
+		, SavedSelectionOverrideType(Node.SelectionOverrideType)
+		, SavedGroup(Node.Group)
+		, SavedIndices(Node.Indices)
+		, SavedRemoveIndices(Node.RemoveIndices)
+	{}
+
+private:
+	FGuid NodeGuid;
+	FString SavedName;
+	EChaosClothAssetSelectionOverrideType SavedSelectionOverrideType;
+	FChaosClothAssetNodeSelectionGroup SavedGroup;
+	TSet<int32> SavedIndices;
+	TSet<int32> SavedRemoveIndices;
+
+	virtual FString ToString() const final { return TEXT("ChaosClothAssetSelectionNodeChange"); }
+
+	virtual void Apply(UObject* Object) final { SwapApplyRevert(Object); }
+	virtual void Revert(UObject* Object) final { SwapApplyRevert(Object); }
+
+	void SwapApplyRevert(UObject* Object)
+	{
+		if (UDataflow* const Dataflow = Cast<UDataflow>(Object))
 		{
-			check(SelectionGroupName == ClothCollectionGroup::SimVertices3D);
-			TSet<int32> Selection2D;
-			ConvertWeightMapToVertexSelection(RemappedWeights, TransferSelectionThreshold, Selection2D);
-			const TConstArrayView<int32> SimVertex3DLookup = ClothFacade.GetSimVertex3DLookup();
-			for (const int32 Vertex2D : Selection2D)
+			if (const TSharedPtr<FDataflowNode> BaseNode = Dataflow->GetDataflow()->FindBaseNode(NodeGuid))
 			{
-				OutSelection.Add(SimVertex3DLookup[Vertex2D]);
+				if (FChaosClothAssetSelectionNode_v2* const Node = BaseNode->AsType<FChaosClothAssetSelectionNode_v2>())
+				{
+					Swap(Node->OutputName.StringValue, SavedName);
+					Swap(Node->SelectionOverrideType, SavedSelectionOverrideType);
+					Swap(Node->Group, SavedGroup);
+					Swap(Node->Indices, SavedIndices);
+					Swap(Node->RemoveIndices, SavedRemoveIndices);
+					Node->Invalidate();
+				}
 			}
 		}
 	}
+};
+
+FChaosClothAssetSelectionNode_v2::FChaosClothAssetSelectionNode_v2(const Dataflow::FNodeParameters& InParam, FGuid InGuid)
+	: FDataflowNode(InParam, InGuid)
+	, Import(FSimpleDelegate::CreateRaw(this, &FChaosClothAssetSelectionNode_v2::OnImport))
+	, ImportSecondary(FSimpleDelegate::CreateRaw(this, &FChaosClothAssetSelectionNode_v2::OnImportSecondary))
+	, Transfer(FSimpleDelegate::CreateRaw(this, &FChaosClothAssetSelectionNode_v2::OnTransfer))
+{
+	RegisterInputConnection(&Collection);
+	RegisterInputConnection(&InputName.StringValue, GET_MEMBER_NAME_CHECKED(FChaosClothAssetConnectableIStringValue, StringValue))
+		.SetCanHidePin(true)
+		.SetPinIsHidden(true);
+	RegisterInputConnection(&TransferCollection)
+		.SetCanHidePin(true)
+		.SetPinIsHidden(true);
+	RegisterOutputConnection(&Collection, &Collection);
+	RegisterOutputConnection(&OutputName.StringValue, (FString*)nullptr, GET_MEMBER_NAME_CHECKED(FChaosClothAssetConnectableOStringValue, StringValue));
+}
+
+void FChaosClothAssetSelectionNode_v2::OnImport()
+{
+	using namespace UE::Chaos::ClothAsset;
+
+	// Create a temporary context (until we find a way to re-use the one from the tool calling this function)
+	Dataflow::FContextThreaded Context;
+
+	FManagedArrayCollection InCollection = GetValue<FManagedArrayCollection>(Context, &Collection);
+	const TSharedRef<FManagedArrayCollection> ClothCollection = MakeShared<FManagedArrayCollection>(MoveTemp(InCollection));
+	const FCollectionClothSelectionConstFacade SelectionFacade(ClothCollection);
+
+	const FName InInputName = GetInputName(Context);
+	if (const TSet<int32>* const SelectionSet = SelectionFacade.IsValid() ? SelectionFacade.FindSelectionSet(InInputName) : nullptr)
+	{
+		Indices = *SelectionSet;
+		RemoveIndices.Empty();
+		Group.Name = SelectionFacade.GetSelectionGroup(InInputName).ToString();
+		SelectionOverrideType = EChaosClothAssetSelectionOverrideType::ReplaceAll;
+	}
 	else
 	{
-		check(SimTransferType == EChaosClothAssetWeightMapTransferType::Use3DSimMesh);
-		FClothGeometryTools::TransferWeightMap(
-			TransferClothFacade.GetSimPosition3D(),
-			TransferClothFacade.GetSimIndices3D(),
-			TransferWeights,
-			ClothFacade.GetSimPosition3D(),
-			ClothFacade.GetSimNormal(),
-			ClothFacade.GetSimIndices3D(),
-			TArrayView<float>(RemappedWeights));
-
-		if (SelectionGroupName == ClothCollectionGroup::SimFaces)
-		{
-			ConvertWeightMapToFaceSelection(RemappedWeights, TransferSelectionThreshold, ClothFacade.GetSimIndices3D(), OutSelection);
-		}
-		else if (SelectionGroupName == ClothCollectionGroup::SimVertices3D)
+		FClothDataflowTools::LogAndToastWarning(*this,
+			LOCTEXT("NoMatchingSelectionFoundHeadline", "No matching selection found"),
+			FText::Format(LOCTEXT("NoMatchingSelectionFoundDetails", "No matching selection with the name \"{0}\" has been found to import."),
+				FText::FromName(InInputName)));
+	}
+}
+
+void FChaosClothAssetSelectionNode_v2::OnImportSecondary()
+{
+	using namespace UE::Chaos::ClothAsset;
+
+	// Create a temporary context (until we find a way to re-use the one from the tool calling this function)
+	Dataflow::FContextThreaded Context;
+
+	FManagedArrayCollection InCollection = GetValue<FManagedArrayCollection>(Context, &Collection);
+	const TSharedRef<FManagedArrayCollection> ClothCollection = MakeShared<FManagedArrayCollection>(MoveTemp(InCollection));
+	const FCollectionClothSelectionConstFacade SelectionFacade(ClothCollection);
+
+	const FName InInputName = GetInputName(Context);
+	if (const TSet<int32>* const SelectionSet = SelectionFacade.IsValid() ? SelectionFacade.FindSelectionSecondarySet(InInputName) : nullptr)
+	{
+		Indices = *SelectionSet;
+		RemoveIndices.Empty();
+		Group.Name = SelectionFacade.GetSelectionSecondaryGroup(InInputName).ToString();
+		SelectionOverrideType = EChaosClothAssetSelectionOverrideType::ReplaceAll;
+	}
+	else
+	{
+		FClothDataflowTools::LogAndToastWarning(*this,
+			LOCTEXT("NoMatchingSecondarySelectionFoundHeadline", "No matching secondary selection found"),
+			FText::Format(LOCTEXT("NoMatchingSecondarySelectionFoundDetails", "No matching secondary selection with the name \"{0}\" has been found to import."),
+				FText::FromName(InInputName)));
+	}
+}
+
+void FChaosClothAssetSelectionNode_v2::OnTransfer()
+{
+	using namespace UE::Chaos::ClothAsset;
+
+	// Create a temporary context (until we find a way to re-use the one from the tool calling this function)
+	Dataflow::FContextThreaded Context;
+
+	// Transfer selection if the transfer collection input has changed and is valid
+	FManagedArrayCollection InCollection = GetValue<FManagedArrayCollection>(Context, &Collection);
+	const TSharedRef<FManagedArrayCollection> ClothCollection = MakeShared<FManagedArrayCollection>(MoveTemp(InCollection));
+	FCollectionClothConstFacade ClothFacade(ClothCollection);
+	if (ClothFacade.HasValidSimulationData())  // Can only act on the collection if it is a valid cloth collection
+	{
+		FManagedArrayCollection InTransferCollection = GetValue<FManagedArrayCollection>(Context, &TransferCollection);
+		const TSharedRef<const FManagedArrayCollection> TransferClothCollection = MakeShared<const FManagedArrayCollection>(MoveTemp(InTransferCollection));
+		FCollectionClothConstFacade TransferClothFacade(TransferClothCollection);
+		FCollectionClothSelectionConstFacade TransferSelectionFacade(TransferClothCollection);
+
+		const FName InInputName = GetInputName(Context);
+		const FName SelectionGroupName(*Group.Name);
+
+		TSet<int32> PrimaryFinalSelection;
+		if (Private::TransferSelectionSet<false>(TransferClothCollection, ClothCollection, InInputName, SelectionGroupName, SimTransferType, TransferSelectionThreshold, PrimaryFinalSelection))
 		{
-			ConvertWeightMapToVertexSelection(RemappedWeights, TransferSelectionThreshold, OutSelection);
+			TSet<int32> InputSelection;
+			FClothGeometryTools::ConvertSelectionToNewGroupType(ClothCollection, InInputName, SelectionGroupName, InputSelection);
+
+			SetIndices(InputSelection, PrimaryFinalSelection);
 		}
-		else
+	}
+}
+
+void FChaosClothAssetSelectionNode_v2::Evaluate(Dataflow::FContext& Context, const FDataflowOutput* Out) const
+{
+	using namespace UE::Chaos::ClothAsset;
+
+	auto CopyIntoSelection = [this](const TSharedRef<FManagedArrayCollection> SelectionCollection,
+		const FName& SelectionGroupName,
+		const TSet<int32>& SourceIndices,
+		TSet<int32>& DestSelectionSet)
 		{
-			check(SelectionGroupName == ClothCollectionGroup::SimVertices2D);
-			TSet<int32> Selection3D;
-			ConvertWeightMapToVertexSelection(RemappedWeights, TransferSelectionThreshold, Selection3D);
-			const TConstArrayView<TArray<int32>> SimVertex2DLookup = ClothFacade.GetSimVertex2DLookup();
-			for (const int32 Vertex3D : Selection3D)
+			const int32 NumElementsInGroup = SelectionCollection->NumElements(SelectionGroupName);
+			bool bFoundAnyInvalidIndex = false;
+
+			DestSelectionSet.Reset();
+
+			for (const int32 Index : SourceIndices)
 			{
-				for (const int32 Vertex2D : SimVertex2DLookup[Vertex3D])
+				if (Index < 0 || Index >= NumElementsInGroup)
 				{
-					OutSelection.Add(Vertex2D);
+					const FText LogErrorMessage = FText::Format(LOCTEXT("SelectionIndexOutOfBoundsDetails", "Selection index {0} not valid for group \"{1}\" with {2} elements"),
+						Index,
+						FText::FromName(SelectionGroupName),
+						NumElementsInGroup);
+
+					// Log all indices, but toast once
+					UE_LOG(LogChaosClothAssetDataflowNodes, Warning, TEXT("%s"), *LogErrorMessage.ToString());
+					bFoundAnyInvalidIndex = true;
 				}
+				else
+				{
+					DestSelectionSet.Add(Index);
+				}
+			}
+
+			if (bFoundAnyInvalidIndex)
+			{
+				// Toast once
+				const FText ToastErrorMessage = FText::Format(LOCTEXT("AnySelectionIndexOutOfBoundsDetails", "Found invalid selection indices for group \"{0}.\" See log for details"),
+					FText::FromName(SelectionGroupName));
+				FClothDataflowTools::LogAndToastWarning(*this, LOCTEXT("AnySelectionIndexOutOfBoundsHeadline", "Invalid selection"), ToastErrorMessage);
 			}
+		};
+
+	if (Out->IsA<FManagedArrayCollection>(&Collection))
+	{
+		// Evaluate InputName
+		const FName InInputName = GetInputName(Context);
+		const FName SelectionName(OutputName.StringValue.IsEmpty() ? InInputName : FName(OutputName.StringValue));
+
+		if (SelectionName == NAME_None || Group.Name.IsEmpty())
+		{
+			const FManagedArrayCollection& SelectionCollection = GetValue<FManagedArrayCollection>(Context, &Collection);
+			SetValue(Context, SelectionCollection, &Collection);
+			return;
 		}
+
+		const FName SelectionGroupName(*Group.Name);
+
+		FManagedArrayCollection InSelectionCollection = GetValue<FManagedArrayCollection>(Context, &Collection);
+		const TSharedRef<FManagedArrayCollection> SelectionCollection = MakeShared<FManagedArrayCollection>(MoveTemp(InSelectionCollection));
+
+		FCollectionClothSelectionFacade SelectionFacade(SelectionCollection);
+		SelectionFacade.DefineSchema();
+		check(SelectionFacade.IsValid());
+
+		TSet<int32> InputSelectionSet;
+		FClothGeometryTools::ConvertSelectionToNewGroupType(SelectionCollection, InInputName, SelectionGroupName, InputSelectionSet);
+		TSet<int32> FinalSet;
+		CalculateFinalSet(InputSelectionSet, FinalSet);
+
+		TSet<int32>& SelectionSet = SelectionFacade.FindOrAddSelectionSet(SelectionName, SelectionGroupName);
+		CopyIntoSelection(SelectionCollection, SelectionGroupName, FinalSet, SelectionSet);
+
+		SetValue(Context, MoveTemp(*SelectionCollection), &Collection);
+	}
+	else if (Out->IsA<FString>(&OutputName.StringValue))
+	{
+		FString InputNameString = GetValue<FString>(Context, &InputName.StringValue);
+		FClothDataflowTools::MakeCollectionName(InputNameString);
+		SetValue(Context, OutputName.StringValue.IsEmpty() ? InputNameString : OutputName.StringValue, &OutputName.StringValue);
 	}
-	return true;
 }
 
-
-void SetIndices(const TSet<int32>& InputSet, const TSet<int32>& FinalSet, EChaosClothAssetSelectionOverrideType OverrideType, TSet<int32>& Indices, TSet<int32>& RemoveIndices)
+void FChaosClothAssetSelectionNode_v2::OnSelected(Dataflow::FContext& Context)
 {
-	if (InputSet.IsEmpty() || OverrideType == EChaosClothAssetSelectionOverrideType::ReplaceAll)
+	using namespace UE::Chaos::ClothAsset;
+
+	// Re-evaluate the input collection
+	FManagedArrayCollection InCollection = GetValue<FManagedArrayCollection>(Context, &Collection);
+	const TSharedRef<FManagedArrayCollection> ClothCollection = MakeShared<FManagedArrayCollection>(MoveTemp(InCollection));
+	FCollectionClothFacade Cloth(ClothCollection);
+
+	// Update the list of used group for the UI customization
+	const TArray<FName> GroupNames = ClothCollection->GroupNames();
+	CachedCollectionGroupNames.Reset(GroupNames.Num());
+	for (const FName& GroupName : GroupNames)
 	{
-		Indices = FinalSet;
-		RemoveIndices.Reset();
-		return;
+		if (Cloth.IsValidClothCollectionGroupName(GroupName))  // Restrict to the cloth facade groups
+		{
+			CachedCollectionGroupNames.Emplace(GroupName);
+		}
 	}
+}
 
-	Indices = FinalSet.Difference(InputSet);
-	RemoveIndices = InputSet.Difference(FinalSet);
+void FChaosClothAssetSelectionNode_v2::OnDeselected()
+{
+	// Clean up, to avoid another toolkit picking up the wrong context evaluation
+	CachedCollectionGroupNames.Reset();
 }
 
-void CalculateFinalSet(const TSet<int32>& InputSet, TSet<int32>& FinalSet, EChaosClothAssetSelectionOverrideType OverrideType, const TSet<int32>& Indices, const TSet<int32>& RemoveIndices)
+FName FChaosClothAssetSelectionNode_v2::GetInputName(Dataflow::FContext& Context) const
 {
-	if (InputSet.IsEmpty() || OverrideType == EChaosClothAssetSelectionOverrideType::ReplaceAll)
-	{
-		FinalSet = Indices;
-		return;
-	}
+	FString InputNameString = GetValue<FString>(Context, &InputName.StringValue);
+	UE::Chaos::ClothAsset::FClothDataflowTools::MakeCollectionName(InputNameString);
+	const FName InInputName(*InputNameString);
+	return InInputName != NAME_None ? InInputName : FName(OutputName.StringValue);
+}
 
-	FinalSet = InputSet;
-	FinalSet.Append(Indices);
-	if (!RemoveIndices.IsEmpty())
-	{
-		FinalSet = FinalSet.Difference(RemoveIndices);
-	}
+void FChaosClothAssetSelectionNode_v2::SetIndices(const TSet<int32>& InputSet, const TSet<int32>& FinalSet)
+{
+	UE::Chaos::ClothAsset::Private::SetIndices(InputSet, FinalSet, SelectionOverrideType, Indices, RemoveIndices);
+}
+
+void FChaosClothAssetSelectionNode_v2::CalculateFinalSet(const TSet<int32>& InputSet, TSet<int32>& FinalSet) const
+{
+	UE::Chaos::ClothAsset::Private::CalculateFinalSet(InputSet, FinalSet, SelectionOverrideType, Indices, RemoveIndices);
 }
+
+TUniquePtr<FToolCommandChange> FChaosClothAssetSelectionNode_v2::MakeSelectedNodeChange(const FChaosClothAssetSelectionNode_v2& Node)
+{
+	return MakeUnique<FChaosClothAssetSelectionNode_v2::FSelectionNodeChange>(Node);
 }
 
 FChaosClothAssetSelectionNode::FChaosClothAssetSelectionNode(const Dataflow::FNodeParameters& InParam, FGuid InGuid)
@@ -242,7 +524,6 @@ FChaosClothAssetSelectionNode::FChaosClothAssetSelectionNode(const Dataflow::FNo
 	RegisterOutputConnection(&Name);
 }
 
-
 void FChaosClothAssetSelectionNode::SetAssetValue(TObjectPtr<UObject> Asset, Dataflow::FContext& Context) const
 {
 	using namespace UE::Chaos::ClothAsset;
@@ -258,7 +539,7 @@ void FChaosClothAssetSelectionNode::SetAssetValue(TObjectPtr<UObject> Asset, Dat
 				check(MutableThis == this);
 
 				// Make the name a valid attribute name, and replace the value in the UI
-				FWeightMapTools::MakeWeightMapName(MutableThis->Name);
+				FClothDataflowTools::MakeCollectionName(MutableThis->Name);
 
 				const FName SelectionGroupName(*Group.Name);
 				const FName SelectionSecondaryGroupName(*SecondaryGroup.Name);
@@ -423,7 +704,7 @@ void FChaosClothAssetSelectionNode::Evaluate(Dataflow::FContext& Context, const
 	else if (Out->IsA<FString>(&Name))
 	{
 		FString InputNameString = GetValue<FString>(Context, &InputName.StringValue);
-		UE::Chaos::ClothAsset::FWeightMapTools::MakeWeightMapName(InputNameString);
+		FClothDataflowTools::MakeCollectionName(InputNameString);
 		SetValue(Context, Name.IsEmpty() ? InputNameString : Name, &Name);
 	}
 }
@@ -481,11 +762,10 @@ PRAGMA_DISABLE_DEPRECATION_WARNINGS
 PRAGMA_ENABLE_DEPRECATION_WARNINGS
 }
 
-
 FName FChaosClothAssetSelectionNode::GetInputName(Dataflow::FContext& Context) const
 {
 	FString InputNameString = GetValue<FString>(Context, &InputName.StringValue);
-	UE::Chaos::ClothAsset::FWeightMapTools::MakeWeightMapName(InputNameString);
+	UE::Chaos::ClothAsset::FClothDataflowTools::MakeCollectionName(InputNameString);
 	const FName InInputName(*InputNameString);
 	return InInputName != NAME_None ? InInputName : FName(Name);
 }
@@ -510,78 +790,4 @@ void FChaosClothAssetSelectionNode::CalculateFinalSecondarySet(const TSet<int32>
 	UE::Chaos::ClothAsset::Private::CalculateFinalSet(InputSet, FinalSet, SelectionOverrideType, SecondaryIndices, RemoveSecondaryIndices);
 }
 
-
-// Object encapsulating a change to the Selection Node's values. Used for Undo/Redo.
-class FChaosClothAssetSelectionNode::FSelectionNodeChange final : public FToolCommandChange
-{
-public:
-
-	FSelectionNodeChange(const FChaosClothAssetSelectionNode& Node) :
-		NodeGuid(Node.GetGuid()),
-		SavedName(Node.Name),
-		SavedSelectionOverrideType(Node.SelectionOverrideType),
-		SavedGroup(Node.Group),
-		SavedIndices(Node.Indices),
-		SavedRemoveIndices(Node.RemoveIndices),
-		SavedSecondaryGroup(Node.SecondaryGroup),
-		SavedSecondaryIndices(Node.SecondaryIndices),
-		SavedRemoveSecondaryIndices(Node.RemoveSecondaryIndices)
-	{}
-
-private:
-
-	FGuid NodeGuid;
-	FString SavedName;
-	EChaosClothAssetSelectionOverrideType SavedSelectionOverrideType;
-	FChaosClothAssetNodeSelectionGroup SavedGroup;
-	TSet<int32> SavedIndices;
-	TSet<int32> SavedRemoveIndices;
-	FChaosClothAssetNodeSelectionGroup SavedSecondaryGroup;
-	TSet<int32> SavedSecondaryIndices;
-	TSet<int32> SavedRemoveSecondaryIndices;
-
-	virtual FString ToString() const final
-	{
-		return TEXT("ChaosClothAssetSelectionNodeChange");
-	}
-
-	virtual void Apply(UObject* Object) final
-	{
-		SwapApplyRevert(Object);
-	}
-
-	virtual void Revert(UObject* Object) final
-	{
-		SwapApplyRevert(Object);
-	}
-
-	void SwapApplyRevert(UObject* Object)
-	{
-		if (UDataflow* const Dataflow = Cast<UDataflow>(Object))
-		{
-			if (const TSharedPtr<FDataflowNode> BaseNode = Dataflow->GetDataflow()->FindBaseNode(NodeGuid))
-			{
-				if (FChaosClothAssetSelectionNode* const Node = BaseNode->AsType<FChaosClothAssetSelectionNode>())
-				{
-					Swap(Node->Name, SavedName);
-					Swap(Node->SelectionOverrideType, SavedSelectionOverrideType);
-					Swap(Node->Group, SavedGroup);
-					Swap(Node->Indices, SavedIndices);
-					Swap(Node->RemoveIndices, SavedRemoveIndices);
-					Swap(Node->SecondaryGroup, SavedSecondaryGroup);
-					Swap(Node->SecondaryIndices, SavedSecondaryIndices);
-					Swap(Node->RemoveSecondaryIndices, SavedRemoveSecondaryIndices);
-					Node->Invalidate();
-				}
-			}
-		}
-	}
-};
-
-TUniquePtr<FToolCommandChange> FChaosClothAssetSelectionNode::MakeWeightMapNodeChange(const FChaosClothAssetSelectionNode& Node)
-{
-	return MakeUnique<FChaosClothAssetSelectionNode::FSelectionNodeChange>(Node);
-}
-
-
 #undef LOCTEXT_NAMESPACE
diff --git a/Engine/Plugins/ChaosClothAssetEditor/Source/ChaosClothAssetDataflowNodes/Private/ChaosClothAsset/SimulationLongRangeAttachmentConfigNode.cpp b/Engine/Plugins/ChaosClothAssetEditor/Source/ChaosClothAssetDataflowNodes/Private/ChaosClothAsset/SimulationLongRangeAttachmentConfigNode.cpp
--- a/Engine/Plugins/ChaosClothAssetEditor/Source/ChaosClothAssetDataflowNodes/Private/ChaosClothAsset/SimulationLongRangeAttachmentConfigNode.cpp
+++ b/Engine/Plugins/ChaosClothAssetEditor/Source/ChaosClothAssetDataflowNodes/Private/ChaosClothAsset/SimulationLongRangeAttachmentConfigNode.cpp
@@ -51,10 +51,12 @@ FChaosClothAssetSimulationLongRangeAttachmentConfigNode::FChaosClothAssetSimulat
 
 void FChaosClothAssetSimulationLongRangeAttachmentConfigNode::AddProperties(FPropertyHelper& PropertyHelper) const
 {
+PRAGMA_DISABLE_DEPRECATION_WARNINGS  // SetProperty functions are templated and cause deprecation warnings with the now deprecated v1
 	PropertyHelper.SetPropertyWeighted(this, &TetherStiffness);
 	PropertyHelper.SetPropertyWeighted(this, &TetherScale);
 	PropertyHelper.SetPropertyBool(this, &bUseGeodesicTethers);
 	PropertyHelper.SetPropertyString(this, &FixedEndWeightMap);
+PRAGMA_ENABLE_DEPRECATION_WARNINGS
 }
 
 void FChaosClothAssetSimulationLongRangeAttachmentConfigNode::EvaluateClothCollection(Dataflow::FContext& Context, const TSharedRef<FManagedArrayCollection>& ClothCollection) const
diff --git a/Engine/Plugins/ChaosClothAssetEditor/Source/ChaosClothAssetDataflowNodes/Private/ChaosClothAsset/TerminalNode.cpp b/Engine/Plugins/ChaosClothAssetEditor/Source/ChaosClothAssetDataflowNodes/Private/ChaosClothAsset/TerminalNode.cpp
--- a/Engine/Plugins/ChaosClothAssetEditor/Source/ChaosClothAssetDataflowNodes/Private/ChaosClothAsset/TerminalNode.cpp
+++ b/Engine/Plugins/ChaosClothAssetEditor/Source/ChaosClothAssetDataflowNodes/Private/ChaosClothAsset/TerminalNode.cpp
@@ -632,7 +632,9 @@ TArray<const FManagedArrayCollection*> FChaosClothAssetTerminalNode::GetCollecti
 	return CollectionLods;
 }
 
+PRAGMA_DISABLE_DEPRECATION_WARNINGS  // Unexpected deprecation message on some platforms otherwise
 const FManagedArrayCollection* FChaosClothAssetTerminalNode::GetCollectionLod(int32 LodIndex) const
+PRAGMA_ENABLE_DEPRECATION_WARNINGS
 {
 	switch (LodIndex)
 	{
diff --git a/Engine/Plugins/ChaosClothAssetEditor/Source/ChaosClothAssetDataflowNodes/Public/ChaosClothAsset/AddWeightMapNode.h b/Engine/Plugins/ChaosClothAssetEditor/Source/ChaosClothAssetDataflowNodes/Public/ChaosClothAsset/AddWeightMapNode.h
--- a/Engine/Plugins/ChaosClothAssetEditor/Source/ChaosClothAssetDataflowNodes/Public/ChaosClothAsset/AddWeightMapNode.h
+++ b/Engine/Plugins/ChaosClothAssetEditor/Source/ChaosClothAssetDataflowNodes/Public/ChaosClothAsset/AddWeightMapNode.h
@@ -44,7 +44,7 @@ enum class EChaosClothAssetWeightMapOverrideType : uint8
 */
 PRAGMA_DISABLE_DEPRECATION_WARNINGS
 USTRUCT(Meta = (DataflowCloth, Deprecated="5.5"))
-struct FChaosClothAssetAddWeightMapNode : public FDataflowTerminalNode
+struct UE_DEPRECATED(5.5, "Use FChaosClothAssetWeightMapNode instead.") FChaosClothAssetAddWeightMapNode : public FDataflowTerminalNode
 {
 	GENERATED_USTRUCT_BODY()
 	DATAFLOW_NODE_DEFINE_INTERNAL(FChaosClothAssetAddWeightMapNode, "AddWeightMap", "Cloth", "Cloth Add Weight Map")
diff --git a/Engine/Plugins/ChaosClothAssetEditor/Source/ChaosClothAssetDataflowNodes/Public/ChaosClothAsset/AttributeNode.h b/Engine/Plugins/ChaosClothAssetEditor/Source/ChaosClothAssetDataflowNodes/Public/ChaosClothAsset/AttributeNode.h
--- a/Engine/Plugins/ChaosClothAssetEditor/Source/ChaosClothAssetDataflowNodes/Public/ChaosClothAsset/AttributeNode.h
+++ b/Engine/Plugins/ChaosClothAssetEditor/Source/ChaosClothAssetDataflowNodes/Public/ChaosClothAsset/AttributeNode.h
@@ -30,10 +30,10 @@ enum class EChaosClothAssetNodeAttributeType : uint8
 
 /** Create a new attribute for the specified group. */
 USTRUCT(Meta = (DataflowCloth, Experimental))
-struct FChaosClothAssetAttributeNode_V2 : public FDataflowNode
+struct FChaosClothAssetAttributeNode_v2 : public FDataflowNode
 {
 	GENERATED_USTRUCT_BODY()
-	DATAFLOW_NODE_DEFINE_INTERNAL(FChaosClothAssetAttributeNode_V2, "Attribute", "Cloth", "Cloth Attribute")
+	DATAFLOW_NODE_DEFINE_INTERNAL(FChaosClothAssetAttributeNode_v2, "Attribute", "Cloth", "Cloth Attribute")
 	DATAFLOW_NODE_RENDER_TYPE("SurfaceRender", FName("FClothCollection"), "Collection")
 
 public:
@@ -65,7 +65,7 @@ public:
 	UPROPERTY(EditAnywhere, Category = "Attribute", Meta = (EditCondition = "Type == EChaosClothAssetNodeAttributeType::Vector", EditConditionHides))
 	FVector3f VectorValue = FVector3f::ZeroVector;
 
-	FChaosClothAssetAttributeNode_V2(const Dataflow::FNodeParameters& InParam, FGuid InGuid = FGuid::NewGuid());
+	FChaosClothAssetAttributeNode_v2(const Dataflow::FNodeParameters& InParam, FGuid InGuid = FGuid::NewGuid());
 
 	/** Return a cached array of all the groups used by the input collection during at the time of the latest evaluation. */
 	const TArray<FName>& GetCachedCollectionGroupNames() const { return CachedCollectionGroupNames; }
@@ -82,7 +82,7 @@ private:
 
 /** Create a new attribute for the specified group. */
 USTRUCT(Meta = (DataflowCloth, Experimental, Deprecated = "5.5"))
-struct FChaosClothAssetAttributeNode : public FDataflowNode
+struct UE_DEPRECATED(5.5, "Use the newer version of this node instead.") FChaosClothAssetAttributeNode : public FDataflowNode
 {
 	GENERATED_USTRUCT_BODY()
 	DATAFLOW_NODE_DEFINE_INTERNAL(FChaosClothAssetAttributeNode, "Attribute", "Cloth", "Cloth Attribute")
diff --git a/Engine/Plugins/ChaosClothAssetEditor/Source/ChaosClothAssetDataflowNodes/Public/ChaosClothAsset/MergeClothCollectionsNode.h b/Engine/Plugins/ChaosClothAssetEditor/Source/ChaosClothAssetDataflowNodes/Public/ChaosClothAsset/MergeClothCollectionsNode.h
--- a/Engine/Plugins/ChaosClothAssetEditor/Source/ChaosClothAssetDataflowNodes/Public/ChaosClothAsset/MergeClothCollectionsNode.h
+++ b/Engine/Plugins/ChaosClothAssetEditor/Source/ChaosClothAssetDataflowNodes/Public/ChaosClothAsset/MergeClothCollectionsNode.h
@@ -44,7 +44,7 @@ private:
 
 /** Merge multiple cloth collections into a single cloth collection of multiple patterns. */
 USTRUCT(Meta = (DataflowCloth, Deprecated="5.5"))
-struct FChaosClothAssetMergeClothCollectionsNode : public FDataflowNode
+struct UE_DEPRECATED(5.5, "Use the newer version of this node instead.") FChaosClothAssetMergeClothCollectionsNode : public FDataflowNode
 {
 	GENERATED_USTRUCT_BODY()
 	DATAFLOW_NODE_DEFINE_INTERNAL(FChaosClothAssetMergeClothCollectionsNode, "MergeClothCollections", "Cloth", "Cloth Merge Collection")
@@ -72,8 +72,10 @@ public:
 	UPROPERTY(Meta = (DisplayName = "Collection 5"))
 	FManagedArrayCollection Collection5;
 	/** The number of inputs currently exposed to the node UI. */
+PRAGMA_DISABLE_DEPRECATION_WARNINGS
 	UPROPERTY()
 	int32 NumInputs = NumInitialOptionalInputs;
+PRAGMA_ENABLE_DEPRECATION_WARNINGS
 
 	FChaosClothAssetMergeClothCollectionsNode(const Dataflow::FNodeParameters& InParam, FGuid InGuid = FGuid::NewGuid());
 
diff --git a/Engine/Plugins/ChaosClothAssetEditor/Source/ChaosClothAssetDataflowNodes/Public/ChaosClothAsset/ProxyDeformerNode.h b/Engine/Plugins/ChaosClothAssetEditor/Source/ChaosClothAssetDataflowNodes/Public/ChaosClothAsset/ProxyDeformerNode.h
--- a/Engine/Plugins/ChaosClothAssetEditor/Source/ChaosClothAssetDataflowNodes/Public/ChaosClothAsset/ProxyDeformerNode.h
+++ b/Engine/Plugins/ChaosClothAssetEditor/Source/ChaosClothAssetDataflowNodes/Public/ChaosClothAsset/ProxyDeformerNode.h
@@ -22,7 +22,7 @@ struct FChaosClothAssetSelectionFilterSet
 };
 
 /** Add the proxy deformer information to this cloth collection's render data. */
-USTRUCT(Meta = (DataflowCloth, Experimental))  // TODO: Remove experimental for 5.5 release
+USTRUCT(Meta = (DataflowCloth))
 struct FChaosClothAssetProxyDeformerNode_v2 : public FDataflowNode
 {
 	GENERATED_USTRUCT_BODY()
@@ -69,8 +69,8 @@ private:
 };
 
 /** Add the proxy deformer information to this cloth collection's render data. */
-USTRUCT(Meta = (DataflowCloth))
-struct FChaosClothAssetProxyDeformerNode : public FDataflowNode
+USTRUCT(Meta = (DataflowCloth, Deprecated = "5.5"))
+struct UE_DEPRECATED(5.5, "Use the newer version of this node instead.") FChaosClothAssetProxyDeformerNode : public FDataflowNode
 {
 	GENERATED_USTRUCT_BODY()
 	DATAFLOW_NODE_DEFINE_INTERNAL(FChaosClothAssetProxyDeformerNode, "ProxyDeformer", "Cloth", "Cloth Simulation Proxy Deformer")
@@ -161,8 +161,10 @@ private:
 	static constexpr int32 MaxNumFilterSets = 10;
 
 	/** The number of filter sets currently exposed to the node UI. */
+PRAGMA_DISABLE_DEPRECATION_WARNINGS
 	UPROPERTY()
 	int32 NumFilterSets = NumInitialOptionalInputs;
+PRAGMA_ENABLE_DEPRECATION_WARNINGS
 
 	//~ Begin FDataflowNode interface
 	virtual void Evaluate(Dataflow::FContext& Context, const FDataflowOutput* Out) const override;
diff --git a/Engine/Plugins/ChaosClothAssetEditor/Source/ChaosClothAssetDataflowNodes/Public/ChaosClothAsset/SelectionNode.h b/Engine/Plugins/ChaosClothAssetEditor/Source/ChaosClothAssetDataflowNodes/Public/ChaosClothAsset/SelectionNode.h
--- a/Engine/Plugins/ChaosClothAssetEditor/Source/ChaosClothAssetDataflowNodes/Public/ChaosClothAsset/SelectionNode.h
+++ b/Engine/Plugins/ChaosClothAssetEditor/Source/ChaosClothAssetDataflowNodes/Public/ChaosClothAsset/SelectionNode.h
@@ -3,6 +3,7 @@
 #pragma once 
 
 #include "Dataflow/DataflowNode.h"
+#include "Dataflow/DataflowFunctionProperty.h"
 #include "GeometryCollection/ManagedArrayCollection.h"
 #include "ChaosClothAsset/AddWeightMapNode.h"
 #include "SelectionNode.generated.h"
@@ -53,11 +54,109 @@ enum class EChaosClothAssetSelectionOverrideType : uint8
 	Modify
 };
 
+USTRUCT(Meta = (DataflowCloth))
+struct FChaosClothAssetSelectionNode_v2 : public FDataflowNode
+{
+	GENERATED_USTRUCT_BODY()
+	DATAFLOW_NODE_DEFINE_INTERNAL(FChaosClothAssetSelectionNode_v2, "Selection", "Cloth", "Cloth Selection")
+	DATAFLOW_NODE_RENDER_TYPE("SurfaceRender", FName("FClothCollection"), "Collection")
+
+public:
+
+	UPROPERTY(Meta = (Dataflowinput, DataflowOutput, DataflowPassthrough = "Collection"))
+	FManagedArrayCollection Collection;
+
+	/** The collection used to transfer sets from. */
+	UPROPERTY(Meta = (DataflowInput))
+	FManagedArrayCollection TransferCollection;
+
+	/** The name to be use as a selection. */
+	UPROPERTY(EditAnywhere, Category = "Selection")
+	FChaosClothAssetConnectableOStringValue OutputName;
+
+	/** The name to populate this set from and override based on Selection Override Type. Output Name will be used if Input Name is empty.*/
+	UPROPERTY(EditAnywhere, Category = "Selection")
+	FChaosClothAssetConnectableIStringValue InputName;
+
+	/** How to apply this node's Indices onto existing sets. Changing this value will change the output set.
+	 *  To change how the node's stored indices are calculated, change the equivalent value on the Selection Tool context.*/
+	UPROPERTY(EditAnywhere, Category = "Selection")
+	EChaosClothAssetSelectionOverrideType SelectionOverrideType = EChaosClothAssetSelectionOverrideType::ReplaceAll;
+
+	/** The type of element the selection refers to */
+	UPROPERTY(EditAnywhere, Category = "Selection")
+	FChaosClothAssetNodeSelectionGroup Group;
+
+	/** Selected element indices */
+	UPROPERTY(EditAnywhere, Category = "Selection", Meta = (ClampMin = "0"))
+	TSet<int32> Indices;
+
+	/** Indices to remove from the Input selection */
+	UPROPERTY(EditAnywhere, Category = "Selection", Meta = (ClampMin = "0", EditCondition = "SelectionOverrideType == EChaosClothAssetSelectionOverrideType::Modify"))
+	TSet<int32> RemoveIndices;
+
+	/** Import (replace) the current selection from the input Collection's selection with the given Input Name (or Output Name if Input Name is empty). */
+	UPROPERTY(EditAnywhere, Category = "Selection", Meta = (ButtonImage = "Icons.Refresh"))
+	FDataflowFunctionProperty Import;
+
+	/**
+	 * Import (replace) the current selection from the input Collection's secondary selection with the given Input Name (or Output Name if Input Name is empty).
+	 * Secondary selections are only supported in v1 of this node. This function is provided as a migration tool to this current version.
+	 */
+	UPROPERTY(EditAnywhere, Category = "Selection", Meta = (ButtonImage = "Icons.Refresh"))
+	FDataflowFunctionProperty ImportSecondary;
+
+	/**
+	 * The type of transfer used to transfer the sim mesh sets when a TransferCollection is connected.
+	 * This property is disabled when no TransferCollection input has been connected.
+	 */
+	UPROPERTY(EditAnywhere, Category = "Selection Transfer")
+	EChaosClothAssetWeightMapTransferType SimTransferType = EChaosClothAssetWeightMapTransferType::Use2DSimMesh;
+
+	/**
+	 * Selections are internally converted to maps in order to do the transfer and then converted back.
+	 * This value is used to do the conversion back. Decrease this value to (possibly) expand the converted selection.
+	 */
+	UPROPERTY(EditAnywhere, Category = "Selection Transfer", Meta = (ClampMin = "0", ClampMax = "1"))
+	float TransferSelectionThreshold = 0.95f;
+
+	/** Transfer the selection from the connected Transfer Collection containing a selection with Input Name (or Output Name if Input Name is empty). */
+	UPROPERTY(EditAnywhere, Category = "Selection Transfer", Meta = (ButtonImage = "Icons.Convert"))
+	FDataflowFunctionProperty Transfer;
+
+	FChaosClothAssetSelectionNode_v2(const Dataflow::FNodeParameters& InParam, FGuid InGuid = FGuid::NewGuid());
+
+	/** Return a cached array of all the groups used by the input collection during at the time of the latest evaluation. */
+	const TArray<FName>& GetCachedCollectionGroupNames() const { return CachedCollectionGroupNames; }
+
+private:
+	friend class UClothMeshSelectionTool;
+
+	class FSelectionNodeChange;
+
+	FName CHAOSCLOTHASSETDATAFLOWNODES_API GetInputName(Dataflow::FContext& Context) const;
+	void CHAOSCLOTHASSETDATAFLOWNODES_API SetIndices(const TSet<int32>& InputSet, const TSet<int32>& FinalSet);
+	void CHAOSCLOTHASSETDATAFLOWNODES_API CalculateFinalSet(const TSet<int32>& InputSet, TSet<int32>& FinalSet) const;
+	static TUniquePtr<class FToolCommandChange> CHAOSCLOTHASSETDATAFLOWNODES_API MakeSelectedNodeChange(const FChaosClothAssetSelectionNode_v2& Node);
+
+	//~ Begin FDataflowNode implementation
+	virtual void Evaluate(Dataflow::FContext& Context, const FDataflowOutput* Out) const override;
+	virtual void OnSelected(Dataflow::FContext& Context) override;
+	virtual void OnDeselected() override;
+	//~ End FDataflowNode implementation
+
+	void OnImport();
+	void OnImportSecondary();
+	void OnTransfer();
+
+	TArray<FName> CachedCollectionGroupNames;
+};
+
 PRAGMA_DISABLE_DEPRECATION_WARNINGS  // For EChaosClothAssetSelectionType
 
 /** Integer index set selection node. */
-USTRUCT(Meta = (DataflowCloth))
-struct FChaosClothAssetSelectionNode : public FDataflowTerminalNode
+USTRUCT(Meta = (DataflowCloth, Deprecated = "5.5"))
+struct UE_DEPRECATED(5.5, "Use the newer version of this node instead.") FChaosClothAssetSelectionNode : public FDataflowTerminalNode
 {
 	GENERATED_USTRUCT_BODY()
 	DATAFLOW_NODE_DEFINE_INTERNAL(FChaosClothAssetSelectionNode, "Selection", "Cloth", "Cloth Selection")
@@ -149,12 +248,6 @@ public:
 	void CHAOSCLOTHASSETDATAFLOWNODES_API CalculateFinalSecondarySet(const TSet<int32>& InputSet, TSet<int32>& FinalSet) const;
 
 private:
-
-	friend class UClothMeshSelectionTool;
-
-	class FSelectionNodeChange;
-	static TUniquePtr<class FToolCommandChange> CHAOSCLOTHASSETDATAFLOWNODES_API MakeWeightMapNodeChange(const FChaosClothAssetSelectionNode& Node);
-
 	virtual void SetAssetValue(TObjectPtr<UObject> Asset, Dataflow::FContext& Context) const override;
 	virtual void Evaluate(Dataflow::FContext& Context, const FDataflowOutput* Out) const override;
 	virtual void OnSelected(Dataflow::FContext& Context) override;
diff --git a/Engine/Plugins/ChaosClothAssetEditor/Source/ChaosClothAssetDataflowNodes/Public/ChaosClothAsset/SimulationLongRangeAttachmentConfigNode.h b/Engine/Plugins/ChaosClothAssetEditor/Source/ChaosClothAssetDataflowNodes/Public/ChaosClothAsset/SimulationLongRangeAttachmentConfigNode.h
--- a/Engine/Plugins/ChaosClothAssetEditor/Source/ChaosClothAssetDataflowNodes/Public/ChaosClothAsset/SimulationLongRangeAttachmentConfigNode.h
+++ b/Engine/Plugins/ChaosClothAssetEditor/Source/ChaosClothAssetDataflowNodes/Public/ChaosClothAsset/SimulationLongRangeAttachmentConfigNode.h
@@ -58,7 +58,7 @@ private:
 
 /** Long range attachment constraint property configuration node. */
 USTRUCT(Meta = (DataflowCloth, Deprecated = "5.5"))
-struct FChaosClothAssetSimulationLongRangeAttachmentConfigNode : public FChaosClothAssetSimulationBaseConfigNode
+struct UE_DEPRECATED(5.5, "Use the newer version of this node instead.") FChaosClothAssetSimulationLongRangeAttachmentConfigNode : public FChaosClothAssetSimulationBaseConfigNode
 {
 	GENERATED_USTRUCT_BODY()
 	DATAFLOW_NODE_DEFINE_INTERNAL(FChaosClothAssetSimulationLongRangeAttachmentConfigNode, "SimulationLongRangeAttachmentConfig", "Cloth", "Cloth Simulation Long Range Attachment Config")
diff --git a/Engine/Plugins/ChaosClothAssetEditor/Source/ChaosClothAssetDataflowNodes/Public/ChaosClothAsset/SkeletalMeshImportNode.h b/Engine/Plugins/ChaosClothAssetEditor/Source/ChaosClothAssetDataflowNodes/Public/ChaosClothAsset/SkeletalMeshImportNode.h
--- a/Engine/Plugins/ChaosClothAssetEditor/Source/ChaosClothAssetDataflowNodes/Public/ChaosClothAsset/SkeletalMeshImportNode.h
+++ b/Engine/Plugins/ChaosClothAssetEditor/Source/ChaosClothAssetDataflowNodes/Public/ChaosClothAsset/SkeletalMeshImportNode.h
@@ -67,7 +67,7 @@ private:
 /** Import a skeletal mesh asset into the cloth collection simulation and/or render mesh containers. 
  * This version re-calculates the Sim Mesh Normals, and they are flipped. */
 USTRUCT(Meta = (DataflowCloth, Deprecated = "5.5"))
-struct FChaosClothAssetSkeletalMeshImportNode : public FDataflowNode
+struct UE_DEPRECATED(5.5, "Use the newer version of this node instead.") FChaosClothAssetSkeletalMeshImportNode : public FDataflowNode
 {
 	GENERATED_USTRUCT_BODY()
 	DATAFLOW_NODE_DEFINE_INTERNAL(FChaosClothAssetSkeletalMeshImportNode, "SkeletalMeshImport", "Cloth", "Cloth Skeletal Mesh Import")
diff --git a/Engine/Plugins/ChaosClothAssetEditor/Source/ChaosClothAssetDataflowNodes/Public/ChaosClothAsset/StaticMeshImportNode.h b/Engine/Plugins/ChaosClothAssetEditor/Source/ChaosClothAssetDataflowNodes/Public/ChaosClothAsset/StaticMeshImportNode.h
--- a/Engine/Plugins/ChaosClothAssetEditor/Source/ChaosClothAssetDataflowNodes/Public/ChaosClothAsset/StaticMeshImportNode.h
+++ b/Engine/Plugins/ChaosClothAssetEditor/Source/ChaosClothAssetDataflowNodes/Public/ChaosClothAsset/StaticMeshImportNode.h
@@ -65,7 +65,7 @@ private:
 /** Import a static mesh asset into the cloth collection simulation and/or render mesh containers. 
 * This version re-calculates the Sim Mesh Normals, and they are flipped. */
 USTRUCT(Meta = (DataflowCloth, Deprecated = "5.5"))
-struct FChaosClothAssetStaticMeshImportNode : public FDataflowNode
+struct UE_DEPRECATED(5.5, "Use the newer version of this node instead.") FChaosClothAssetStaticMeshImportNode : public FDataflowNode
 {
 	GENERATED_USTRUCT_BODY()
 	DATAFLOW_NODE_DEFINE_INTERNAL(FChaosClothAssetStaticMeshImportNode, "StaticMeshImport", "Cloth", "Cloth Static Mesh Import")
diff --git a/Engine/Plugins/ChaosClothAssetEditor/Source/ChaosClothAssetDataflowNodes/Public/ChaosClothAsset/TerminalNode.h b/Engine/Plugins/ChaosClothAssetEditor/Source/ChaosClothAssetDataflowNodes/Public/ChaosClothAsset/TerminalNode.h
--- a/Engine/Plugins/ChaosClothAssetEditor/Source/ChaosClothAssetDataflowNodes/Public/ChaosClothAsset/TerminalNode.h
+++ b/Engine/Plugins/ChaosClothAssetEditor/Source/ChaosClothAssetDataflowNodes/Public/ChaosClothAsset/TerminalNode.h
@@ -68,7 +68,7 @@ private:
 
 /** Cloth terminal node to generate a cloth asset from a cloth collection. */
 USTRUCT(Meta = (DataflowCloth, DataflowTerminal, Deprecated = 5.5))
-struct FChaosClothAssetTerminalNode : public FDataflowTerminalNode
+struct UE_DEPRECATED(5.5, "Use the newer version of this node instead.") FChaosClothAssetTerminalNode : public FDataflowTerminalNode
 {
 	GENERATED_USTRUCT_BODY()
 	DATAFLOW_NODE_DEFINE_INTERNAL(FChaosClothAssetTerminalNode, "ClothAssetTerminal", "Cloth", "Cloth Terminal")  // TODO: Should the category be Terminal instead like all other terminal nodes
@@ -96,8 +96,10 @@ public:
 	UPROPERTY(meta = (DisplayName = "Collection LOD 5"))
 	FManagedArrayCollection CollectionLod5;
 	/** The number of LODs currently exposed to the node UI. */
+PRAGMA_DISABLE_DEPRECATION_WARNINGS
 	UPROPERTY()
 	int32 NumLods = NumInitialCollectionLods;
+PRAGMA_ENABLE_DEPRECATION_WARNINGS
 	/**
 	 * Refresh the asset even if the ClothCollection hasn't changed.
 	 * Note that it is not required to manually refresh the cloth asset, this is done automatically when there is a change in the Dataflow.
diff --git a/Engine/Plugins/ChaosClothAssetEditor/Source/ChaosClothAssetEditor/Private/ChaosClothAsset/ClothEditorMode.cpp b/Engine/Plugins/ChaosClothAssetEditor/Source/ChaosClothAssetEditor/Private/ChaosClothAsset/ClothEditorMode.cpp
--- a/Engine/Plugins/ChaosClothAssetEditor/Source/ChaosClothAssetEditor/Private/ChaosClothAsset/ClothEditorMode.cpp
+++ b/Engine/Plugins/ChaosClothAssetEditor/Source/ChaosClothAssetEditor/Private/ChaosClothAsset/ClothEditorMode.cpp
@@ -353,7 +353,7 @@ void UChaosClothAssetEditorMode::RegisterTools()
 
 	UClothMeshSelectionToolBuilder* SelectionToolBuilder = NewObject<UClothMeshSelectionToolBuilder>();
 	RegisterClothTool(CommandInfos.BeginMeshSelectionTool, FChaosClothAssetEditorCommands::BeginMeshSelectionToolIdentifier, SelectionToolBuilder, SelectionToolBuilder, ConstructionViewportToolsContext);
-	RegisterAddNodeCommand(CommandInfos.AddMeshSelectionNode, FChaosClothAssetSelectionNode::StaticType(), CommandInfos.BeginMeshSelectionTool);
+	RegisterAddNodeCommand(CommandInfos.AddMeshSelectionNode, FChaosClothAssetSelectionNode_v2::StaticType(), CommandInfos.BeginMeshSelectionTool);
 }
 
 bool UChaosClothAssetEditorMode::ShouldToolStartBeAllowed(const FString& ToolIdentifier) const
diff --git a/Engine/Plugins/ChaosClothAssetEditor/Source/ChaosClothAssetEditorTools/Private/ChaosClothAsset/ChaosClothAssetEditorToolsModule.cpp b/Engine/Plugins/ChaosClothAssetEditor/Source/ChaosClothAssetEditorTools/Private/ChaosClothAsset/ChaosClothAssetEditorToolsModule.cpp
--- a/Engine/Plugins/ChaosClothAssetEditor/Source/ChaosClothAssetEditorTools/Private/ChaosClothAsset/ChaosClothAssetEditorToolsModule.cpp
+++ b/Engine/Plugins/ChaosClothAssetEditor/Source/ChaosClothAssetEditorTools/Private/ChaosClothAsset/ChaosClothAssetEditorToolsModule.cpp
@@ -21,7 +21,7 @@ namespace UE::Chaos::ClothAsset
 			const TSharedRef<const FClothToolActionCommandBindings> ClothToolActions = MakeShared<FClothToolActionCommandBindings>();
 
 			ToolRegistry.AddNodeToToolMapping(FChaosClothAssetWeightMapNode::StaticType(), NewObject<UClothEditorWeightMapPaintToolBuilder>(), ClothToolActions);
-			ToolRegistry.AddNodeToToolMapping(FChaosClothAssetSelectionNode::StaticType(), NewObject<UClothMeshSelectionToolBuilder>(), ClothToolActions);
+			ToolRegistry.AddNodeToToolMapping(FChaosClothAssetSelectionNode_v2::StaticType(), NewObject<UClothMeshSelectionToolBuilder>(), ClothToolActions);
 			ToolRegistry.AddNodeToToolMapping(FChaosClothAssetTransferSkinWeightsNode::StaticType(), NewObject<UClothTransferSkinWeightsToolBuilder>(), ClothToolActions);
 		}
 
@@ -29,7 +29,7 @@ namespace UE::Chaos::ClothAsset
 		{
 			Dataflow::FDataflowToolRegistry& ToolRegistry = Dataflow::FDataflowToolRegistry::Get();
 			ToolRegistry.RemoveNodeToToolMapping(FChaosClothAssetWeightMapNode::StaticType());
-			ToolRegistry.RemoveNodeToToolMapping(FChaosClothAssetSelectionNode::StaticType());
+			ToolRegistry.RemoveNodeToToolMapping(FChaosClothAssetSelectionNode_v2::StaticType());
 			ToolRegistry.RemoveNodeToToolMapping(FChaosClothAssetTransferSkinWeightsNode::StaticType());
 		}
 	};
diff --git a/Engine/Plugins/ChaosClothAssetEditor/Source/ChaosClothAssetEditorTools/Private/ChaosClothAsset/ClothEditorToolBuilders.cpp b/Engine/Plugins/ChaosClothAssetEditor/Source/ChaosClothAssetEditorTools/Private/ChaosClothAsset/ClothEditorToolBuilders.cpp
--- a/Engine/Plugins/ChaosClothAssetEditor/Source/ChaosClothAssetEditorTools/Private/ChaosClothAsset/ClothEditorToolBuilders.cpp
+++ b/Engine/Plugins/ChaosClothAssetEditor/Source/ChaosClothAssetEditorTools/Private/ChaosClothAsset/ClothEditorToolBuilders.cpp
@@ -91,7 +91,7 @@ bool UClothMeshSelectionToolBuilder::CanBuildTool(const FToolBuilderState& Scene
 {
 	if (UDataflowContextObject* const DataflowContextObject = SceneState.ToolManager->GetContextObjectStore()->FindContext<UDataflowContextObject>())
 	{
-		return DataflowContextObject->GetSelectedNodeOfType<FChaosClothAssetSelectionNode>() != nullptr && (SceneState.TargetManager->CountSelectedAndTargetable(SceneState, GetTargetRequirements()) == 1);
+		return DataflowContextObject->GetSelectedNodeOfType<FChaosClothAssetSelectionNode_v2>() != nullptr && (SceneState.TargetManager->CountSelectedAndTargetable(SceneState, GetTargetRequirements()) == 1);
 	}
 
 	return false;
diff --git a/Engine/Plugins/ChaosClothAssetEditor/Source/ChaosClothAssetEditorTools/Private/ChaosClothAsset/ClothMeshSelectionTool.cpp b/Engine/Plugins/ChaosClothAssetEditor/Source/ChaosClothAssetEditorTools/Private/ChaosClothAsset/ClothMeshSelectionTool.cpp
--- a/Engine/Plugins/ChaosClothAssetEditor/Source/ChaosClothAssetEditorTools/Private/ChaosClothAsset/ClothMeshSelectionTool.cpp
+++ b/Engine/Plugins/ChaosClothAssetEditor/Source/ChaosClothAssetEditorTools/Private/ChaosClothAsset/ClothMeshSelectionTool.cpp
@@ -424,7 +424,7 @@ bool UClothMeshSelectionTool::GetSelectedNodeInfo(FString& OutSelectionName, UE:
 {
 	using namespace UE::Chaos::ClothAsset;
 
-	SelectionNodeToUpdate = DataflowContextObject->GetSelectedNodeOfType<FChaosClothAssetSelectionNode>();
+	SelectionNodeToUpdate = DataflowContextObject->GetSelectedNodeOfType<FChaosClothAssetSelectionNode_v2>();
 	check(SelectionNodeToUpdate);
 
 	// We need to sanitize the incoming indices, as the user can manually set them to anything on the node
@@ -507,7 +507,7 @@ bool UClothMeshSelectionTool::GetSelectedNodeInfo(FString& OutSelectionName, UE:
 	SelectionNodeToUpdate->CalculateFinalSet(InputSelectionSet, FinalSet);
 	AppendSet(SelectionNodeToUpdate->Group, FinalSet);
 
-	OutSelectionName = SelectionNodeToUpdate->Name;
+	OutSelectionName = SelectionNodeToUpdate->OutputName.StringValue;
 	OutOverrideType = SelectionNodeToUpdate->SelectionOverrideType;
 
 	return true;
@@ -524,7 +524,7 @@ void UClothMeshSelectionTool::UpdateSelectedNode()
 	if (UDataflow* const Dataflow = DataflowContextObject->GetDataflowAsset())
 	{
 		GetToolManager()->GetContextTransactionsAPI()->AppendChange(Dataflow, 
-			FChaosClothAssetSelectionNode::MakeWeightMapNodeChange(*SelectionNodeToUpdate),
+			FChaosClothAssetSelectionNode_v2::MakeSelectedNodeChange(*SelectionNodeToUpdate),
 			LOCTEXT("SelectionNodeChangeDescription", "Update Selection Node"));
 	}
 
@@ -592,7 +592,7 @@ void UClothMeshSelectionTool::UpdateSelectedNode()
 		}
 	};
 
-	SelectionNodeToUpdate->Name = ToolProperties->Name;
+	SelectionNodeToUpdate->OutputName.StringValue = ToolProperties->Name;
 	SelectionNodeToUpdate->SelectionOverrideType = ToolProperties->SelectionOverrideType;
 
 	TSet<int32> FinalSet;
@@ -615,12 +615,6 @@ void UClothMeshSelectionTool::ApplyAction(EClothMeshSelectionToolActions ActionT
 {
 	switch (ActionType)
 	{
-	case EClothMeshSelectionToolActions::ImportFromCollection:
-		ImportFromCollection(/*bImportFromSecondarySet = */ false);
-		break;
-	case EClothMeshSelectionToolActions::ImportSecondaryFromCollection:
-		ImportFromCollection(/*bImportFromSecondarySet = */ true);
-		break;
 	case EClothMeshSelectionToolActions::GrowSelection:
 		SelectionMechanic->GrowSelection();
 		break;
@@ -636,118 +630,5 @@ void UClothMeshSelectionTool::ApplyAction(EClothMeshSelectionToolActions ActionT
 	}
 }
 
-void UClothMeshSelectionTool::ImportFromCollection(bool bImportFromSecondarySet)
-{
-	if (const TSharedPtr<const FManagedArrayCollection> ClothCollection = DataflowContextObject->GetSelectedCollection())
-	{
-		ensure(DataflowContextObject->IsUsingInputCollection());
-
-		using namespace UE::Chaos::ClothAsset;
-		const FCollectionClothSelectionConstFacade SelectionFacade(ClothCollection.ToSharedRef());
-		if (SelectionFacade.IsValid())
-		{
-			const EClothPatternVertexType ViewMode = DataflowViewModeToClothViewMode(DataflowContextObject->GetConstructionViewMode());
-
-			FName GroupName = ClothCollectionGroup::SimVertices2D;
-			if (SelectionMechanic->Properties->bSelectVertices)
-			{
-				check(!SelectionMechanic->Properties->bSelectEdges);
-				check(!SelectionMechanic->Properties->bSelectFaces);
-
-				switch (ViewMode)
-				{
-				case EClothPatternVertexType::Sim2D:
-					GroupName = ClothCollectionGroup::SimVertices2D;
-					break;
-				case EClothPatternVertexType::Sim3D:
-					GroupName = ClothCollectionGroup::SimVertices3D;
-					break;
-				case EClothPatternVertexType::Render:
-					GroupName = ClothCollectionGroup::RenderVertices;
-					break;
-				}
-			}
-			else
-			{
-				switch (ViewMode)
-				{
-				case EClothPatternVertexType::Sim2D:
-				case EClothPatternVertexType::Sim3D:
-					GroupName = ClothCollectionGroup::SimFaces;
-					break;
-				case EClothPatternVertexType::Render:
-					GroupName = ClothCollectionGroup::RenderFaces;
-					break;
-				}
-			}
-			const FName InSelectionName(ToolProperties->Name);
-
-			if (const TSet<int32>* const SelectionSet = bImportFromSecondarySet ? SelectionFacade.FindSelectionSecondarySet(InSelectionName) : SelectionFacade.FindSelectionSet(InSelectionName))
-			{
-				const FName& ExistingSelectionGroup = bImportFromSecondarySet ? SelectionFacade.GetSelectionSecondaryGroup(InSelectionName) : SelectionFacade.GetSelectionGroup(InSelectionName);
-
-				if (ExistingSelectionGroup == GroupName)
-				{
-					auto AppendVerticesIfValid = [this]<typename T>(TSet<int32>&Dest, const T & Source)
-					{
-						PreviewMesh->ProcessMesh([this, &Dest, &Source](const UE::Geometry::FDynamicMesh3& Mesh)
-						{
-							for (const int32& VertexIndex : Source)
-							{
-								if (Mesh.IsVertex(VertexIndex))
-								{
-									Dest.Add(VertexIndex);
-								}
-							}
-						});
-					};
-
-					auto AppendFacesIfValid = [this](TSet<int32>& Dest, const TSet<int32>& Source)
-					{
-						PreviewMesh->ProcessMesh([this, &Dest, &Source](const UE::Geometry::FDynamicMesh3& Mesh)
-						{
-							for (const int32& FaceIndex : Source)
-							{
-								if (Mesh.IsTriangle(FaceIndex))
-								{
-									Dest.Add(FaceIndex);
-								}
-							}
-						});
-					};
-
-					UE::Geometry::FGroupTopologySelection OutSelection;
-
-					if (GroupName == ClothCollectionGroup::SimVertices2D ||
-						GroupName == ClothCollectionGroup::SimVertices3D ||
-						GroupName == ClothCollectionGroup::RenderVertices)
-					{
-						if (bHasNonManifoldMapping)
-						{
-							for (const int32 SelectionIndex : *SelectionSet)
-							{
-								if (SelectionIndex < SelectionToDynamicMesh.Num())		// Could be loading a render mesh selection where NumRenderVertices > NumSimVertices
-								{
-									AppendVerticesIfValid(OutSelection.SelectedCornerIDs, SelectionToDynamicMesh[SelectionIndex]);
-								}
-							}
-						}
-						else
-						{
-							AppendVerticesIfValid(OutSelection.SelectedCornerIDs, *SelectionSet);
-						}
-					}
-					else if (GroupName == ClothCollectionGroup::SimFaces ||
-						GroupName == ClothCollectionGroup::RenderFaces)
-					{
-						AppendFacesIfValid(OutSelection.SelectedGroupIDs, *SelectionSet);
-					}
-					SelectionMechanic->SetSelection(OutSelection);
-				}
-			}
-		}
-	}
-}
-
 
 #undef LOCTEXT_NAMESPACE
diff --git a/Engine/Plugins/ChaosClothAssetEditor/Source/ChaosClothAssetEditorTools/Private/ChaosClothAsset/ClothMeshSelectionTool.h b/Engine/Plugins/ChaosClothAssetEditor/Source/ChaosClothAssetEditorTools/Private/ChaosClothAsset/ClothMeshSelectionTool.h
--- a/Engine/Plugins/ChaosClothAssetEditor/Source/ChaosClothAssetEditorTools/Private/ChaosClothAsset/ClothMeshSelectionTool.h
+++ b/Engine/Plugins/ChaosClothAssetEditor/Source/ChaosClothAssetEditorTools/Private/ChaosClothAsset/ClothMeshSelectionTool.h
@@ -9,7 +9,7 @@
 class UPolygonSelectionMechanic;
 class UDataflowContextObject;
 class UPreviewMesh;
-struct FChaosClothAssetSelectionNode;
+struct FChaosClothAssetSelectionNode_v2;
 enum class EChaosClothAssetSelectionOverrideType : uint8;
 
 namespace UE::Geometry
@@ -23,9 +23,6 @@ enum class EClothMeshSelectionToolActions
 {
 	NoAction,
 
-	ImportFromCollection,
-	ImportSecondaryFromCollection,
-
 	GrowSelection,
 	ShrinkSelection,
 	FloodSelection,
@@ -57,18 +54,6 @@ public:
 
 	void PostAction(EClothMeshSelectionToolActions Action);
 
-	UFUNCTION(CallInEditor, Category = Import)
-	void ImportFromCollection()
-	{
-		PostAction(EClothMeshSelectionToolActions::ImportFromCollection);
-	}
-
-	UFUNCTION(CallInEditor, Category = Migrate)
-	void ImportSecondaryFromCollection()
-	{
-		PostAction(EClothMeshSelectionToolActions::ImportSecondaryFromCollection);
-	}
-
 	UFUNCTION(CallInEditor, Category = Selection)
 	void GrowSelection()
 	{
@@ -168,7 +153,7 @@ private:
 	TArray<int32> DynamicMeshToSelection;
 	TArray<TArray<int32>> SelectionToDynamicMesh;
 
-	FChaosClothAssetSelectionNode* SelectionNodeToUpdate = nullptr;
+	FChaosClothAssetSelectionNode_v2* SelectionNodeToUpdate = nullptr;
 	TSet<int32> InputSelectionSet;
 	//
 	// Action support
@@ -185,6 +170,5 @@ private:
 	EClothMeshSelectionToolActions PendingAction;
 	virtual void ApplyAction(EClothMeshSelectionToolActions ActionType);
 
-	void ImportFromCollection(bool bImportFromSecondarySet);
 };
 
