Cloth Editor: when nothing is selected in the preview scene, disable the numeric transform inputs

diff --git a/Engine/Plugins/ChaosClothAssetEditor/Source/ChaosClothAssetEditor/Private/ChaosClothAsset/ClothEditor3DViewportClient.cpp b/Engine/Plugins/ChaosClothAssetEditor/Source/ChaosClothAssetEditor/Private/ChaosClothAsset/ClothEditor3DViewportClient.cpp
--- a/Engine/Plugins/ChaosClothAssetEditor/Source/ChaosClothAssetEditor/Private/ChaosClothAsset/ClothEditor3DViewportClient.cpp
+++ b/Engine/Plugins/ChaosClothAssetEditor/Source/ChaosClothAssetEditor/Private/ChaosClothAsset/ClothEditor3DViewportClient.cpp
@@ -428,8 +428,16 @@ void FChaosClothAssetEditor3DViewportClient::ComponentSelectionChanged(UObject*
 			Gizmo->SetVisibility(false);
 		}
 
-		// TODO: Set UChaosClothPreviewSceneDescription::bValidSelectionForTransform here once we figure out why it's not
-		// properly affecting the EditCondition on the other properties (UE-189504)
+		if (const TSharedPtr<FChaosClothPreviewScene> PinnedClothPreviewScene = ClothPreviewScene.Pin())
+		{
+			if (UChaosClothPreviewSceneDescription* const SceneDescription = PinnedClothPreviewScene->GetPreviewSceneDescription())
+			{
+				SceneDescription->bValidSelectionForTransform = (Components.Num() > 0);
+
+				FPropertyChangedEvent Event(UChaosClothPreviewSceneDescription::StaticClass()->FindPropertyByName(GET_MEMBER_NAME_CHECKED(UChaosClothPreviewSceneDescription, bValidSelectionForTransform)));
+				SceneDescription->PostEditChangeProperty(Event);
+			}
+		}
 
 	}
 }
diff --git a/Engine/Plugins/ChaosClothAssetEditor/Source/ChaosClothAssetEditor/Private/ChaosClothAsset/ClothEditorPreviewScene.h b/Engine/Plugins/ChaosClothAssetEditor/Source/ChaosClothAssetEditor/Private/ChaosClothAsset/ClothEditorPreviewScene.h
--- a/Engine/Plugins/ChaosClothAssetEditor/Source/ChaosClothAssetEditor/Private/ChaosClothAsset/ClothEditorPreviewScene.h
+++ b/Engine/Plugins/ChaosClothAssetEditor/Source/ChaosClothAssetEditor/Private/ChaosClothAsset/ClothEditorPreviewScene.h
@@ -53,29 +53,25 @@ public:
 	UPROPERTY(EditAnywhere, Transient, Category = "SkeletalMesh")
 	bool bPostProcessBlueprint;
 
-	UPROPERTY(EditAnywhere, Transient, Category = "Transform", Meta=(DisplayName="Location"))
+	UPROPERTY(EditAnywhere, Transient, Category = "Transform", Meta = (EditCondition = "bValidSelectionForTransform", HideEditConditionToggle))
 	FVector3d Translation = FVector3d::ZeroVector;
 
-	UPROPERTY(EditAnywhere, Transient, Category = "Transform")
+	UPROPERTY(EditAnywhere, Transient, Category = "Transform", Meta = (EditCondition = "bValidSelectionForTransform", HideEditConditionToggle))
 	FVector3d Rotation = FVector3d::ZeroVector;
 
-	UPROPERTY(EditAnywhere, Transient, Category = "Transform", Meta = (AllowPreserveRatio))
+	UPROPERTY(EditAnywhere, Transient, Category = "Transform", Meta = (AllowPreserveRatio, EditCondition = "bValidSelectionForTransform", HideEditConditionToggle))
 	FVector3d Scale = FVector3d::OneVector;
 
 	UPROPERTY(EditAnywhere, Transient, Category = "ClothComponent", Meta = (UIMin = 0.0, UIMax = 10.0, ClampMin = 0.0, ClampMax = 10000.0))
 	float SolverGeometryScale = 1.f;
 
-	// TODO: We should be able to hook this boolean property up to the EditCondition meta tag for the properties above and toggle it
-	// on and off when the selection changes in the scene. However the EditCondition does not seem to propagate for some reason, 
-	// even if we manually call PostEditChangeProperty() after toggling it. It will take some more digging to figure out exactly
-	// what's going on. (UE-189504)
-	//UPROPERTY(Transient)
-	//bool bValidSelectionForTransform = true;
-
-private:
+	UPROPERTY(Transient)
+	bool bValidSelectionForTransform = false;
 
 	virtual void PostEditChangeProperty(struct FPropertyChangedEvent& PropertyChangedEvent) override;
 
+private:
+
 	virtual void PostTransacted(const FTransactionObjectEvent& TransactionEvent) override;
 
 	UE::Chaos::ClothAsset::FChaosClothPreviewScene* PreviewScene;
diff --git a/Engine/Plugins/ChaosClothAssetEditor/Source/ChaosClothAssetEditor/Private/ChaosClothAsset/ClothEditorToolkit.cpp b/Engine/Plugins/ChaosClothAssetEditor/Source/ChaosClothAssetEditor/Private/ChaosClothAsset/ClothEditorToolkit.cpp
--- a/Engine/Plugins/ChaosClothAssetEditor/Source/ChaosClothAssetEditor/Private/ChaosClothAsset/ClothEditorToolkit.cpp
+++ b/Engine/Plugins/ChaosClothAssetEditor/Source/ChaosClothAssetEditor/Private/ChaosClothAsset/ClothEditorToolkit.cpp
@@ -44,6 +44,7 @@
 #include "FileHelpers.h"
 #include "UObject/PackageReload.h"
 #include "ContextObjectStore.h"
+#include "SClothEditorAdvancedPreviewDetailsTab.h"
 
 #define LOCTEXT_NAMESPACE "ChaosClothAssetEditorToolkit"
 
@@ -795,12 +796,12 @@ void FChaosClothAssetEditorToolkit::InitDetailsViewPanel()
 	TArray<FAdvancedPreviewSceneModule::FDetailDelegates> Delegates;	
 
 	ensure(ClothPreviewScene.IsValid());
-	FAdvancedPreviewSceneModule& AdvancedPreviewSceneModule = FModuleManager::LoadModuleChecked<FAdvancedPreviewSceneModule>("AdvancedPreviewScene");
-	AdvancedPreviewSettingsWidget = AdvancedPreviewSceneModule.CreateAdvancedPreviewSceneSettingsWidget(ClothPreviewScene.ToSharedRef(), 
-		ClothPreviewScene->GetPreviewSceneDescription(),
-		TArray<FAdvancedPreviewSceneModule::FDetailCustomizationInfo>(), 
-		TArray<FAdvancedPreviewSceneModule::FPropertyTypeCustomizationInfo>(),
-		Delegates);
+
+	AdvancedPreviewSettingsWidget = SNew(SChaosClothEditorAdvancedPreviewDetailsTab, ClothPreviewScene.ToSharedRef())
+		.AdditionalSettings(ClothPreviewScene->GetPreviewSceneDescription())
+		.DetailCustomizations(TArray<FAdvancedPreviewSceneModule::FDetailCustomizationInfo>())
+		.PropertyTypeCustomizations(TArray<FAdvancedPreviewSceneModule::FPropertyTypeCustomizationInfo>())
+		.Delegates(Delegates);
 
 	if (PreviewSceneDockTab.IsValid())
 	{
diff --git a/Engine/Plugins/ChaosClothAssetEditor/Source/ChaosClothAssetEditor/Private/ChaosClothAsset/SClothEditorAdvancedPreviewDetailsTab.cpp b/Engine/Plugins/ChaosClothAssetEditor/Source/ChaosClothAssetEditor/Private/ChaosClothAsset/SClothEditorAdvancedPreviewDetailsTab.cpp
new file mode 100644
--- /dev/null
+++ b/Engine/Plugins/ChaosClothAssetEditor/Source/ChaosClothAssetEditor/Private/ChaosClothAsset/SClothEditorAdvancedPreviewDetailsTab.cpp
@@ -0,0 +1,22 @@
+// Copyright Epic Games, Inc. All Rights Reserved.
+
+#include "SClothEditorAdvancedPreviewDetailsTab.h"
+#include "ChaosClothAsset/ClothEditorPreviewScene.h"
+#include "IDetailsView.h"
+
+SChaosClothEditorAdvancedPreviewDetailsTab::SChaosClothEditorAdvancedPreviewDetailsTab()
+	: SAdvancedPreviewDetailsTab()
+{
+	PropertyChangedDelegateHandle = FCoreUObjectDelegates::OnObjectPropertyChanged.AddLambda([this](UObject* Object, struct FPropertyChangedEvent& Event)
+	{
+		if (Object->IsA<UChaosClothPreviewSceneDescription>())
+		{
+			SettingsView->InvalidateCachedState();
+		}
+	});
+}
+
+SChaosClothEditorAdvancedPreviewDetailsTab::~SChaosClothEditorAdvancedPreviewDetailsTab()
+{
+	FCoreUObjectDelegates::OnObjectPropertyChanged.Remove(PropertyChangedDelegateHandle);
+}
diff --git a/Engine/Plugins/ChaosClothAssetEditor/Source/ChaosClothAssetEditor/Private/ChaosClothAsset/SClothEditorAdvancedPreviewDetailsTab.h b/Engine/Plugins/ChaosClothAssetEditor/Source/ChaosClothAssetEditor/Private/ChaosClothAsset/SClothEditorAdvancedPreviewDetailsTab.h
new file mode 100644
--- /dev/null
+++ b/Engine/Plugins/ChaosClothAssetEditor/Source/ChaosClothAssetEditor/Private/ChaosClothAsset/SClothEditorAdvancedPreviewDetailsTab.h
@@ -0,0 +1,18 @@
+// Copyright Epic Games, Inc. All Rights Reserved.
+
+#pragma once
+
+#include "SAdvancedPreviewDetailsTab.h"
+
+// Custom subclass of SAdvancedPreviewDetailsTab that allows invalidating the cached state of the SettingsView when the preview scene changes
+class CHAOSCLOTHASSETEDITOR_API SChaosClothEditorAdvancedPreviewDetailsTab : public SAdvancedPreviewDetailsTab
+{
+public:
+
+	SChaosClothEditorAdvancedPreviewDetailsTab();
+	virtual ~SChaosClothEditorAdvancedPreviewDetailsTab() override;
+
+private:
+
+	FDelegateHandle PropertyChangedDelegateHandle;
+};
