Dataflow Editor:
- Allow changing the construction view mode while a tool is running, if the tool supports it
- Make the Cloth Weight Map Paint Tool support this
- If the tool does not support changing the view mode, follow the Cloth Editor behavior of shutting down and restarting the tool. (Initially the Selection Tool will have this behavior, but in a future CL we will likely enable changing view modes on the fly for this tool as well.)
- If the tool requires a view mode that we are not already in, automatically change to a valid view mode before starting the tool

diff --git a/Engine/Plugins/ChaosClothAssetEditor/Source/ChaosClothAssetEditorTools/Private/ChaosClothAsset/ClothEditorToolBuilders.cpp b/Engine/Plugins/ChaosClothAssetEditor/Source/ChaosClothAssetEditorTools/Private/ChaosClothAsset/ClothEditorToolBuilders.cpp
--- a/Engine/Plugins/ChaosClothAssetEditor/Source/ChaosClothAssetEditorTools/Private/ChaosClothAsset/ClothEditorToolBuilders.cpp
+++ b/Engine/Plugins/ChaosClothAssetEditor/Source/ChaosClothAssetEditorTools/Private/ChaosClothAsset/ClothEditorToolBuilders.cpp
@@ -20,6 +20,72 @@
 
 // ------------------- Weight Map Paint Tool -------------------
 
+void UClothEditorWeightMapPaintToolBuilder::GetSupportedViewModes(const UDataflowContextObject& ContextObject, TArray<const Dataflow::IDataflowConstructionViewMode*>& Modes) const
+{
+	using namespace UE::Chaos::ClothAsset;
+
+	const Dataflow::FRenderingViewModeFactory& Factory = Dataflow::FRenderingViewModeFactory::GetInstance();
+	const Dataflow::IDataflowConstructionViewMode* const Sim2DMode = Factory.GetViewMode(ClothViewModeToDataflowViewModeName(EClothPatternVertexType::Sim2D));
+	const Dataflow::IDataflowConstructionViewMode* const Sim3DMode = Factory.GetViewMode(ClothViewModeToDataflowViewModeName(EClothPatternVertexType::Sim3D));
+	const Dataflow::IDataflowConstructionViewMode* const RenderMode = Factory.GetViewMode(ClothViewModeToDataflowViewModeName(EClothPatternVertexType::Render));
+
+	checkf(Sim2DMode, TEXT("Couldn't find DataflowConstructionViewMode corresponding to EClothPatternVertexType::Sim2D"));
+	checkf(Sim3DMode, TEXT("Couldn't find DataflowConstructionViewMode corresponding to EClothPatternVertexType::Sim3D"));
+	checkf(RenderMode, TEXT("Couldn't find DataflowConstructionViewMode corresponding to EClothPatternVertexType::Render"));
+
+	const FChaosClothAssetWeightMapNode* const WeightMapNode = ContextObject.GetSelectedNodeOfType<FChaosClothAssetWeightMapNode>();
+	if (WeightMapNode)
+	{
+		if (WeightMapNode->MeshTarget == EChaosClothAssetWeightMapMeshTarget::Render)
+		{
+			Modes.Add(RenderMode);
+		}
+		else
+		{
+			check(WeightMapNode->MeshTarget == EChaosClothAssetWeightMapMeshTarget::Simulation);
+			Modes.Add(Sim2DMode);
+			Modes.Add(Sim3DMode);
+		}
+	}
+	else
+	{
+		const EClothPatternVertexType CurrentViewMode = DataflowViewModeToClothViewMode(ContextObject.GetConstructionViewMode());
+
+		if (CurrentViewMode == EClothPatternVertexType::Render)
+		{
+			Modes.Add(RenderMode);
+		}
+		else
+		{
+			Modes.Add(Sim2DMode);
+			Modes.Add(Sim3DMode);
+		}
+	}
+}
+
+bool UClothEditorWeightMapPaintToolBuilder::CanSceneStateChange(const UInteractiveTool* ActiveTool, const FToolBuilderState& SceneState) const
+{
+	return ActiveTool->IsA<UClothEditorWeightMapPaintTool>();
+}
+
+void UClothEditorWeightMapPaintToolBuilder::SceneStateChanged(UInteractiveTool* ActiveTool, const FToolBuilderState& SceneState)
+{
+	check(CanSceneStateChange(ActiveTool, SceneState));
+
+	UClothEditorWeightMapPaintTool* const PaintTool = Cast<UClothEditorWeightMapPaintTool>(ActiveTool);
+	checkf(PaintTool, TEXT("Expected the ActiveTool to be UClothEditorWeightMapPaintTool"));
+
+	UToolTarget* const Target = SceneState.TargetManager->BuildFirstSelectedTargetable(SceneState, GetTargetRequirements());
+	check(Target);
+	check(Target->IsValid());
+	PaintTool->SetTarget(Target);
+	PaintTool->NotifyTargetChanged();
+
+	// These are likely to be empty functions but are called here for completeness (see UInteractiveToolManager::ActivateToolInternal())
+	PostBuildTool(ActiveTool, SceneState);
+	PostSetupTool(ActiveTool, SceneState);
+}
+
 void UClothEditorWeightMapPaintToolBuilder::GetSupportedViewModes(const UDataflowContextObject& ContextObject, TArray<UE::Chaos::ClothAsset::EClothPatternVertexType>& Modes) const
 {
 	using namespace UE::Chaos::ClothAsset;
@@ -71,8 +137,65 @@ UMeshSurfacePointTool* UClothEditorWeightMapPaintToolBuilder::CreateNewTool(cons
 	return PaintTool;
 }
 
+
 // ------------------- Selection Tool -------------------
 
+
+void UClothMeshSelectionToolBuilder::GetSupportedViewModes(const UDataflowContextObject& ContextObject, TArray<const Dataflow::IDataflowConstructionViewMode*>& Modes) const
+{
+	using namespace UE::Chaos::ClothAsset;
+
+	const Dataflow::FRenderingViewModeFactory& Factory = Dataflow::FRenderingViewModeFactory::GetInstance();
+	const Dataflow::IDataflowConstructionViewMode* const Sim2DMode = Factory.GetViewMode(ClothViewModeToDataflowViewModeName(EClothPatternVertexType::Sim2D));
+	const Dataflow::IDataflowConstructionViewMode* const Sim3DMode = Factory.GetViewMode(ClothViewModeToDataflowViewModeName(EClothPatternVertexType::Sim3D));
+	const Dataflow::IDataflowConstructionViewMode* const RenderMode = Factory.GetViewMode(ClothViewModeToDataflowViewModeName(EClothPatternVertexType::Render));
+
+	checkf(Sim2DMode, TEXT("Couldn't find DataflowConstructionViewMode corresponding to EClothPatternVertexType::Sim2D"));
+	checkf(Sim3DMode, TEXT("Couldn't find DataflowConstructionViewMode corresponding to EClothPatternVertexType::Sim3D"));
+	checkf(RenderMode, TEXT("Couldn't find DataflowConstructionViewMode corresponding to EClothPatternVertexType::Render"));
+
+	const FChaosClothAssetSelectionNode_v2* const SelectionNode = ContextObject.GetSelectedNodeOfType<FChaosClothAssetSelectionNode_v2>();
+	if (SelectionNode)
+	{
+		if (SelectionNode->Group.Name == ClothCollectionGroup::RenderVertices.ToString() || SelectionNode->Group.Name == ClothCollectionGroup::RenderFaces.ToString())
+		{
+			Modes.Add(RenderMode);
+			return;
+		}
+		
+		if (SelectionNode->Group.Name == ClothCollectionGroup::SimVertices2D.ToString() || SelectionNode->Group.Name == ClothCollectionGroup::SimVertices3D.ToString() || SelectionNode->Group.Name == ClothCollectionGroup::SimFaces.ToString())
+		{
+			Modes.Add(Sim2DMode);
+			Modes.Add(Sim3DMode);
+			return;
+		}
+	}
+
+	// No node selected or no valid group name set in the node -- use the current view mode to decide
+
+	const EClothPatternVertexType CurrentViewMode = DataflowViewModeToClothViewMode(ContextObject.GetConstructionViewMode());
+
+	if (CurrentViewMode == EClothPatternVertexType::Render)
+	{
+		Modes.Add(RenderMode);
+	}
+	else
+	{
+		Modes.Add(Sim2DMode);
+		Modes.Add(Sim3DMode);
+	}
+}
+
+bool UClothMeshSelectionToolBuilder::CanSceneStateChange(const UInteractiveTool* ActiveTool, const FToolBuilderState& SceneState) const
+{
+	return false;
+}
+
+void UClothMeshSelectionToolBuilder::SceneStateChanged(UInteractiveTool* ActiveTool, const FToolBuilderState& SceneState)
+{
+	check(CanSceneStateChange(ActiveTool, SceneState));
+}
+
 void UClothMeshSelectionToolBuilder::GetSupportedViewModes(const UDataflowContextObject& ContextObject, TArray<UE::Chaos::ClothAsset::EClothPatternVertexType>& Modes) const
 {
 	// TODO: When the Secondary Selection set is removed, update this function to be similar to UClothEditorWeightMapPaintToolBuilder::GetSupportedViewModes above
@@ -81,6 +204,7 @@ void UClothMeshSelectionToolBuilder::GetSupportedViewModes(const UDataflowContex
 	Modes.Add(UE::Chaos::ClothAsset::EClothPatternVertexType::Render);
 }
 
+
 const FToolTargetTypeRequirements& UClothMeshSelectionToolBuilder::GetTargetRequirements() const
 {
 	static FToolTargetTypeRequirements TypeRequirements(UPrimitiveComponentBackedTarget::StaticClass());
@@ -116,6 +240,26 @@ UInteractiveTool* UClothMeshSelectionToolBuilder::BuildTool(const FToolBuilderSt
 
 // ------------------- Skin Weight Transfer Tool -------------------
 
+
+void UClothTransferSkinWeightsToolBuilder::GetSupportedViewModes(const UDataflowContextObject& ContextObject, TArray<const Dataflow::IDataflowConstructionViewMode*>& Modes) const
+{
+	using namespace UE::Chaos::ClothAsset;
+	const Dataflow::FRenderingViewModeFactory& Factory = Dataflow::FRenderingViewModeFactory::GetInstance();
+	const Dataflow::IDataflowConstructionViewMode* const Sim3DMode = Factory.GetViewMode(ClothViewModeToDataflowViewModeName(EClothPatternVertexType::Sim3D));
+	Modes.Add(Sim3DMode);
+}
+
+bool UClothTransferSkinWeightsToolBuilder::CanSceneStateChange(const UInteractiveTool* ActiveTool, const FToolBuilderState& SceneState) const
+{
+	return false;
+}
+
+void UClothTransferSkinWeightsToolBuilder::SceneStateChanged(UInteractiveTool* ActiveTool, const FToolBuilderState& SceneState)
+{
+	check(CanSceneStateChange(ActiveTool, SceneState));
+}
+
+
 void UClothTransferSkinWeightsToolBuilder::GetSupportedViewModes(const UDataflowContextObject& ContextObject, TArray<UE::Chaos::ClothAsset::EClothPatternVertexType>& Modes) const
 {
 	Modes.Add(UE::Chaos::ClothAsset::EClothPatternVertexType::Sim3D);
diff --git a/Engine/Plugins/ChaosClothAssetEditor/Source/ChaosClothAssetEditorTools/Private/ChaosClothAsset/ClothWeightMapPaintTool.cpp b/Engine/Plugins/ChaosClothAssetEditor/Source/ChaosClothAssetEditorTools/Private/ChaosClothAsset/ClothWeightMapPaintTool.cpp
--- a/Engine/Plugins/ChaosClothAssetEditor/Source/ChaosClothAssetEditorTools/Private/ChaosClothAsset/ClothWeightMapPaintTool.cpp
+++ b/Engine/Plugins/ChaosClothAssetEditor/Source/ChaosClothAssetEditorTools/Private/ChaosClothAsset/ClothWeightMapPaintTool.cpp
@@ -121,53 +121,6 @@
 	FActorSpawnParameters SpawnInfo;
 	PreviewMeshActor = TargetWorld->SpawnActor<AInternalToolFrameworkActor>(FVector::ZeroVector, FRotator::ZeroRotator, SpawnInfo);
 	DynamicMeshComponent = NewObject<UDynamicMeshComponent>(PreviewMeshActor);
-	InitializeSculptMeshComponent(DynamicMeshComponent, PreviewMeshActor);
 
-	// assign materials
-	FComponentMaterialSet MaterialSet = UE::ToolTarget::GetMaterialSet(Target);
-	for (int k = 0; k < MaterialSet.Materials.Num(); ++k)
-	{
-		DynamicMeshComponent->SetMaterial(k, MaterialSet.Materials[k]);
-	}
-
-	DynamicMeshComponent->SetInvalidateProxyOnChangeEnabled(false);
-	OnDynamicMeshComponentChangedHandle = DynamicMeshComponent->OnMeshVerticesChanged.AddUObject(this, &UClothEditorWeightMapPaintTool::OnDynamicMeshComponentChanged);
-
-	FDynamicMesh3* Mesh = GetSculptMesh();
-	Mesh->EnableVertexColors(FVector3f::One());
-	Mesh->Attributes()->EnablePrimaryColors();
-	Mesh->Attributes()->PrimaryColors()->CreateFromPredicate([](int ParentVID, int TriIDA, int TriIDB) {return true; }, 0.f);
-	FAxisAlignedBox3d Bounds = Mesh->GetBounds(true);
-
-	TFuture<void> PrecomputeFuture = Async(UE::Chaos::ClothAsset::Private::WeightPaintToolAsyncExecTarget, [&]()
-	{
-		PrecomputeFilterData();
-	});
-
-	TFuture<void> OctreeFuture = Async(UE::Chaos::ClothAsset::Private::WeightPaintToolAsyncExecTarget, [&]()
-	{
-		// initialize dynamic octree
-		if (Mesh->TriangleCount() > 100000)
-		{
-			Octree.RootDimension = Bounds.MaxDim() / 10.0;
-			Octree.SetMaxTreeDepth(4);
-		}
-		else
-		{
-			Octree.RootDimension = Bounds.MaxDim();
-			Octree.SetMaxTreeDepth(8);
-		}
-		Octree.Initialize(Mesh);
-	});
-
-	// initialize render decomposition
-	TUniquePtr<FMeshRenderDecomposition> Decomp = MakeUnique<FMeshRenderDecomposition>();
-	FMeshRenderDecomposition::BuildChunkedDecomposition(Mesh, &MaterialSet, *Decomp);
-	Decomp->BuildAssociations(Mesh);
-	DynamicMeshComponent->SetExternalDecomposition(MoveTemp(Decomp));
-
-	// initialize brush radius range interval, brush properties
-	InitializeBrushSizeRange(Bounds);
-
 	// Set up control points mechanic
 	PolyLassoMechanic = NewObject<UPolyLassoMarqueeMechanic>(this);
@@ -190,12 +143,6 @@
 	PolygonSelectionMechanic->Properties->bSelectEdges = false;
 	PolygonSelectionMechanic->Properties->bSelectFaces = false;
 	PolygonSelectionMechanic->Properties->bSelectVertices = true;
-	
-	constexpr bool bAutoBuild = true;
-	GradientSelectionTopology = MakeUnique<UE::Geometry::FTriangleGroupTopology>(DynamicMeshComponent->GetMesh(), bAutoBuild);
-	
-	MeshSpatial = MakeUnique<UE::Geometry::FDynamicMeshAABBTree3>(DynamicMeshComponent->GetMesh(), bAutoBuild);
-	PolygonSelectionMechanic->Initialize(DynamicMeshComponent, GradientSelectionTopology.Get(), [this]() { return MeshSpatial.Get(); });
 
 	UpdateWeightMapProperties = NewObject<UClothEditorUpdateWeightMapProperties>(this);
 	UpdateWeightMapProperties->Name = WeightMapNodeToUpdate->OutputName.StringValue;
@@ -273,21 +220,14 @@
 	AddToolPropertySource(UMeshSculptToolBase::GizmoProperties);
 	SetToolPropertySourceEnabled(UMeshSculptToolBase::GizmoProperties, false);
 
-
-	// register watchers
-	FilterProperties->WatchProperty( FilterProperties->PrimaryBrushType,
-		[this](EClothEditorWeightMapPaintBrushType NewType) { UpdateBrushType(NewType); });
-
-	// must call before updating brush type so that we register all brush properties?
-	UMeshSculptToolBase::OnCompleteSetup();
-
-	UpdateBrushType(FilterProperties->PrimaryBrushType);
-	SetActiveSecondaryBrushType((int32)EClothEditorWeightMapPaintBrushType::Erase);
-	
 	ActionsProps = NewObject<UClothEditorMeshWeightMapPaintToolActions>(this);
 	ActionsProps->Initialize(this);
 	AddToolPropertySource(ActionsProps);
 
+	// register watchers
+	FilterProperties->WatchProperty( FilterProperties->PrimaryBrushType,
+		[this](EClothEditorWeightMapPaintBrushType NewType) { UpdateBrushType(NewType); });
+
 	MeshElementsDisplay = NewObject<UMeshElementsVisualizer>(this);
 	MeshElementsDisplay->CreateInWorld(DynamicMeshComponent->GetWorld(), DynamicMeshComponent->GetComponentTransform());
 	if (ensure(MeshElementsDisplay->Settings))
@@ -365,6 +305,9 @@
 	);
 	AddToolPropertySource(ShowHideProperties);
 
+	UpdateWeightMapProperties->Name = WeightMapNodeToUpdate->OutputName.StringValue;
+	UpdateWeightMapProperties->MapOverrideType = WeightMapNodeToUpdate->MapOverrideType;
+	SetToolPropertySourceEnabled(UpdateWeightMapProperties, true);
 
 	// disable view properties
 	SetViewPropertiesEnabled(false);
@@ -375,6 +318,101 @@
 	// configure panels
 	UpdateSubToolType(FilterProperties->SubToolType);
 
+	// must call before updating brush type so that we register all brush properties?
+	UMeshSculptToolBase::OnCompleteSetup();
+
+	UpdateBrushType(FilterProperties->PrimaryBrushType);
+	SetActiveSecondaryBrushType((int32)EClothEditorWeightMapPaintBrushType::Erase);
+
+	SetPrimaryFalloffType(EMeshSculptFalloffType::Smooth);
+
+
+
+	InitializeSculptMeshFromTarget();
+
+	UpdateShowHideProperties();
+
+
+	// Copy weights from selected node to the preview mesh
+	const int32 NumExpectedWeights = bHaveDynamicMeshToWeightConversion ? WeightToDynamicMesh.Num() : GetSculptMesh()->MaxVertexID();
+	TArray<float> CurrentWeights;
+	CurrentWeights.SetNumZeroed(NumExpectedWeights);
+	WeightMapNodeToUpdate->CalculateFinalVertexWeightValues(InputWeightMap, TArrayView<float>(CurrentWeights));
+
+	if (bHaveDynamicMeshToWeightConversion)
+	{
+		for (int32 WeightID = 0; WeightID < CurrentWeights.Num(); ++WeightID)
+		{
+			for (const int32 VertexID : WeightToDynamicMesh[WeightID])
+			{
+				ActiveWeightMap->SetValue(VertexID, &CurrentWeights[WeightID]);
+			}
+		}
+	}
+	else
+	{
+		for (int32 VertexID = 0; VertexID < CurrentWeights.Num(); ++VertexID)
+		{
+			ActiveWeightMap->SetValue(VertexID, &CurrentWeights[VertexID]);
+		}
+	}
+
+	// Initialize vertex colors from attribute layer values
+	UpdateVertexColorOverlay();
+	DynamicMeshComponent->FastNotifyVertexAttributesUpdated(EMeshRenderAttributeFlags::VertexColors);
+
+	PostSetupCheck();
+
+	GetToolManager()->PostInvalidation();
+}
+
+void UClothEditorWeightMapPaintTool::InitializeSculptMeshFromTarget()
+{
+	InitializeSculptMeshComponent(DynamicMeshComponent, PreviewMeshActor);
+
+	// assign materials
+	FComponentMaterialSet MaterialSet = UE::ToolTarget::GetMaterialSet(Target);
+	for (int k = 0; k < MaterialSet.Materials.Num(); ++k)
+	{
+		DynamicMeshComponent->SetMaterial(k, MaterialSet.Materials[k]);
+	}
+
+	FDynamicMesh3* Mesh = GetSculptMesh();
+	Mesh->EnableVertexColors(FVector3f::One());
+	Mesh->Attributes()->EnablePrimaryColors();
+	Mesh->Attributes()->PrimaryColors()->CreateFromPredicate([](int ParentVID, int TriIDA, int TriIDB) {return true; }, 0.f);
+	FAxisAlignedBox3d Bounds = Mesh->GetBounds(true);
+
+	TFuture<void> PrecomputeFuture = Async(UE::Chaos::ClothAsset::Private::WeightPaintToolAsyncExecTarget, [this]()
+	{
+		PrecomputeFilterData();
+	});
+
+	TFuture<void> OctreeFuture = Async(UE::Chaos::ClothAsset::Private::WeightPaintToolAsyncExecTarget, [Mesh, &Bounds, this]()
+	{
+		// initialize dynamic octree
+		if (Mesh->TriangleCount() > 100000)
+		{
+			Octree.RootDimension = Bounds.MaxDim() / 10.0;
+			Octree.SetMaxTreeDepth(4);
+		}
+		else
+		{
+			Octree.RootDimension = Bounds.MaxDim();
+			Octree.SetMaxTreeDepth(8);
+		}
+		Octree.Initialize(Mesh);
+	});
+
+	// initialize render decomposition
+	TUniquePtr<FMeshRenderDecomposition> Decomp = MakeUnique<FMeshRenderDecomposition>();
+	FMeshRenderDecomposition::BuildChunkedDecomposition(Mesh, &MaterialSet, *Decomp);
+	Decomp->BuildAssociations(Mesh);
+	DynamicMeshComponent->SetExternalDecomposition(MoveTemp(Decomp));
+
+	// initialize brush radius range interval, brush properties
+	InitializeBrushSizeRange(Bounds);
+
 	// Setup DynamicMeshToWeight conversion and get Input weight map (if it exists)
 	InputWeightMap = TConstArrayView<float>();
 
@@ -419,6 +457,61 @@
 			const EClothPatternVertexType ViewMode = UE::Chaos::ClothAsset::DataflowViewModeToClothViewMode(DataflowContextObject->GetConstructionViewMode());
 			const bool bIsRenderMode = ViewMode == EClothPatternVertexType::Render;
 
+			// Find the map if it exists.
+			if (TSharedPtr<Dataflow::FEngineContext> DataflowContext = DataflowContextObject->GetDataflowContext())
+			{
+				const FName InputName = WeightMapNodeToUpdate->GetInputName(*DataflowContext);
+				if (bIsRenderMode)
+				{
+					InputWeightMap = Cloth.GetUserDefinedAttribute<float>(InputName, ClothCollectionGroup::RenderVertices);
+				}
+				else
+				{
+					InputWeightMap = Cloth.GetWeightMap(InputName);
+				}
+			}
+		}
+	}
+
+	PrecomputeFuture.Wait();
+	OctreeFuture.Wait();
+
+	// Create an attribute layer to temporarily paint into
+	const int NumAttributeLayers = Mesh->Attributes()->NumWeightLayers();
+	Mesh->Attributes()->SetNumWeightLayers(NumAttributeLayers + 1);
+	ActiveWeightMap = Mesh->Attributes()->GetWeightLayer(NumAttributeLayers);
+	ActiveWeightMap->SetName(FName("PaintLayer"));
+
+	// Setup support for hiding specific triangles
+	DynamicMeshComponent->EnableSecondaryTriangleBuffers(
+		[this](const FDynamicMesh3* Mesh, int32 TriangleID)
+		{
+			return PendingHiddenTriangles.Contains(TriangleID) || HiddenTriangles.Contains(TriangleID);
+		});
+	DynamicMeshComponent->SetSecondaryBuffersVisibility(false);
+
+	// Rebuild mechanics that depend on Mesh topology
+	constexpr bool bAutoBuild = true;
+	GradientSelectionTopology = MakeUnique<UE::Geometry::FTriangleGroupTopology>(DynamicMeshComponent->GetMesh(), bAutoBuild);
+	MeshSpatial = MakeUnique<UE::Geometry::FDynamicMeshAABBTree3>(DynamicMeshComponent->GetMesh(), bAutoBuild);
+	PolygonSelectionMechanic->Initialize(DynamicMeshComponent, GradientSelectionTopology.Get(), [this]() { return MeshSpatial.Get(); });
+}
+
+void UClothEditorWeightMapPaintTool::UpdateShowHideProperties()
+{
+	if (DataflowContextObject)
+	{
+		ensure(DataflowContextObject->IsUsingInputCollection());
+		if (TSharedPtr<const FManagedArrayCollection> ClothCollection = DataflowContextObject->GetSelectedCollection())
+		{
+			using namespace UE::Chaos::ClothAsset;
+
+			FCollectionClothConstFacade Cloth(ClothCollection.ToSharedRef());
+			check(Cloth.IsValid());
+
+			const EClothPatternVertexType ViewMode = UE::Chaos::ClothAsset::DataflowViewModeToClothViewMode(DataflowContextObject->GetConstructionViewMode());
+			const bool bIsRenderMode = ViewMode == EClothPatternVertexType::Render;
+
 			const int32 NumPatterns = bIsRenderMode ? Cloth.GetNumRenderPatterns() : Cloth.GetNumSimPatterns();
 
 			PatternTriangleOffsetAndNum.SetNum(NumPatterns);
@@ -453,75 +546,131 @@
 			{
 				ShowHideProperties->ShowPatterns.Add({ PatternID, false });
 			}
-
-			// Find the map if it exists.
-			if (TSharedPtr<Dataflow::FEngineContext> DataflowContext = DataflowContextObject->GetDataflowContext())
-			{
-				const FName InputName = WeightMapNodeToUpdate->GetInputName(*DataflowContext);
-				if (bIsRenderMode)
-				{
-					InputWeightMap = Cloth.GetUserDefinedAttribute<float>(InputName, ClothCollectionGroup::RenderVertices);
-				}
-				else
-				{
-					InputWeightMap = Cloth.GetWeightMap(InputName);
-				}
-			}
 		}
 	}
-
-	PrecomputeFuture.Wait();
-	OctreeFuture.Wait();
+}
 
 
-	// Create an attribute layer to temporarily paint into
-	const int NumAttributeLayers = Mesh->Attributes()->NumWeightLayers();
-	Mesh->Attributes()->SetNumWeightLayers(NumAttributeLayers + 1);
-	ActiveWeightMap = Mesh->Attributes()->GetWeightLayer(NumAttributeLayers);
-	ActiveWeightMap->SetName(FName("PaintLayer"));
+void UClothEditorWeightMapPaintTool::NotifyTargetChanged()
+{
+	//
+	// The target mesh has changed due to a view mode change. We will attempt to transfer the current in-progress paint values to the new mesh.
+	// First, temporarily save the existing weights from the paint layer on the mesh.
+	// 
 
-	// Copy weights from selected node to the preview mesh
-	const int32 NumExpectedWeights = bHaveDynamicMeshToWeightConversion ? WeightToDynamicMesh.Num() : Mesh->MaxVertexID();
-	TArray<float> CurrentWeights;
-	CurrentWeights.SetNumZeroed(NumExpectedWeights);
-	WeightMapNodeToUpdate->CalculateFinalVertexWeightValues(InputWeightMap, TArrayView<float>(CurrentWeights));
-	
+	TArray<float> SavedWeights;
+	GetCurrentWeightMap(SavedWeights);
+
 	if (bHaveDynamicMeshToWeightConversion)
 	{
-		for (int32 WeightID = 0; WeightID < CurrentWeights.Num(); ++WeightID)
+		TArray<float> MappedWeights;
+		MappedWeights.Init(0.f, WeightToDynamicMesh.Num());
+		for (int32 DynamicMeshIdx = 0; DynamicMeshIdx < SavedWeights.Num(); ++DynamicMeshIdx)
 		{
+			MappedWeights[DynamicMeshToWeight[DynamicMeshIdx]] = SavedWeights[DynamicMeshIdx];
+		}
+
+		SavedWeights = MoveTemp(MappedWeights);
+	}
+
+
+	//
+	// Now re-initialize everything that depends on the mesh
+	//
+
+	InitializeSculptMeshFromTarget();
+
+	UpdateShowHideProperties();
+
+	//
+	// Copy saved values back to the new preview mesh
+	//
+
+	checkf(ActiveWeightMap, TEXT("UClothEditorWeightMapPaintTool: no ActiveWeightMap after re-initializing the preview mesh"));
+
+	if (bHaveDynamicMeshToWeightConversion)
+	{
+		for (int32 WeightID = 0; WeightID < SavedWeights.Num(); ++WeightID)
+		{
 			for (const int32 VertexID : WeightToDynamicMesh[WeightID])
 			{
-				ActiveWeightMap->SetValue(VertexID, &CurrentWeights[WeightID]);
+				ActiveWeightMap->SetValue(VertexID, &SavedWeights[WeightID]);
 			}
 		}
 	}
 	else
 	{
-		for (int32 VertexID = 0; VertexID < CurrentWeights.Num(); ++VertexID)
+		for (int32 VertexID = 0; VertexID < SavedWeights.Num(); ++VertexID)
 		{
-			ActiveWeightMap->SetValue(VertexID, &CurrentWeights[VertexID]);
+			ActiveWeightMap->SetValue(VertexID, &SavedWeights[VertexID]);
 		}
 	}
 
-	UpdateWeightMapProperties->Name = WeightMapNodeToUpdate->OutputName.StringValue;
-	UpdateWeightMapProperties->MapOverrideType = WeightMapNodeToUpdate->MapOverrideType;
-	SetToolPropertySourceEnabled(UpdateWeightMapProperties, true);
+	//
+	// Update visualization
+	//
 
-	DynamicMeshComponent->EnableSecondaryTriangleBuffers(
-		[this](const FDynamicMesh3* Mesh, int32 TriangleID)
-		{
-			return PendingHiddenTriangles.Contains(TriangleID) || HiddenTriangles.Contains(TriangleID);
-		});
-	DynamicMeshComponent->SetSecondaryBuffersVisibility(false);
+	checkf(DynamicMeshComponent, TEXT("UClothEditorWeightMapPaintTool: no preview mesh after the tool target changed"));
 
-
-	// update colors
 	UpdateVertexColorOverlay();
 	DynamicMeshComponent->FastNotifyVertexAttributesUpdated(EMeshRenderAttributeFlags::VertexColors);
+
+	MeshElementsDisplay->NotifyMeshChanged();
+
 	GetToolManager()->PostInvalidation();
 
-	SetPrimaryFalloffType(EMeshSculptFalloffType::Smooth);
+	PostSetupCheck();
 }
 
+
+void UClothEditorWeightMapPaintTool::PostSetupCheck() const
+{
+	check(WeightMapNodeToUpdate);
+
+	check(PreviewMeshActor);
+	check(DynamicMeshComponent);
+	check(DynamicMeshComponent->GetAttachParent() == PreviewMeshActor->GetRootComponent());
+	check(DynamicMeshComponent->GetMesh() == GetSculptMesh());
+
+	check(ActiveWeightMap);
+	check(ActiveWeightMap->GetParent() == GetSculptMesh());
+
+	check(Octree.Mesh == GetSculptMesh());
+
+	check(MeshSpatial);
+	check(MeshSpatial->GetMesh() == GetSculptMesh());
+
+	check(TriNormals.Num() == GetSculptMesh()->MaxTriangleID());
+	check(UVSeamEdges.Num() == GetSculptMesh()->MaxEdgeID());
+	check(NormalSeamEdges.Num() == GetSculptMesh()->MaxEdgeID());
+
+	check(PolygonSelectionMechanic);
+	check(PolyLassoMechanic);
+	check(GradientSelectionTopology);
+	check(GradientSelectionTopology->GetMesh() == GetSculptMesh());
+
+	check(BrushProperties);
+	check(ToolPropertyObjects.Contains(BrushProperties));
+	check(GizmoProperties);
+	check(ToolPropertyObjects.Contains(GizmoProperties));
+	check(ViewProperties);
+	check(ToolPropertyObjects.Contains(ViewProperties));
+	check(UpdateWeightMapProperties);
+	check(ToolPropertyObjects.Contains(UpdateWeightMapProperties));
+	check(FilterProperties);
+	check(ToolPropertyObjects.Contains(FilterProperties));
+	check(PaintBrushOpProperties);
+	check(ToolPropertyObjects.Contains(PaintBrushOpProperties));
+	check(SmoothBrushOpProperties);
+	check(ToolPropertyObjects.Contains(SmoothBrushOpProperties));
+	check(EraseBrushOpProperties);
+	check(ToolPropertyObjects.Contains(EraseBrushOpProperties));
+	check(ActionsProps);
+	check(ToolPropertyObjects.Contains(ActionsProps));
+	check(ShowHideProperties);
+	check(ToolPropertyObjects.Contains(ShowHideProperties));
+
+	check(MeshElementsDisplay);
+}
+
 void UClothEditorWeightMapPaintTool::InitializeBrushSizeRange(const UE::Geometry::FAxisAlignedBox3d& TargetBounds)
@@ -2152,6 +2301,11 @@
 	{
 		this->WaitForPendingUndoRedo();
 	};
+	NewChange->AfterModify = [this](bool bRevert)
+	{
+		this->UpdateVertexColorOverlay();
+		this->DynamicMeshComponent->FastNotifyVertexAttributesUpdated(EMeshRenderAttributeFlags::VertexColors);
+	};
 
 	GetToolManager()->EmitObjectChange(this, MoveTemp(NewChange), LOCTEXT("VertexWeightChange", "Weight Stroke"));
 
diff --git a/Engine/Plugins/ChaosClothAssetEditor/Source/ChaosClothAssetEditorTools/Private/ChaosClothAsset/ClothWeightMapPaintTool.h b/Engine/Plugins/ChaosClothAssetEditor/Source/ChaosClothAssetEditorTools/Private/ChaosClothAsset/ClothWeightMapPaintTool.h
--- a/Engine/Plugins/ChaosClothAssetEditor/Source/ChaosClothAssetEditorTools/Private/ChaosClothAsset/ClothWeightMapPaintTool.h
+++ b/Engine/Plugins/ChaosClothAssetEditor/Source/ChaosClothAssetEditorTools/Private/ChaosClothAsset/ClothWeightMapPaintTool.h
@@ -331,6 +331,15 @@ public:
 
 	void SetDataflowContextObject(TObjectPtr<UDataflowContextObject> InDataflowContextObject);
 
+private:
+
+	// Initialization support
+	void InitializeSculptMeshFromTarget();
+	void UpdateShowHideProperties();
+
+	// Make sure things are set up correctly after initialization or re-initialization
+	void PostSetupCheck() const;
+
 public:
 
 
@@ -581,6 +590,7 @@ protected:
 	virtual bool ShowWorkPlane() const override { return false; }
 
 	friend class UClothEditorWeightMapPaintToolBuilder;
+	void NotifyTargetChanged();
 
 	bool bAnyChangeMade = false;
 
diff --git a/Engine/Plugins/ChaosClothAssetEditor/Source/ChaosClothAssetEditorTools/Public/ChaosClothAsset/ClothEditorToolBuilders.h b/Engine/Plugins/ChaosClothAssetEditor/Source/ChaosClothAssetEditorTools/Public/ChaosClothAsset/ClothEditorToolBuilders.h
--- a/Engine/Plugins/ChaosClothAssetEditor/Source/ChaosClothAssetEditorTools/Public/ChaosClothAsset/ClothEditorToolBuilders.h
+++ b/Engine/Plugins/ChaosClothAssetEditor/Source/ChaosClothAssetEditorTools/Public/ChaosClothAsset/ClothEditorToolBuilders.h
@@ -6,6 +6,8 @@
 #include "UObject/Interface.h"
 #include "BaseTools/SingleSelectionMeshEditingTool.h"
 #include "BaseTools/MeshSurfacePointMeshEditingTool.h"
+#include "DataflowEditorTools/DataflowEditorToolBuilder.h"
+
 #include "ClothEditorToolBuilders.generated.h"
 
 class UDataflowContextObject;
@@ -49,41 +51,67 @@ public:
 
 
 UCLASS()
-class CHAOSCLOTHASSETEDITORTOOLS_API UClothEditorWeightMapPaintToolBuilder : public UMeshSurfacePointMeshEditingToolBuilder, public IChaosClothAssetEditorToolBuilder
+class CHAOSCLOTHASSETEDITORTOOLS_API UClothEditorWeightMapPaintToolBuilder : public UMeshSurfacePointMeshEditingToolBuilder, public IChaosClothAssetEditorToolBuilder, public IDataflowEditorToolBuilder
 {
 	GENERATED_BODY()
 
 private:
+
+	// IDataflowEditorToolBuilder
+	virtual void GetSupportedViewModes(const UDataflowContextObject& ContextObject, TArray<const Dataflow::IDataflowConstructionViewMode*>& Modes) const override;
+	virtual bool CanSceneStateChange(const UInteractiveTool* ActiveTool, const FToolBuilderState& SceneState) const override;
+	virtual void SceneStateChanged(UInteractiveTool* ActiveTool, const FToolBuilderState& SceneState) override;
+
+	// IChaosClothAssetEditorToolBuilder
 	virtual void GetSupportedViewModes(const UDataflowContextObject& ContextObject, TArray<UE::Chaos::ClothAsset::EClothPatternVertexType>& Modes) const override;
-	virtual UMeshSurfacePointTool* CreateNewTool(const FToolBuilderState& SceneState) const override;
 	virtual bool CanSetConstructionViewWireframeActive() const { return false; }
+
+	// UMeshSurfacePointMeshEditingToolBuilder
+	virtual UMeshSurfacePointTool* CreateNewTool(const FToolBuilderState& SceneState) const override;
 };
 
 
 UCLASS()
-class CHAOSCLOTHASSETEDITORTOOLS_API UClothMeshSelectionToolBuilder : public UInteractiveToolWithToolTargetsBuilder, public IChaosClothAssetEditorToolBuilder
+class CHAOSCLOTHASSETEDITORTOOLS_API UClothMeshSelectionToolBuilder : public UInteractiveToolWithToolTargetsBuilder, public IChaosClothAssetEditorToolBuilder, public IDataflowEditorToolBuilder
 {
 	GENERATED_BODY()
 
 	virtual void GetSupportedViewModes(TArray<UE::Chaos::ClothAsset::EClothPatternVertexType>& Modes) const override {};
 
 private:
+
+	// IDataflowEditorToolBuilder
+	virtual void GetSupportedViewModes(const UDataflowContextObject& ContextObject, TArray<const Dataflow::IDataflowConstructionViewMode*>& Modes) const override;
+	virtual bool CanSceneStateChange(const UInteractiveTool* ActiveTool, const FToolBuilderState& SceneState) const override;
+	virtual void SceneStateChanged(UInteractiveTool* ActiveTool, const FToolBuilderState& SceneState) override;
+
+	// IChaosClothAssetEditorToolBuilder
 	virtual void GetSupportedViewModes(const UDataflowContextObject& ContextObject, TArray<UE::Chaos::ClothAsset::EClothPatternVertexType>& Modes) const override;
+	virtual bool CanSetConstructionViewWireframeActive() const { return false; }
+
+	// UInteractiveToolWithToolTargetsBuilder
 	virtual bool CanBuildTool(const FToolBuilderState& SceneState) const override;
 	virtual UInteractiveTool* BuildTool(const FToolBuilderState& SceneState) const override;
 	virtual const FToolTargetTypeRequirements& GetTargetRequirements() const override;
-	virtual bool CanSetConstructionViewWireframeActive() const { return false; }
 };
 
 UCLASS()
-class CHAOSCLOTHASSETEDITORTOOLS_API UClothTransferSkinWeightsToolBuilder : public USingleSelectionMeshEditingToolBuilder, public IChaosClothAssetEditorToolBuilder
+class CHAOSCLOTHASSETEDITORTOOLS_API UClothTransferSkinWeightsToolBuilder : public USingleSelectionMeshEditingToolBuilder, public IChaosClothAssetEditorToolBuilder, public IDataflowEditorToolBuilder
 {
 	GENERATED_BODY()
 
 private:
+
+	// IDataflowEditorToolBuilder
+	virtual void GetSupportedViewModes(const UDataflowContextObject& ContextObject, TArray<const Dataflow::IDataflowConstructionViewMode*>& Modes) const override;
+	virtual bool CanSceneStateChange(const UInteractiveTool* ActiveTool, const FToolBuilderState& SceneState) const override;
+	virtual void SceneStateChanged(UInteractiveTool* ActiveTool, const FToolBuilderState& SceneState) override;
+
+	// IChaosClothAssetEditorToolBuilder
 	virtual void GetSupportedViewModes(const UDataflowContextObject& ContextObject, TArray<UE::Chaos::ClothAsset::EClothPatternVertexType>& Modes) const override;
-	virtual USingleSelectionMeshEditingTool* CreateNewTool(const FToolBuilderState& SceneState) const override;
 
+	// USingleSelectionMeshEditingToolBuilder
+	virtual USingleSelectionMeshEditingTool* CreateNewTool(const FToolBuilderState& SceneState) const override;
 };
 
 
diff --git a/Engine/Plugins/Experimental/Dataflow/Source/DataflowEditor/Private/Dataflow/DataflowConstructionScene.cpp b/Engine/Plugins/Experimental/Dataflow/Source/DataflowEditor/Private/Dataflow/DataflowConstructionScene.cpp
--- a/Engine/Plugins/Experimental/Dataflow/Source/DataflowEditor/Private/Dataflow/DataflowConstructionScene.cpp
+++ b/Engine/Plugins/Experimental/Dataflow/Source/DataflowEditor/Private/Dataflow/DataflowConstructionScene.cpp
@@ -29,6 +29,13 @@ FDataflowConstructionScene::~FDataflowConstructionScene()
 	ResetDynamicMeshComponents();
 }
 
+TArray<TObjectPtr<UDynamicMeshComponent>> FDataflowConstructionScene::GetDynamicMeshComponents() const
+{
+	TArray<TObjectPtr<UDynamicMeshComponent>> OutValues;
+	DynamicMeshComponents.GenerateValueArray(OutValues);
+	return MoveTemp(OutValues);
+}
+
 /** Hide all or a single component */
 void FDataflowConstructionScene::SetVisibility(bool bVisibility, UActorComponent* InComponent)
 {
diff --git a/Engine/Plugins/Experimental/Dataflow/Source/DataflowEditor/Private/Dataflow/DataflowEditorMode.cpp b/Engine/Plugins/Experimental/Dataflow/Source/DataflowEditor/Private/Dataflow/DataflowEditorMode.cpp
--- a/Engine/Plugins/Experimental/Dataflow/Source/DataflowEditor/Private/Dataflow/DataflowEditorMode.cpp
+++ b/Engine/Plugins/Experimental/Dataflow/Source/DataflowEditor/Private/Dataflow/DataflowEditorMode.cpp
@@ -133,11 +133,52 @@ void UDataflowEditorMode::RegisterDataflowTool(TSharedPtr<FUICommandInfo> UIComm
 	const TSharedRef<FUICommandList>& CommandList = Toolkit->GetToolkitCommands();
 
 	CommandList->MapAction(UICommand, 
-		FExecuteAction::CreateWeakLambda(ToolsContext, [this, ToolsContext, ToolIdentifier]()
+		FExecuteAction::CreateWeakLambda(ToolsContext, [this, ToolsContext, ToolIdentifier, Builder]()
 		{
 			UDataflowContextObject* ContextObject = ToolsContext->ContextObjectStore->FindContext<UDataflowContextObject>();
 			check(ContextObject);
 
+			if (const IDataflowEditorToolBuilder* const DataflowToolBuilder = Cast<IDataflowEditorToolBuilder>(Builder))
+			{
+				// Check if we need to switch view modes before starting the tool
+				TArray<const Dataflow::IDataflowConstructionViewMode*> SupportedModes;
+
+				DataflowToolBuilder->GetSupportedViewModes(*ContextObject, SupportedModes);
+
+				if (SupportedModes.Num() > 0 && !SupportedModes.Contains(GetConstructionViewMode()))
+				{
+					if (!bShouldRestoreSavedConstructionViewMode)
+					{
+						// remember the current view mode so we can restore it later
+						SavedConstructionViewMode = GetConstructionViewMode()->GetName();
+						bShouldRestoreSavedConstructionViewMode = true;
+					}
+
+					bool bHadSingleSelectedMesh = false;
+					if (const USelection* const SelectedComponents = GetModeManager()->GetSelectedComponents())
+					{
+						bHadSingleSelectedMesh = (SelectedComponents->Num() == 1);
+					}
+
+					// switch to the preferred view mode for the tool that's about to start
+					SetConstructionViewMode(SupportedModes[0]->GetName());
+
+					if (bHadSingleSelectedMesh)
+					{
+						// If there is a single dynamic mesh component in the scene, select it so the tool can start
+						const TArray<UDynamicMeshComponent*> DynamicMeshComponents = ConstructionScene->GetDynamicMeshComponents();
+						if (DynamicMeshComponents.Num() == 1)
+						{
+							if (USelection* const SelectedComponents = GetModeManager()->GetSelectedComponents())
+							{
+								SelectedComponents->Select(DynamicMeshComponents[0]);
+							}
+						}
+					}
+
+				}
+			}
+
 			// Make sure the ContextObject's selected Collection is the from the Input side of the selected node (so that the tool gets the Collection as it appears before node execution)
 
 			if (TSharedPtr<Dataflow::FEngineContext> DataflowContext = ContextObject->GetDataflowContext())
@@ -538,6 +579,43 @@ void UDataflowEditorMode::ModeTick(float DeltaTime)
 
 		NodeTypeForPendingToolStart = FName();
 	}
+
+
+	if (bShouldRestartToolNextTick)
+	{
+		// If we ended the active tool in order to change view mode, restart it now
+
+		const UDataflowBaseContent* const EditorContent = ConstructionScene->GetEditorContent();
+		checkf(EditorContent, TEXT("Expected EditorContent in ConstructionScene"));
+
+		if (!EditorContent->IsConstructionDirty())		// hold off restarting the tool until the scene finishes rebuilding
+		{
+			// First select the lone mesh in the construction scene if there is one
+			if (bHadSingleSelectionBeforeToolShutdown)
+			{
+				const TArray<UDynamicMeshComponent*> DynamicMeshComponents = ConstructionScene->GetDynamicMeshComponents();
+				if (DynamicMeshComponents.Num() == 1)
+				{
+					if (USelection* const SelectedComponents = GetModeManager()->GetSelectedComponents())
+					{
+						SelectedComponents->Select(DynamicMeshComponents[0]);
+					}
+				}
+			}
+
+			// Now start the tool
+			if (const TSharedPtr<const SDataflowGraphEditor> PinnedGraphEditor = DataflowGraphEditor.Pin())
+			{
+				const FGraphPanelSelectionSet& SelectedNodes = PinnedGraphEditor->GetSelectedNodes();
+				if (SelectedNodes.Num() == 1)
+				{
+					StartToolForSelectedNode(*SelectedNodes.CreateConstIterator());
+				}
+			}
+
+			bShouldRestartToolNextTick = false;
+		}
+	}
 }
 
 void UDataflowEditorMode::ConstructionViewportResized(FViewport* ConstructionViewport, uint32 /*Unused*/)
@@ -583,25 +661,56 @@ FBox UDataflowEditorMode::SelectionBoundingBox() const
 
 void UDataflowEditorMode::SetConstructionViewMode(const FName& NewViewModeName)
 {
-	// We will first check if there is an active tool. If so, we'll shut down the tool and save the results to the Node, then change view modes, then restart the tool again.
-	
-	bool bEndedActiveTool = false;
-	UInteractiveToolManager* const ToolManager = GetInteractiveToolsContext()->ToolManager;
+	if (NewViewModeName == ConstructionViewMode->GetName())
+	{
+		return;
+	}
+
+	const TObjectPtr<UInteractiveToolManager> ToolManager = GetInteractiveToolsContext()->ToolManager;
 	checkf(ToolManager, TEXT("No valid ToolManager found for UDataflowEditorMode"));
+
+	checkf(ConstructionScene, TEXT("Expected UDataflowEditorMode::ConstructionScene to have been initialized"));
+
+	// Check if we have a single component selected. If we do, we will attempt to re-select it once the Construction Scene is rebuilt (if we have a tool running)
+	bHadSingleSelectionBeforeToolShutdown = false;
+
+	// Also check if we needed to shut down a running tool or not
+	bool bToolWasShutDown = false;
+
 	if (UInteractiveTool* const ActiveTool = ToolManager->GetActiveTool(EToolSide::Left))
 	{
-		// avoid switching back to the previous view mode when the tool ends here
-		const bool bTempShouldRestoreVal = bShouldRestoreSavedConstructionViewMode;
-		bShouldRestoreSavedConstructionViewMode = false;
+		if (const USelection* const SelectedComponents = GetModeManager()->GetSelectedComponents())
+		{
+			if (SelectedComponents->Num() == 1)			// TODO: Extend this to handle multiple selected components
+			{
+				if (const UDynamicMeshComponent* const SelectedDynamicMeshComponent = Cast<UDynamicMeshComponent>(SelectedComponents->GetSelectedObject(0)))
+				{
+					if (ConstructionScene->GetDynamicMeshComponents().Contains(SelectedDynamicMeshComponent))
+					{
+						bHadSingleSelectionBeforeToolShutdown = true;
+					}
+				}
+			}
+		}
 
-		ToolManager->PostActiveToolShutdownRequest(ActiveTool, EToolShutdownType::Accept);
-		bEndedActiveTool = true;
+		const UInteractiveToolBuilder* const ActiveToolBuilder = ToolManager->GetActiveToolBuilder(EToolSide::Left);
+		checkf(ActiveToolBuilder, TEXT("Found active tool with no active tool builder"));
 
-		// now we can restore the previous view mode the next time the tool ends
-		bShouldRestoreSavedConstructionViewMode = bTempShouldRestoreVal;
-	}
+		bool bToolCanHandleStateChange = false;
 
-	// Next get the ViewMode pointer from the given name
+		if (const IDataflowEditorToolBuilder* const DataflowToolBuilder = Cast<IDataflowEditorToolBuilder>(ActiveToolBuilder))
+		{
+			FToolBuilderState SceneState;
+			ToolManager->GetContextQueriesAPI()->GetCurrentSelectionState(SceneState);
+			bToolCanHandleStateChange = DataflowToolBuilder->CanSceneStateChange(ActiveTool, SceneState);
+		}
+
+		if (!bToolCanHandleStateChange)
+		{
+			ToolManager->PostActiveToolShutdownRequest(ActiveTool, EToolShutdownType::Accept);
+			bToolWasShutDown = true;
+		}
+	}
 
 	const Dataflow::FRenderingViewModeFactory& ViewModes = Dataflow::FRenderingViewModeFactory::GetInstance();
 	const Dataflow::IDataflowConstructionViewMode* const NewMode = ViewModes.GetViewMode(NewViewModeName);
@@ -626,17 +735,41 @@ void UDataflowEditorMode::SetConstructionViewMode(const FName& NewViewModeName)
 	// If we are switching to a mode with a valid mesh for the first time, focus the camera on it
 	FirstTimeFocusConstructionViewport();
 
-	if (bEndedActiveTool)
+
+	if (bToolWasShutDown)
+	{
+		// Tool restart must be done on the next tick because shutting down the current tool will cause the ConstructionView to be rebuilt next tick as well
+		bShouldRestartToolNextTick = true;
+	}
+	else if (UInteractiveTool* const ActiveTool = ToolManager->GetActiveTool(EToolSide::Left))
 	{
-		// If we ended the active tool in order to change modes, restart it now
-		if (const TSharedPtr<const SDataflowGraphEditor> PinnedGraphEditor = DataflowGraphEditor.Pin())
+		// If there is a currently active tool, notify it that the scene has changed
+
+		// First check if we previously had a single selected component before changing view modes. If so, and if there is now a single component in the construction scene, select it.
+		// TODO: Extend this to handle multiple selected components
+
+		if (bHadSingleSelectionBeforeToolShutdown)
 		{
-			const FGraphPanelSelectionSet& SelectedNodes = PinnedGraphEditor->GetSelectedNodes();
-			if (SelectedNodes.Num() == 1)
+			const TArray<UDynamicMeshComponent*> DynamicMeshComponents = ConstructionScene->GetDynamicMeshComponents();
+			if (DynamicMeshComponents.Num() == 1)
 			{
-				StartToolForSelectedNode(*SelectedNodes.CreateConstIterator());
+				if (USelection* const SelectedComponents = GetModeManager()->GetSelectedComponents())
+				{
+					SelectedComponents->Select(DynamicMeshComponents[0]);
+				}
 			}
 		}
+
+		// Now notify the active tool that the SceneState is different
+		UInteractiveToolBuilder* const ActiveToolBuilder = ToolManager->GetActiveToolBuilder(EToolSide::Left);
+		checkf(ActiveToolBuilder, TEXT("Found active tool with no active tool builder"));
+
+		if (IDataflowEditorToolBuilder* const DataflowToolBuilder = Cast<IDataflowEditorToolBuilder>(ActiveToolBuilder))
+		{
+			FToolBuilderState SceneState;
+			ToolManager->GetContextQueriesAPI()->GetCurrentSelectionState(SceneState);
+			DataflowToolBuilder->SceneStateChanged(ActiveTool, SceneState);
+		}
 	}
 }
 
@@ -672,18 +805,28 @@ bool UDataflowEditorMode::CanChangeConstructionViewModeTo(const FName& NewViewMo
 		return false;
 	}
 
-	// TODO: Check active tool to see if we can switch modes while the tool is running
-//	/*
-//	const UInteractiveToolBuilder* const ActiveToolBuilder = GetToolManager()->GetActiveToolBuilder(EToolSide::Left);
-//	checkf(ActiveToolBuilder, TEXT("No Active Tool Builder found despite having an Active Tool"));
-//
-//	const IDataflowEditorToolBuilder* const DataflowToolBuilder = Cast<const IDataflowEditorToolBuilder>(ActiveToolBuilder);
-//	checkf(ClothToolBuilder, TEXT("Cloth Editor has an active Tool Builder that does not implement IDataflowEditorToolBuilder"));
-//
-//	TArray<Dataflow::EDataflowPatternVertexType> SupportedViewModes;
-//	ClothToolBuilder->GetSupportedViewModes(SupportedViewModes);
-//	return SupportedViewModes.Contains(NewViewMode);
-//	*/
+
+	// Check active tool to see if we can switch modes while the tool is running
+
+	const UInteractiveToolBuilder* const ActiveToolBuilder = GetToolManager()->GetActiveToolBuilder(EToolSide::Left);
+	checkf(ActiveToolBuilder, TEXT("No Active Tool Builder found despite having an Active Tool"));
+
+	if (const IDataflowEditorToolBuilder* const DataflowToolBuilder = Cast<const IDataflowEditorToolBuilder>(ActiveToolBuilder))
+	{
+		const UEditorInteractiveToolsContext* const ConstructionToolsContext = GetInteractiveToolsContext();
+		checkf(ConstructionToolsContext, TEXT("No Tools Context found in Dataflow Editor"));
+
+		const UDataflowContextObject* const DataflowContextObject = ConstructionToolsContext->ContextObjectStore->FindContext<UDataflowContextObject>();
+		checkf(DataflowContextObject, TEXT("No Dataflow Context Object found in ContextObjectStore, despite having an Active Tool. This should have been created by the time a tool is activated"));
+
+		TArray<const Dataflow::IDataflowConstructionViewMode*> SupportedViewModes;
+		DataflowToolBuilder->GetSupportedViewModes(*DataflowContextObject, SupportedViewModes);
+
+		if (const Dataflow::IDataflowConstructionViewMode* const NewViewMode = Dataflow::FRenderingViewModeFactory::GetInstance().GetViewMode(NewViewModeName))
+		{
+			return SupportedViewModes.Contains(NewViewMode);
+		}
+	}
 
 	return false;
 }
diff --git a/Engine/Plugins/Experimental/Dataflow/Source/DataflowEditor/Private/DataflowEditorTools/DataflowEditorWeightMapPaintTool.cpp b/Engine/Plugins/Experimental/Dataflow/Source/DataflowEditor/Private/DataflowEditorTools/DataflowEditorWeightMapPaintTool.cpp
--- a/Engine/Plugins/Experimental/Dataflow/Source/DataflowEditor/Private/DataflowEditorTools/DataflowEditorWeightMapPaintTool.cpp
+++ b/Engine/Plugins/Experimental/Dataflow/Source/DataflowEditor/Private/DataflowEditorTools/DataflowEditorWeightMapPaintTool.cpp
@@ -68,7 +68,7 @@ namespace Dataflow::Private
 /*
  * ToolBuilder
  */
-void UDataflowEditorWeightMapPaintToolBuilder::GetSupportedViewModes(TArray<Dataflow::IDataflowConstructionViewMode*>& Modes) const
+void UDataflowEditorWeightMapPaintToolBuilder::GetSupportedViewModes(const UDataflowContextObject& ContextObject, TArray<const Dataflow::IDataflowConstructionViewMode*>& Modes) const
 {
 	//Modes.Add(Dataflow::EDataflowPatternVertexType::Sim3D);
 	//Modes.Add(Dataflow::EDataflowPatternVertexType::Sim2D);
diff --git a/Engine/Plugins/Experimental/Dataflow/Source/DataflowEditor/Public/Dataflow/DataflowConstructionScene.h b/Engine/Plugins/Experimental/Dataflow/Source/DataflowEditor/Public/Dataflow/DataflowConstructionScene.h
--- a/Engine/Plugins/Experimental/Dataflow/Source/DataflowEditor/Public/Dataflow/DataflowConstructionScene.h
+++ b/Engine/Plugins/Experimental/Dataflow/Source/DataflowEditor/Public/Dataflow/DataflowConstructionScene.h
@@ -54,6 +54,9 @@ public:
 	/** Hide all or a single component */
 	void SetVisibility(bool bVisible, UActorComponent* Component = nullptr);
 
+	/** Array of all DynamicMeshComponents in the scene */
+	TArray<TObjectPtr<UDynamicMeshComponent>> GetDynamicMeshComponents() const;
+
 private:
 
 	/** Reset the scene mesh visualizer */
diff --git a/Engine/Plugins/Experimental/Dataflow/Source/DataflowEditor/Public/Dataflow/DataflowEditorMode.h b/Engine/Plugins/Experimental/Dataflow/Source/DataflowEditor/Public/Dataflow/DataflowEditorMode.h
--- a/Engine/Plugins/Experimental/Dataflow/Source/DataflowEditor/Public/Dataflow/DataflowEditorMode.h
+++ b/Engine/Plugins/Experimental/Dataflow/Source/DataflowEditor/Public/Dataflow/DataflowEditorMode.h
@@ -186,7 +186,7 @@ private:
 	bool bShouldRestoreSavedConstructionViewMode = false;
 	
 	// The Construction view mode that was active before starting the current tool. When the tool ends, restore this view mode if bShouldRestoreSavedConstructionViewMode is true.
-	const FName SavedConstructionViewMode;
+	FName SavedConstructionViewMode;
 
 	bool bConstructionViewWireframe = false;
 	bool bShouldRestoreConstructionViewWireframe = false;
@@ -212,6 +212,13 @@ private:
 	// Dataflow node type whose corresponding tool should be started on the next Tick
 	FName NodeTypeForPendingToolStart;
 
+	// If we should start the tool for the selected node on the next ModeTick call. Used when switching ViewModes to allow the scene to be rebuilt before restarting the tool.
+	bool bShouldRestartToolNextTick = false;
+
+	// Whether we have a single mesh selected in the Construction Scene prior to shutting down a running tool. The Scene will likely be rebuilt when the tool ends, resulting in
+	// losing selection information, so we save this before shutting the tool down.
+	bool bHadSingleSelectionBeforeToolShutdown = false;
+
 	UDataflowEditor* DataflowEditor = nullptr;
 
 };
diff --git a/Engine/Plugins/Experimental/Dataflow/Source/DataflowEditor/Public/DataflowEditorTools/DataflowEditorToolBuilder.h b/Engine/Plugins/Experimental/Dataflow/Source/DataflowEditor/Public/DataflowEditorTools/DataflowEditorToolBuilder.h
--- a/Engine/Plugins/Experimental/Dataflow/Source/DataflowEditor/Public/DataflowEditorTools/DataflowEditorToolBuilder.h
+++ b/Engine/Plugins/Experimental/Dataflow/Source/DataflowEditor/Public/DataflowEditorTools/DataflowEditorToolBuilder.h
@@ -10,6 +10,9 @@
 namespace Dataflow {
 	class IDataflowConstructionViewMode;
 }
+struct FToolBuilderState;
+class UDataflowContextObject;
+class UInteractiveTool;
 
 UINTERFACE(MinimalAPI)
 class UDataflowEditorToolBuilder : public UInterface
@@ -25,9 +28,18 @@ class IDataflowEditorToolBuilder
 public:
 
 	/** Returns all Construction View modes that this tool can operate in. The first element should be the preferred mode to switch to if necessary. */
-	virtual void GetSupportedViewModes(TArray<Dataflow::IDataflowConstructionViewMode*>& Modes) const = 0;
+	virtual void GetSupportedViewModes(const UDataflowContextObject& ContextObject, TArray<const Dataflow::IDataflowConstructionViewMode*>& Modes) const = 0;
 
 	/** Returns whether or not view can be set to wireframe when this tool is active.. */
 	virtual bool CanSetConstructionViewWireframeActive() const { return true; }
+
+	/** Returns true if the tool can keep running when the SceneState changes */
+	virtual bool CanSceneStateChange(const UInteractiveTool* ActiveTool, const FToolBuilderState& SceneState) const { return false; }
+
+	/** Respond to SceneState changing */
+	virtual void SceneStateChanged(UInteractiveTool* ActiveTool, const FToolBuilderState& SceneState) 
+	{ 
+		checkf(CanSceneStateChange(ActiveTool, SceneState), TEXT("Current tool cannot handle changing scene state while running"));
+	};
 };
 
diff --git a/Engine/Plugins/Experimental/Dataflow/Source/DataflowEditor/Public/DataflowEditorTools/DataflowEditorWeightMapPaintTool.h b/Engine/Plugins/Experimental/Dataflow/Source/DataflowEditor/Public/DataflowEditorTools/DataflowEditorWeightMapPaintTool.h
--- a/Engine/Plugins/Experimental/Dataflow/Source/DataflowEditor/Public/DataflowEditorTools/DataflowEditorWeightMapPaintTool.h
+++ b/Engine/Plugins/Experimental/Dataflow/Source/DataflowEditor/Public/DataflowEditorTools/DataflowEditorWeightMapPaintTool.h
@@ -63,7 +63,7 @@ public:
 	void SetEditorMode(UDataflowEditorMode* InMode) { Mode = InMode; }
 
 private:
-	virtual void GetSupportedViewModes(TArray<Dataflow::IDataflowConstructionViewMode*>& Modes) const override;
+	virtual void GetSupportedViewModes(const UDataflowContextObject& ContextObject, TArray<const Dataflow::IDataflowConstructionViewMode*>& Modes) const override;
 	virtual bool CanBuildTool(const FToolBuilderState& SceneState) const override;
 	virtual UMeshSurfacePointTool* CreateNewTool(const FToolBuilderState& SceneState) const override;
 	virtual bool CanSetConstructionViewWireframeActive() const { return false; }
