Dataflow Editor: rename GetSupportedViewModes to GetSupportedConstructionViewModes. Fixes compiler error about hiding virtual function in classes that inherit from both Dataflow Tool Builder and Cloth Tool Builder.

diff --git a/Engine/Plugins/ChaosClothAssetEditor/Source/ChaosClothAssetEditorTools/Private/ChaosClothAsset/ClothEditorToolBuilders.cpp b/Engine/Plugins/ChaosClothAssetEditor/Source/ChaosClothAssetEditorTools/Private/ChaosClothAsset/ClothEditorToolBuilders.cpp
--- a/Engine/Plugins/ChaosClothAssetEditor/Source/ChaosClothAssetEditorTools/Private/ChaosClothAsset/ClothEditorToolBuilders.cpp
+++ b/Engine/Plugins/ChaosClothAssetEditor/Source/ChaosClothAssetEditorTools/Private/ChaosClothAsset/ClothEditorToolBuilders.cpp
@@ -20,7 +20,7 @@
 
 // ------------------- Weight Map Paint Tool -------------------
 
-void UClothEditorWeightMapPaintToolBuilder::GetSupportedViewModes(const UDataflowContextObject& ContextObject, TArray<const Dataflow::IDataflowConstructionViewMode*>& Modes) const
+void UClothEditorWeightMapPaintToolBuilder::GetSupportedConstructionViewModes(const UDataflowContextObject& ContextObject, TArray<const Dataflow::IDataflowConstructionViewMode*>& Modes) const
 {
 	using namespace UE::Chaos::ClothAsset;
 
@@ -141,7 +141,7 @@ UMeshSurfacePointTool* UClothEditorWeightMapPaintToolBuilder::CreateNewTool(cons
 // ------------------- Selection Tool -------------------
 
 
-void UClothMeshSelectionToolBuilder::GetSupportedViewModes(const UDataflowContextObject& ContextObject, TArray<const Dataflow::IDataflowConstructionViewMode*>& Modes) const
+void UClothMeshSelectionToolBuilder::GetSupportedConstructionViewModes(const UDataflowContextObject& ContextObject, TArray<const Dataflow::IDataflowConstructionViewMode*>& Modes) const
 {
 	using namespace UE::Chaos::ClothAsset;
 
@@ -241,7 +241,7 @@ UInteractiveTool* UClothMeshSelectionToolBuilder::BuildTool(const FToolBuilderSt
 // ------------------- Skin Weight Transfer Tool -------------------
 
 
-void UClothTransferSkinWeightsToolBuilder::GetSupportedViewModes(const UDataflowContextObject& ContextObject, TArray<const Dataflow::IDataflowConstructionViewMode*>& Modes) const
+void UClothTransferSkinWeightsToolBuilder::GetSupportedConstructionViewModes(const UDataflowContextObject& ContextObject, TArray<const Dataflow::IDataflowConstructionViewMode*>& Modes) const
 {
 	using namespace UE::Chaos::ClothAsset;
 	const Dataflow::FRenderingViewModeFactory& Factory = Dataflow::FRenderingViewModeFactory::GetInstance();
diff --git a/Engine/Plugins/ChaosClothAssetEditor/Source/ChaosClothAssetEditorTools/Public/ChaosClothAsset/ClothEditorToolBuilders.h b/Engine/Plugins/ChaosClothAssetEditor/Source/ChaosClothAssetEditorTools/Public/ChaosClothAsset/ClothEditorToolBuilders.h
--- a/Engine/Plugins/ChaosClothAssetEditor/Source/ChaosClothAssetEditorTools/Public/ChaosClothAsset/ClothEditorToolBuilders.h
+++ b/Engine/Plugins/ChaosClothAssetEditor/Source/ChaosClothAssetEditorTools/Public/ChaosClothAsset/ClothEditorToolBuilders.h
@@ -58,7 +58,7 @@ class CHAOSCLOTHASSETEDITORTOOLS_API UClothEditorWeightMapPaintToolBuilder : pub
 private:
 
 	// IDataflowEditorToolBuilder
-	virtual void GetSupportedViewModes(const UDataflowContextObject& ContextObject, TArray<const Dataflow::IDataflowConstructionViewMode*>& Modes) const override;
+	virtual void GetSupportedConstructionViewModes(const UDataflowContextObject& ContextObject, TArray<const Dataflow::IDataflowConstructionViewMode*>& Modes) const override;
 	virtual bool CanSceneStateChange(const UInteractiveTool* ActiveTool, const FToolBuilderState& SceneState) const override;
 	virtual void SceneStateChanged(UInteractiveTool* ActiveTool, const FToolBuilderState& SceneState) override;
 
@@ -81,7 +81,7 @@ class CHAOSCLOTHASSETEDITORTOOLS_API UClothMeshSelectionToolBuilder : public UIn
 private:
 
 	// IDataflowEditorToolBuilder
-	virtual void GetSupportedViewModes(const UDataflowContextObject& ContextObject, TArray<const Dataflow::IDataflowConstructionViewMode*>& Modes) const override;
+	virtual void GetSupportedConstructionViewModes(const UDataflowContextObject& ContextObject, TArray<const Dataflow::IDataflowConstructionViewMode*>& Modes) const override;
 	virtual bool CanSceneStateChange(const UInteractiveTool* ActiveTool, const FToolBuilderState& SceneState) const override;
 	virtual void SceneStateChanged(UInteractiveTool* ActiveTool, const FToolBuilderState& SceneState) override;
 
@@ -103,7 +103,7 @@ class CHAOSCLOTHASSETEDITORTOOLS_API UClothTransferSkinWeightsToolBuilder : publ
 private:
 
 	// IDataflowEditorToolBuilder
-	virtual void GetSupportedViewModes(const UDataflowContextObject& ContextObject, TArray<const Dataflow::IDataflowConstructionViewMode*>& Modes) const override;
+	virtual void GetSupportedConstructionViewModes(const UDataflowContextObject& ContextObject, TArray<const Dataflow::IDataflowConstructionViewMode*>& Modes) const override;
 	virtual bool CanSceneStateChange(const UInteractiveTool* ActiveTool, const FToolBuilderState& SceneState) const override;
 	virtual void SceneStateChanged(UInteractiveTool* ActiveTool, const FToolBuilderState& SceneState) override;
 
diff --git a/Engine/Plugins/Experimental/Dataflow/Source/DataflowEditor/Private/Dataflow/DataflowEditorMode.cpp b/Engine/Plugins/Experimental/Dataflow/Source/DataflowEditor/Private/Dataflow/DataflowEditorMode.cpp
--- a/Engine/Plugins/Experimental/Dataflow/Source/DataflowEditor/Private/Dataflow/DataflowEditorMode.cpp
+++ b/Engine/Plugins/Experimental/Dataflow/Source/DataflowEditor/Private/Dataflow/DataflowEditorMode.cpp
@@ -143,7 +143,7 @@ void UDataflowEditorMode::RegisterDataflowTool(TSharedPtr<FUICommandInfo> UIComm
 				// Check if we need to switch view modes before starting the tool
 				TArray<const Dataflow::IDataflowConstructionViewMode*> SupportedModes;
 
-				DataflowToolBuilder->GetSupportedViewModes(*ContextObject, SupportedModes);
+				DataflowToolBuilder->GetSupportedConstructionViewModes(*ContextObject, SupportedModes);
 
 				if (SupportedModes.Num() > 0 && !SupportedModes.Contains(GetConstructionViewMode()))
 				{
@@ -831,7 +831,7 @@ bool UDataflowEditorMode::CanChangeConstructionViewModeTo(const FName& NewViewMo
 		checkf(DataflowContextObject, TEXT("No Dataflow Context Object found in ContextObjectStore, despite having an Active Tool. This should have been created by the time a tool is activated"));
 
 		TArray<const Dataflow::IDataflowConstructionViewMode*> SupportedViewModes;
-		DataflowToolBuilder->GetSupportedViewModes(*DataflowContextObject, SupportedViewModes);
+		DataflowToolBuilder->GetSupportedConstructionViewModes(*DataflowContextObject, SupportedViewModes);
 
 		if (const Dataflow::IDataflowConstructionViewMode* const NewViewMode = Dataflow::FRenderingViewModeFactory::GetInstance().GetViewMode(NewViewModeName))
 		{
diff --git a/Engine/Plugins/Experimental/Dataflow/Source/DataflowEditor/Private/DataflowEditorTools/DataflowEditorWeightMapPaintTool.cpp b/Engine/Plugins/Experimental/Dataflow/Source/DataflowEditor/Private/DataflowEditorTools/DataflowEditorWeightMapPaintTool.cpp
--- a/Engine/Plugins/Experimental/Dataflow/Source/DataflowEditor/Private/DataflowEditorTools/DataflowEditorWeightMapPaintTool.cpp
+++ b/Engine/Plugins/Experimental/Dataflow/Source/DataflowEditor/Private/DataflowEditorTools/DataflowEditorWeightMapPaintTool.cpp
@@ -68,7 +68,7 @@ namespace Dataflow::Private
 /*
  * ToolBuilder
  */
-void UDataflowEditorWeightMapPaintToolBuilder::GetSupportedViewModes(const UDataflowContextObject& ContextObject, TArray<const Dataflow::IDataflowConstructionViewMode*>& Modes) const
+void UDataflowEditorWeightMapPaintToolBuilder::GetSupportedConstructionViewModes(const UDataflowContextObject& ContextObject, TArray<const Dataflow::IDataflowConstructionViewMode*>& Modes) const
 {
 	//Modes.Add(Dataflow::EDataflowPatternVertexType::Sim3D);
 	//Modes.Add(Dataflow::EDataflowPatternVertexType::Sim2D);
diff --git a/Engine/Plugins/Experimental/Dataflow/Source/DataflowEditor/Public/DataflowEditorTools/DataflowEditorToolBuilder.h b/Engine/Plugins/Experimental/Dataflow/Source/DataflowEditor/Public/DataflowEditorTools/DataflowEditorToolBuilder.h
--- a/Engine/Plugins/Experimental/Dataflow/Source/DataflowEditor/Public/DataflowEditorTools/DataflowEditorToolBuilder.h
+++ b/Engine/Plugins/Experimental/Dataflow/Source/DataflowEditor/Public/DataflowEditorTools/DataflowEditorToolBuilder.h
@@ -28,7 +28,7 @@ class IDataflowEditorToolBuilder
 public:
 
 	/** Returns all Construction View modes that this tool can operate in. The first element should be the preferred mode to switch to if necessary. */
-	virtual void GetSupportedViewModes(const UDataflowContextObject& ContextObject, TArray<const Dataflow::IDataflowConstructionViewMode*>& Modes) const = 0;
+	virtual void GetSupportedConstructionViewModes(const UDataflowContextObject& ContextObject, TArray<const Dataflow::IDataflowConstructionViewMode*>& Modes) const = 0;
 
 	/** Returns whether or not view can be set to wireframe when this tool is active.. */
 	virtual bool CanSetConstructionViewWireframeActive() const { return true; }
diff --git a/Engine/Plugins/Experimental/Dataflow/Source/DataflowEditor/Public/DataflowEditorTools/DataflowEditorWeightMapPaintTool.h b/Engine/Plugins/Experimental/Dataflow/Source/DataflowEditor/Public/DataflowEditorTools/DataflowEditorWeightMapPaintTool.h
--- a/Engine/Plugins/Experimental/Dataflow/Source/DataflowEditor/Public/DataflowEditorTools/DataflowEditorWeightMapPaintTool.h
+++ b/Engine/Plugins/Experimental/Dataflow/Source/DataflowEditor/Public/DataflowEditorTools/DataflowEditorWeightMapPaintTool.h
@@ -63,7 +63,7 @@ public:
 	void SetEditorMode(UDataflowEditorMode* InMode) { Mode = InMode; }
 
 private:
-	virtual void GetSupportedViewModes(const UDataflowContextObject& ContextObject, TArray<const Dataflow::IDataflowConstructionViewMode*>& Modes) const override;
+	virtual void GetSupportedConstructionViewModes(const UDataflowContextObject& ContextObject, TArray<const Dataflow::IDataflowConstructionViewMode*>& Modes) const override;
 	virtual bool CanBuildTool(const FToolBuilderState& SceneState) const override;
 	virtual UMeshSurfacePointTool* CreateNewTool(const FToolBuilderState& SceneState) const override;
 	virtual bool CanSetConstructionViewWireframeActive() const { return false; }
