Cloth Editor: remove interactive tools support from Preview Viewport. This shouldn't have any effect on behavior but should simplify the code a bit. This can be reversed if we end up re-adding tools to the preview viewport.

diff --git a/Engine/Plugins/ChaosClothAssetEditor/Source/ChaosClothAssetEditor/Private/ChaosClothAsset/ClothEditorMode.cpp b/Engine/Plugins/ChaosClothAssetEditor/Source/ChaosClothAssetEditor/Private/ChaosClothAsset/ClothEditorMode.cpp
--- a/Engine/Plugins/ChaosClothAssetEditor/Source/ChaosClothAssetEditor/Private/ChaosClothAsset/ClothEditorMode.cpp
+++ b/Engine/Plugins/ChaosClothAssetEditor/Source/ChaosClothAssetEditor/Private/ChaosClothAsset/ClothEditorMode.cpp
@@ -331,10 +331,6 @@ void UChaosClothAssetEditorMode::RegisterAddNodeCommand(TSharedPtr<FUICommandInf
 }
 
 
-void UChaosClothAssetEditorMode::RegisterPreviewTools()
-{
-}
-
 void UChaosClothAssetEditorMode::RegisterTools()
 {
 	using namespace UE::Chaos::ClothAsset;
@@ -364,14 +360,6 @@ bool UChaosClothAssetEditorMode::ShouldToolStartBeAllowed(const FString& ToolIde
 		return false;
 	}
 	
-	if (PreviewScene && PreviewScene->GetClothPreviewEditorModeManager() && PreviewScene->GetClothPreviewEditorModeManager()->GetInteractiveToolsContext())
-	{
-		if (PreviewScene->GetClothPreviewEditorModeManager()->GetInteractiveToolsContext()->HasActiveTool())
-		{
-			return false;
-		}
-	}
-
 	if (UE::Chaos::ClothAsset::Private::bClothEditorEnableToolsInPIE)
 	{
 		// UEdMode::ShouldToolStartBeAllowed returns (!GEditor->PlayWorld && !GIsPlayInEditorWorld) but we want to allow tools to start while in PIE
@@ -510,28 +498,7 @@ void UChaosClothAssetEditorMode::Exit()
 
 void UChaosClothAssetEditorMode::SetPreviewScene(UE::Chaos::ClothAsset::FChaosClothPreviewScene* InPreviewScene)
 {
-	using namespace UE::Chaos::ClothAsset;
-
 	PreviewScene = InPreviewScene;
-
-
-	UEditorInteractiveToolsContext* const PreviewToolsContext = PreviewScene->GetClothPreviewEditorModeManager()->GetInteractiveToolsContext();
-	UInteractiveToolManager* const PreviewToolManager = PreviewToolsContext->ToolManager;
-	PreviewToolsContext->TargetManager->AddTargetFactory(NewObject<USkeletalMeshComponentToolTargetFactory>(PreviewToolManager));
-
-	PreviewToolManager->OnToolStarted.AddUObject(this, &UChaosClothAssetEditorMode::OnToolStarted);
-	PreviewToolManager->OnToolEnded.AddUObject(this, &UChaosClothAssetEditorMode::OnToolEnded);
-
-	check(Toolkit.IsValid());
-	
-	// FBaseToolkit's OnToolStarted and OnToolEnded are protected, so we use the subclass to get at them
-	FChaosClothAssetEditorModeToolkit* const ClothModeToolkit = static_cast<FChaosClothAssetEditorModeToolkit*>(Toolkit.Get());	
-
-	PreviewToolManager->OnToolStarted.AddSP(ClothModeToolkit, &FChaosClothAssetEditorModeToolkit::OnToolStarted);
-	PreviewToolManager->OnToolEnded.AddSP(ClothModeToolkit, &FChaosClothAssetEditorModeToolkit::OnToolEnded);
-
-	// TODO: It would be nice if the PreviewScene could be specified before the UBaseCharacterFXEditorMode::Enter() function is called. Then we could register both sets of tools from there.
-	RegisterPreviewTools();
 }
 
 void UChaosClothAssetEditorMode::CreateToolTargets(const TArray<TObjectPtr<UObject>>& AssetsIn) 
diff --git a/Engine/Plugins/ChaosClothAssetEditor/Source/ChaosClothAssetEditor/Private/ChaosClothAsset/ClothEditorMode.h b/Engine/Plugins/ChaosClothAssetEditor/Source/ChaosClothAssetEditor/Private/ChaosClothAsset/ClothEditorMode.h
--- a/Engine/Plugins/ChaosClothAssetEditor/Source/ChaosClothAssetEditor/Private/ChaosClothAsset/ClothEditorMode.h
+++ b/Engine/Plugins/ChaosClothAssetEditor/Source/ChaosClothAssetEditor/Private/ChaosClothAsset/ClothEditorMode.h
@@ -170,9 +170,6 @@ private:
 
 	void RegisterAddNodeCommand(TSharedPtr<FUICommandInfo> AddNodeCommand, const FName& NewNodeType, TSharedPtr<FUICommandInfo> StartToolCommand);
 
-	// Register the set of tools that operate on objects in the 3D preview world 
-	void RegisterPreviewTools();
-
 	void SetPreviewScene(UE::Chaos::ClothAsset::FChaosClothPreviewScene* PreviewScene);
 
 	// Bounding box for rest space meshes
diff --git a/Engine/Plugins/ChaosClothAssetEditor/Source/ChaosClothAssetEditor/Private/ChaosClothAsset/ClothEditorModeToolkit.cpp b/Engine/Plugins/ChaosClothAssetEditor/Source/ChaosClothAssetEditor/Private/ChaosClothAsset/ClothEditorModeToolkit.cpp
--- a/Engine/Plugins/ChaosClothAssetEditor/Source/ChaosClothAssetEditor/Private/ChaosClothAsset/ClothEditorModeToolkit.cpp
+++ b/Engine/Plugins/ChaosClothAssetEditor/Source/ChaosClothAssetEditor/Private/ChaosClothAsset/ClothEditorModeToolkit.cpp
@@ -166,38 +166,6 @@ const FSlateBrush* FChaosClothAssetEditorModeToolkit::GetActiveToolIcon(const FS
 	return FChaosClothAssetEditorStyle::Get().GetOptionalBrush(ActiveToolIconName);
 }
 
-SBaseCharacterFXEditorViewport* FChaosClothAssetEditorModeToolkit::GetViewportWidgetForManager(UInteractiveToolManager* Manager)
-{
-	if (OwningEditorMode.IsValid(false))
-	{
-		const UEdMode* const Mode = OwningEditorMode.Get();
-
-		if (const UChaosClothAssetEditorMode* const ClothEdMode = Cast<UChaosClothAssetEditorMode>(Mode))
-		{
-			const FChaosClothPreviewScene* const PreviewScene = ClothEdMode->PreviewScene;
-			const UEditorInteractiveToolsContext* const PreviewToolsContext = PreviewScene->GetClothPreviewEditorModeManager()->GetInteractiveToolsContext();
-			const UInteractiveToolManager* const PreviewToolManager = PreviewToolsContext->ToolManager;
-			 
-			if (Manager == PreviewToolManager)
-			{
-				if (const TSharedPtr<SChaosClothAssetEditor3DViewport> Widget = PreviewViewportWidget.Pin())
-				{
-					return Widget.Get();
-				}
-			}
-			else
-			{
-				if (const TSharedPtr<SChaosClothAssetEditorRestSpaceViewport> Widget = RestSpaceViewportWidget.Pin())
-				{
-					return Widget.Get();
-				}
-			}
-		}
-	}
-
-	return nullptr;
-}
-
 UEditorInteractiveToolsContext* FChaosClothAssetEditorModeToolkit::GetCurrentToolsContext()
 {
 	if (OwningEditorMode.IsValid(/*bEvenIfPendingKill*/ false))
@@ -243,7 +211,7 @@ void FChaosClothAssetEditorModeToolkit::OnToolStarted(UInteractiveToolManager* M
 	ActiveToolIdentifier.InsertAt(0, ".");
 	ActiveToolIcon = GetActiveToolIcon(ActiveToolIdentifier);
 
-	if (SBaseCharacterFXEditorViewport* Widget = GetViewportWidgetForManager(Manager))
+	if (TSharedPtr<SBaseCharacterFXEditorViewport> Widget = RestSpaceViewportWidget.Pin())
 	{
 		Widget->AddOverlayWidget(ViewportOverlayWidget.ToSharedRef());
 
@@ -261,7 +229,7 @@ void FChaosClothAssetEditorModeToolkit::OnToolEnded(UInteractiveToolManager* Man
 	ClearNotification();
 	ClearWarning();
 	
-	if (SBaseCharacterFXEditorViewport* const Widget = GetViewportWidgetForManager(Manager))
+	if (TSharedPtr<SBaseCharacterFXEditorViewport> Widget = RestSpaceViewportWidget.Pin())
 	{
 		Widget->RemoveOverlayWidget(ViewportOverlayWidget.ToSharedRef());
 	}
@@ -278,11 +246,6 @@ void FChaosClothAssetEditorModeToolkit::SetRestSpaceViewportWidget(TWeakPtr<SCha
 	RestSpaceViewportWidget = InRestSpaceViewportWidget;
 }
 
-void FChaosClothAssetEditorModeToolkit::SetPreviewViewportWidget(TWeakPtr<SChaosClothAssetEditor3DViewport> InPreviewViewportWidget)
-{
-	PreviewViewportWidget = InPreviewViewportWidget;
-}
-
 } // namespace UE::Chaos::ClothAsset
 
 #undef LOCTEXT_NAMESPACE
diff --git a/Engine/Plugins/ChaosClothAssetEditor/Source/ChaosClothAssetEditor/Private/ChaosClothAsset/ClothEditorModeToolkit.h b/Engine/Plugins/ChaosClothAssetEditor/Source/ChaosClothAssetEditor/Private/ChaosClothAsset/ClothEditorModeToolkit.h
--- a/Engine/Plugins/ChaosClothAssetEditor/Source/ChaosClothAssetEditor/Private/ChaosClothAsset/ClothEditorModeToolkit.h
+++ b/Engine/Plugins/ChaosClothAssetEditor/Source/ChaosClothAssetEditor/Private/ChaosClothAsset/ClothEditorModeToolkit.h
@@ -37,18 +37,11 @@ public:
 	virtual FText GetBaseToolkitName() const override;
 
 	void SetRestSpaceViewportWidget(TWeakPtr<SChaosClothAssetEditorRestSpaceViewport>);
-	void SetPreviewViewportWidget(TWeakPtr<SChaosClothAssetEditor3DViewport>);
 
 private:
 
-	// Get the viewport widget associated with the given manager
-	// TODO: This should not be necessary any more as we do not run tools in the Cloth Preview Viewport (JIRA UE-201248)
-	SBaseCharacterFXEditorViewport* GetViewportWidgetForManager(UInteractiveToolManager* Manager);
-
 	UEditorInteractiveToolsContext* GetCurrentToolsContext();
 
 	TWeakPtr<SChaosClothAssetEditorRestSpaceViewport> RestSpaceViewportWidget;
-	TWeakPtr<SChaosClothAssetEditor3DViewport> PreviewViewportWidget;
-
 };
 } // namespace UE::Chaos::ClothAsset
diff --git a/Engine/Plugins/ChaosClothAssetEditor/Source/ChaosClothAssetEditor/Private/ChaosClothAsset/ClothEditorPreviewScene.cpp b/Engine/Plugins/ChaosClothAssetEditor/Source/ChaosClothAssetEditor/Private/ChaosClothAsset/ClothEditorPreviewScene.cpp
--- a/Engine/Plugins/ChaosClothAssetEditor/Source/ChaosClothAssetEditor/Private/ChaosClothAsset/ClothEditorPreviewScene.cpp
+++ b/Engine/Plugins/ChaosClothAssetEditor/Source/ChaosClothAssetEditor/Private/ChaosClothAsset/ClothEditorPreviewScene.cpp
@@ -316,17 +316,11 @@ const USkeletalMeshComponent* FChaosClothPreviewScene::GetSkeletalMeshComponent(
 	return SkeletalMeshComponent;
 }
 
-void FChaosClothPreviewScene::SetModeManager(TSharedPtr<FAssetEditorModeManager> InClothPreviewEditorModeManager)
+void FChaosClothPreviewScene::SetModeManager(const TSharedPtr<FAssetEditorModeManager>& InClothPreviewEditorModeManager)
 {
 	ClothPreviewEditorModeManager = InClothPreviewEditorModeManager;
 }
 
-const TSharedPtr<const FAssetEditorModeManager> FChaosClothPreviewScene::GetClothPreviewEditorModeManager() const
-{
-	return ClothPreviewEditorModeManager;
-}
-
-
 bool FChaosClothPreviewScene::IsComponentSelected(const UPrimitiveComponent* InComponent)
 {
 	if (const UTypedElementSelectionSet* const TypedElementSelectionSet = ClothPreviewEditorModeManager->GetEditorSelectionSet())
diff --git a/Engine/Plugins/ChaosClothAssetEditor/Source/ChaosClothAssetEditor/Private/ChaosClothAsset/ClothEditorPreviewScene.h b/Engine/Plugins/ChaosClothAssetEditor/Source/ChaosClothAssetEditor/Private/ChaosClothAsset/ClothEditorPreviewScene.h
--- a/Engine/Plugins/ChaosClothAssetEditor/Source/ChaosClothAssetEditor/Private/ChaosClothAsset/ClothEditorPreviewScene.h
+++ b/Engine/Plugins/ChaosClothAssetEditor/Source/ChaosClothAssetEditor/Private/ChaosClothAsset/ClothEditorPreviewScene.h
@@ -108,8 +108,8 @@ public:
 	
 	const USkeletalMeshComponent* GetSkeletalMeshComponent() const;
 
-	void SetModeManager(TSharedPtr<FAssetEditorModeManager> InClothPreviewEditorModeManager);
-	const TSharedPtr<const FAssetEditorModeManager> GetClothPreviewEditorModeManager() const;
+	/** Set the scene's ModeManager, which is mainly used to track selected components */
+	void SetModeManager(const TSharedPtr<FAssetEditorModeManager>& InClothPreviewEditorModeManager);
 
 	void SetGizmoDataBinder(TSharedPtr<FTransformGizmoDataBinder> InDataBinder);
 
diff --git a/Engine/Plugins/ChaosClothAssetEditor/Source/ChaosClothAssetEditor/Private/ChaosClothAsset/ClothEditorToolkit.cpp b/Engine/Plugins/ChaosClothAssetEditor/Source/ChaosClothAssetEditor/Private/ChaosClothAsset/ClothEditorToolkit.cpp
--- a/Engine/Plugins/ChaosClothAssetEditor/Source/ChaosClothAssetEditor/Private/ChaosClothAsset/ClothEditorToolkit.cpp
+++ b/Engine/Plugins/ChaosClothAssetEditor/Source/ChaosClothAssetEditor/Private/ChaosClothAsset/ClothEditorToolkit.cpp
@@ -158,13 +158,14 @@ FChaosClothAssetEditorToolkit::FChaosClothAssetEditorToolkit(UAssetEditor* InOwn
 	
 	ClothPreviewScene = MakeShared<FChaosClothPreviewScene>(PreviewSceneArgs);
 	ClothPreviewScene->SetFloorVisibility(false, true);
+
+	// The PreviewScene's ModeManager is used to keep track of selected components
 	ClothPreviewEditorModeManager = MakeShared<FAssetEditorModeManager>();
 	ClothPreviewEditorModeManager->SetPreviewScene(ClothPreviewScene.Get());
 	ClothPreviewScene->SetModeManager(ClothPreviewEditorModeManager);
 
 	ClothEditorSimulationVisualization = MakeShared<FClothEditorSimulationVisualization>();
 
-	//ClothPreviewInputRouter = ClothPreviewEditorModeManager->GetInteractiveToolsContext()->InputRouter;
 	ClothPreviewTabContent = MakeShareable(new FEditorViewportTabContent());
 	ClothPreviewViewportClient = MakeShared<FChaosClothAssetEditor3DViewportClient>(ClothPreviewEditorModeManager.Get(), ClothPreviewScene, ClothEditorSimulationVisualization);
 	ClothPreviewViewportClient->RegisterDelegates();
@@ -268,7 +269,6 @@ void FChaosClothAssetEditorToolkit::InitializeEdMode(UBaseCharacterFXEditorMode*
 	{
 		FChaosClothAssetEditorModeToolkit* ClothModeToolkit = static_cast<FChaosClothAssetEditorModeToolkit*>(ModeToolkit.Get());
 		ClothModeToolkit->SetRestSpaceViewportWidget(RestSpaceViewportWidget);
-		ClothModeToolkit->SetPreviewViewportWidget(PreviewViewportWidget);
 
 		FName ParentToolbarName;
 		const FName ToolBarName = GetToolMenuToolbarName(ParentToolbarName);
