Cloth Editor: Pass nullptr as DataflowEditor argument when constructing SDataflowGraphEditor. Fixes crash on ClothEditor startup.

diff --git a/Engine/Plugins/ChaosClothAssetEditor/Source/ChaosClothAssetEditor/Private/ChaosClothAsset/ClothEditorToolkit.cpp b/Engine/Plugins/ChaosClothAssetEditor/Source/ChaosClothAssetEditor/Private/ChaosClothAsset/ClothEditorToolkit.cpp
--- a/Engine/Plugins/ChaosClothAssetEditor/Source/ChaosClothAssetEditor/Private/ChaosClothAsset/ClothEditorToolkit.cpp
+++ b/Engine/Plugins/ChaosClothAssetEditor/Source/ChaosClothAssetEditor/Private/ChaosClothAsset/ClothEditorToolkit.cpp
@@ -909,7 +909,8 @@ TSharedRef<SDataflowGraphEditor> FChaosClothAssetEditorToolkit::CreateGraphEdito
 		.GraphToEdit(Dataflow)
 		.GraphEvents(InEvents)
 		.DetailsView(NodeDetailsEditor)
-		.EvaluateGraph(EvalLambda);
+		.EvaluateGraph(EvalLambda)
+		.DataflowEditor(nullptr);
 
 	NewGraphEditor->OnSelectionChangedMulticast.AddSP(this, &FChaosClothAssetEditorToolkit::OnNodeSelectionChanged);
 	NewGraphEditor->OnNodeDeletedMulticast.AddSP(this, &FChaosClothAssetEditorToolkit::OnNodeDeleted);
