Cloth Transfer Skin Weights Tool: add scale to transform gizmo

diff --git a/Engine/Plugins/ChaosClothAssetEditor/Source/ChaosClothAssetEditorTools/Private/ChaosClothAsset/ClothTransferSkinWeightsTool.cpp b/Engine/Plugins/ChaosClothAssetEditor/Source/ChaosClothAssetEditorTools/Private/ChaosClothAsset/ClothTransferSkinWeightsTool.cpp
--- a/Engine/Plugins/ChaosClothAssetEditor/Source/ChaosClothAssetEditorTools/Private/ChaosClothAsset/ClothTransferSkinWeightsTool.cpp
+++ b/Engine/Plugins/ChaosClothAssetEditor/Source/ChaosClothAssetEditorTools/Private/ChaosClothAsset/ClothTransferSkinWeightsTool.cpp
@@ -162,7 +162,7 @@ void UClothTransferSkinWeightsTool::Setup()
 		}
 	});
 
-	SourceMeshTransformGizmo = UE::TransformGizmoUtil::CreateCustomTransformGizmo(GizmoManager, ETransformGizmoSubElements::StandardTranslateRotate, this);
+	SourceMeshTransformGizmo = UE::TransformGizmoUtil::CreateCustomTransformGizmo(GizmoManager, ETransformGizmoSubElements::FullTranslateRotateScale, this);
 	ensure(SourceMeshTransformGizmo);
 
 	SourceMeshTransformGizmo->SetActiveTarget(SourceMeshTransformProxy, GetToolManager());
