Chaos Cloth Asset: fix Aerodynamics debug draw.

Softs Evolution doesn't calculate Aerodynamics forces in two passes, so the Forces buffer is not calculated. Instead just calculate the force based on the current state. This will represent the force for next frame instead of current frame, but most of our debug display shows based on current (end of frame) state rather than last frame, so this is consistent with that.

Also add a user-scale menu option and set the Simulation Statistics debug display to true, and fix a DPI scale bug in the text visualization.

diff --git a/Engine/Plugins/ChaosCloth/Source/ChaosCloth/Private/ChaosCloth/ChaosClothVisualization.cpp b/Engine/Plugins/ChaosCloth/Source/ChaosCloth/Private/ChaosCloth/ChaosClothVisualization.cpp
--- a/Engine/Plugins/ChaosCloth/Source/ChaosCloth/Private/ChaosCloth/ChaosClothVisualization.cpp
+++ b/Engine/Plugins/ChaosCloth/Source/ChaosCloth/Private/ChaosCloth/ChaosClothVisualization.cpp
@@ -214,7 +214,9 @@ namespace Chaos
 			FVector2D PixelLocation;
 			if (SceneView->WorldToPixel(Pos, PixelLocation))
 			{
-				FCanvasTextItem TextItem(PixelLocation, Text, GEngine->GetSmallFont(), Color);
+				// WorldToPixel doesn't account for DPIScale
+				const float DPIScale = Canvas->GetDPIScale();
+				FCanvasTextItem TextItem(PixelLocation / DPIScale, Text, GEngine->GetSmallFont(), Color);
 				TextItem.Scale = FVector2D::UnitVector  * Scale;
 				TextItem.EnableShadow(FLinearColor::Black);
 				TextItem.Draw(Canvas);
@@ -2391,15 +2393,13 @@ namespace Chaos
 		}
 	}
 
-	void FClothVisualizationNoGC::DrawWindAndPressureForces(FPrimitiveDrawInterface* PDI) const
+	void FClothVisualizationNoGC::DrawWindAndPressureForces(FPrimitiveDrawInterface* PDI, const FReal ForceLength) const
 	{
 		if (!Solver)
 		{
 			return;
 		}
 
-		constexpr FReal ForceLength = (FReal)10.;
-
 		for (const FClothingSimulationCloth* const Cloth : Solver->GetCloths())
 		{
 			const int32 ParticleRangeId = Cloth->GetParticleRangeId(Solver);
@@ -2426,10 +2426,11 @@ namespace Chaos
 			// Constraints are locally indexed for new solver
 			const TConstArrayView<Softs::FSolverVec3> Positions = Solver->IsLegacySolver() ? TConstArrayView<Softs::FSolverVec3>(Solver->GetParticleXs()) : Solver->GetParticleXsView(ParticleRangeId);
 			const TConstArrayView<Softs::FSolverReal> InvMasses = Solver->IsLegacySolver() ? TConstArrayView<Softs::FSolverReal>(Solver->GetParticleInvMasses()) : Solver->GetParticleInvMassesView(ParticleRangeId);
+			const TConstArrayView<Softs::FSolverVec3> Velocities = Solver->IsLegacySolver() ? TConstArrayView<Softs::FSolverVec3>(Solver->GetParticleVs()) : Solver->GetParticleVsView(ParticleRangeId);
 
 			const TConstArrayView<TVec3<int32>>& Elements = VelocityField->GetElements();
-			const TConstArrayView<Softs::FSolverVec3> Forces = VelocityField->GetForces();
 			check(InvMasses.Num() == Positions.Num());
+			check(InvMasses.Num() == Velocities.Num());
 
 			for (int32 ElementIndex = 0; ElementIndex < Elements.Num(); ++ElementIndex)
 			{
@@ -2444,7 +2445,7 @@ namespace Chaos
 				const bool bIsKinematic2 = !InvMasses[Element.Z];
 				const bool bIsKinematic = bIsKinematic0 || bIsKinematic1 || bIsKinematic2;
 
-				const FVec3 Force = GetWorldVector(*Solver, FVec3(Forces[ElementIndex]) * ForceLength);
+				const FVec3 Force = GetWorldVector(*Solver, FVec3(VelocityField->CalculateForce(Positions, Velocities, ElementIndex)) * ForceLength);
 				DrawLine(PDI, Position, Position + Force, bIsKinematic ? FColor::Cyan : FColor::Green);
 			}
 		}
diff --git a/Engine/Plugins/ChaosCloth/Source/ChaosCloth/Public/ChaosCloth/ChaosClothVisualization.h b/Engine/Plugins/ChaosCloth/Source/ChaosCloth/Public/ChaosCloth/ChaosClothVisualization.h
--- a/Engine/Plugins/ChaosCloth/Source/ChaosCloth/Public/ChaosCloth/ChaosClothVisualization.h
+++ b/Engine/Plugins/ChaosCloth/Source/ChaosCloth/Public/ChaosCloth/ChaosClothVisualization.h
@@ -48,7 +48,7 @@ namespace Chaos
 		CHAOSCLOTH_API void DrawEdgeConstraint(FPrimitiveDrawInterface* PDI = nullptr) const;
 		CHAOSCLOTH_API void DrawBendingConstraint(FPrimitiveDrawInterface* PDI = nullptr) const;
 		CHAOSCLOTH_API void DrawLongRangeConstraint(FPrimitiveDrawInterface* PDI = nullptr) const;
-		CHAOSCLOTH_API void DrawWindAndPressureForces(FPrimitiveDrawInterface* PDI = nullptr) const;
+		CHAOSCLOTH_API void DrawWindAndPressureForces(FPrimitiveDrawInterface* PDI = nullptr, const FReal ForceLengthScale = (FReal)10.) const;
 		CHAOSCLOTH_API void DrawLocalSpace(FPrimitiveDrawInterface* PDI = nullptr) const;
 		CHAOSCLOTH_API void DrawSelfCollision(FPrimitiveDrawInterface* PDI = nullptr) const;
 		CHAOSCLOTH_API void DrawSelfIntersection(FPrimitiveDrawInterface* PDI = nullptr) const;
@@ -79,7 +79,7 @@ namespace Chaos
 		void DrawEdgeConstraint(FPrimitiveDrawInterface* /*PDI*/ = nullptr) const {}
 		void DrawBendingConstraint(FPrimitiveDrawInterface* /*PDI*/ = nullptr) const {}
 		void DrawLongRangeConstraint(FPrimitiveDrawInterface* /*PDI*/ = nullptr) const {}
-		void DrawWindAndPressureForces(FPrimitiveDrawInterface* /*PDI*/ = nullptr) const {}
+		void DrawWindAndPressureForces(FPrimitiveDrawInterface* /*PDI*/ = nullptr, const FReal /*ForceLengthScale*/ = (FReal)10.) const {}
 		void DrawLocalSpace(FPrimitiveDrawInterface* /*PDI*/ = nullptr) const {}
 		void DrawSelfCollision(FPrimitiveDrawInterface* /*PDI*/ = nullptr) const {}
 		void DrawSelfIntersection(FPrimitiveDrawInterface* /*PDI*/ = nullptr) const {}
diff --git a/Engine/Plugins/ChaosClothAssetEditor/Source/ChaosClothAssetEditor/Private/ChaosClothAsset/ClothEditorSimulationVisualization.cpp b/Engine/Plugins/ChaosClothAssetEditor/Source/ChaosClothAssetEditor/Private/ChaosClothAsset/ClothEditorSimulationVisualization.cpp
--- a/Engine/Plugins/ChaosClothAssetEditor/Source/ChaosClothAssetEditor/Private/ChaosClothAsset/ClothEditorSimulationVisualization.cpp
+++ b/Engine/Plugins/ChaosClothAssetEditor/Source/ChaosClothAssetEditor/Private/ChaosClothAsset/ClothEditorSimulationVisualization.cpp
@@ -115,37 +115,42 @@ struct FVisualizationOption
 	FText ToolTip;             // Text for menu tooltips.
 	bool bDisablesSimulation;  // Whether or not this option requires the simulation to be disabled.
 	bool bHidesClothSections;  // Hides the cloth section to avoid zfighting with the debug geometry.
+	bool bDefaultFlagValue = false;
 
-	FVisualizationOption(const FClothVisualizationDebugDraw& InClothVisualizationDebugDraw, const FText& InDisplayName, const FText& InToolTip, bool bInDisablesSimulation = false, bool bInHidesClothSections = false, const FAdditionalMenus& InAdditionalMenus = FAdditionalMenus())
+	FVisualizationOption(const FClothVisualizationDebugDraw& InClothVisualizationDebugDraw, const FText& InDisplayName, const FText& InToolTip, bool bInDisablesSimulation = false, bool bInHidesClothSections = false, bool bInDefaultFlagValue = false, const FAdditionalMenus& InAdditionalMenus = FAdditionalMenus())
 		: ClothVisualizationDebugDraw(InClothVisualizationDebugDraw)
 		, AdditionalMenus(InAdditionalMenus)
 		, DisplayName(InDisplayName)
 		, ToolTip(InToolTip)
 		, bDisablesSimulation(bInDisablesSimulation)
 		, bHidesClothSections(bInHidesClothSections)
+		, bDefaultFlagValue(bInDefaultFlagValue)
 	{}
 
-	FVisualizationOption(const FLocalDebugDisplayString& InLocalDebugDisplayString, const FText& InDisplayName, const FText& InToolTip, bool bInDisablesSimulation = false, bool bInHidesClothSections = false)
+	FVisualizationOption(const FLocalDebugDisplayString& InLocalDebugDisplayString, const FText& InDisplayName, const FText& InToolTip, bool bInDisablesSimulation = false, bool bInHidesClothSections = false, bool bInDefaultFlagValue = false)
 		: LocalDebugDisplayString(InLocalDebugDisplayString)
 		, DisplayName(InDisplayName)
 		, ToolTip(InToolTip)
 		, bDisablesSimulation(bInDisablesSimulation)
 		, bHidesClothSections(bInHidesClothSections)
+		, bDefaultFlagValue(bInDefaultFlagValue)
 	{}
 
-	FVisualizationOption(const FClothVisualizationDebugDrawTexts& InClothVisualizationDebugDrawTexts, const FText & InDisplayName, const FText & InToolTip, bool bInDisablesSimulation = false, bool bInHidesClothSections = false)
+	FVisualizationOption(const FClothVisualizationDebugDrawTexts& InClothVisualizationDebugDrawTexts, const FText & InDisplayName, const FText & InToolTip, bool bInDisablesSimulation = false, bool bInHidesClothSections = false, bool bInDefaultFlagValue = false)
 		: ClothVisualizationDebugDrawTexts(InClothVisualizationDebugDrawTexts)
 		, DisplayName(InDisplayName)
 		, ToolTip(InToolTip)
 		, bDisablesSimulation(bInDisablesSimulation)
 		, bHidesClothSections(bInHidesClothSections)
+		, bDefaultFlagValue(bInDefaultFlagValue)
 	{}
 	
-	FVisualizationOption(const FText & InDisplayName, const FText & InToolTip, bool bInDisablesSimulation = false, bool bInHidesClothSections = false)
+	FVisualizationOption(const FText & InDisplayName, const FText & InToolTip, bool bInDisablesSimulation = false, bool bInHidesClothSections = false, bool bInDefaultFlagValue = false)
 		: DisplayName(InDisplayName)
 		, ToolTip(InToolTip)
 		, bDisablesSimulation(bInDisablesSimulation)
 		, bHidesClothSections(bInHidesClothSections)
+		, bDefaultFlagValue(bInDefaultFlagValue)
 	{}
 
 	~FVisualizationOption()
@@ -155,6 +160,14 @@ struct FVisualizationOption
 
 const FVisualizationOption FVisualizationOption::OptionData[] =
 {
+	FVisualizationOption(
+		FLocalDebugDisplayString::CreateLambda([](const FClothEditorSimulationVisualization&,const FClothSimulationProxy& Proxy)
+		{
+			return GetSimulationStatisticsString(Proxy);
+		}),
+		LOCTEXT("ChaosVisName_SimulationStatistics", "Simulation Statistics"),
+		LOCTEXT("ChaosVisName_SimulationStatistics_Tooltip", "Displays simulation statistics"),
+		/*bDisablesSimulation =*/false, /*bHidesClothSections=*/true, /*bDefaultFlagValue=*/true),
 	FVisualizationOption(
 		FClothVisualizationDebugDraw::CreateLambda([](const FClothEditorSimulationVisualization&, const ::Chaos::FClothVisualizationNoGC& Visualization, FPrimitiveDrawInterface* PDI)
 		{
@@ -202,7 +215,7 @@ const FVisualizationOption FVisualizationOption::OptionData[] =
 		}),
 		LOCTEXT("ChaosVisName_PointNormals", "Physical Mesh Normals"), 
 		LOCTEXT("ChaosVisName_PointNormals_ToolTip", "Draws the current point normals for the simulation mesh"),
-		/*bDisablesSimulation =*/false, /*bHidesClothSections=*/false,
+		/*bDisablesSimulation =*/false, /*bHidesClothSections=*/false, /*bDefaultFlagValue=*/false, 
 		FAdditionalMenus::CreateLambda([](FClothEditorSimulationVisualization& EditorVisualization, FMenuBuilder& MenuBuilder, TSharedRef<FChaosClothAssetEditor3DViewportClient> ViewportClient)
 		{
 			EditorVisualization.ExtendViewportShowMenuPointNormalsLength(MenuBuilder);
@@ -221,7 +234,7 @@ const FVisualizationOption FVisualizationOption::OptionData[] =
 		}),
 		LOCTEXT("ChaosVisName_AnimNormals", "Animated Mesh Normals"), 
 		LOCTEXT("ChaosVisName_AnimNormals_ToolTip", "Draws the current point normals for the animated mesh"),
-		/*bDisablesSimulation =*/false, /*bHidesClothSections=*/false,
+		/*bDisablesSimulation =*/false, /*bHidesClothSections=*/false, /*bDefaultFlagValue=*/false,
 		FAdditionalMenus::CreateLambda([](FClothEditorSimulationVisualization& EditorVisualization, FMenuBuilder& MenuBuilder, TSharedRef<FChaosClothAssetEditor3DViewportClient> ViewportClient)
 		{
 			EditorVisualization.ExtendViewportShowMenuAnimatedNormalsLength(MenuBuilder);
@@ -291,12 +304,17 @@ const FVisualizationOption FVisualizationOption::OptionData[] =
 		LOCTEXT("ChaosVisName_LongRangeConstraint" , "Long Range Constraint"), 
 		LOCTEXT("ChaosVisName_LongRangeConstraint_Tooltip", "Draws the long range attachment constraint distances")),
 	FVisualizationOption(
-		FClothVisualizationDebugDraw::CreateLambda([](const FClothEditorSimulationVisualization&, const ::Chaos::FClothVisualizationNoGC& Visualization, FPrimitiveDrawInterface* PDI)
+		FClothVisualizationDebugDraw::CreateLambda([](const FClothEditorSimulationVisualization& EditorVisualization, const ::Chaos::FClothVisualizationNoGC& Visualization, FPrimitiveDrawInterface* PDI)
 		{
-			Visualization.DrawWindAndPressureForces(PDI);
+			Visualization.DrawWindAndPressureForces(PDI, EditorVisualization.GetAerodynamicsLengthScale());
 		}),
 		LOCTEXT("ChaosVisName_WindAndPressureForces", "Wind Aerodynamic And Pressure Forces"), 
-		LOCTEXT("ChaosVisName_WindAndPressure_Tooltip", "Draws the Wind drag and lift and pressure forces")),
+		LOCTEXT("ChaosVisName_WindAndPressure_Tooltip", "Draws the Wind drag and lift and pressure forces"),
+		/*bDisablesSimulation =*/false, /*bHidesClothSections=*/false, /*bDefaultFlagValue=*/false,
+		FAdditionalMenus::CreateLambda([](FClothEditorSimulationVisualization& EditorVisualization, FMenuBuilder& MenuBuilder, TSharedRef<FChaosClothAssetEditor3DViewportClient> ViewportClient)
+		{
+			EditorVisualization.ExtendViewportShowMenuAerodynamicsLengthScale(MenuBuilder);
+		})),
 	FVisualizationOption(
 		FClothVisualizationDebugDraw::CreateLambda([](const FClothEditorSimulationVisualization&, const ::Chaos::FClothVisualizationNoGC& Visualization, FPrimitiveDrawInterface* PDI)
 		{
@@ -368,13 +386,6 @@ const FVisualizationOption FVisualizationOption::OptionData[] =
 		}),
 		LOCTEXT("ChaosVisName_DrawMultiResConstraint", "Draw Multi Res Constraint"), 
 		LOCTEXT("ChaosVisName_DrawMultiResConstraint_Tooltip", "Draw multi res constraint coarse mesh and targets.")),
-	FVisualizationOption(
-		FLocalDebugDisplayString::CreateLambda([](const FClothEditorSimulationVisualization&,const FClothSimulationProxy& Proxy)
-		{
-			return GetSimulationStatisticsString(Proxy);
-		}),
-		LOCTEXT("ChaosVisName_SimulationStatistics", "Simulation Statistics"), 
-		LOCTEXT("ChaosVisName_SimulationStatistics_Tooltip", "Displays simulation statistics")),
 	FVisualizationOption(
 		FClothVisualizationDebugDraw::CreateLambda([](const FClothEditorSimulationVisualization& EditorVisualization, const ::Chaos::FClothVisualizationNoGC& Visualization, FPrimitiveDrawInterface* PDI)
 		{
@@ -383,7 +394,7 @@ const FVisualizationOption FVisualizationOption::OptionData[] =
 		}),
 		LOCTEXT("ChaosVisName_DrawWeightMap", "Weight Map"), 
 		LOCTEXT("ChaosVisName_DrawWeightMap_ToolTip", "Draw the weight map for the simulation mesh. You can control the name of the map to be visualized by setting the p.ChaosClothVisualization.WeightMapName console variable."), 
-		/*bDisablesSimulation =*/false, /*bHidesClothSections=*/true,
+		/*bDisablesSimulation =*/false, /*bHidesClothSections=*/true, /*bDefaultFlagValue=*/false,
 		FAdditionalMenus::CreateLambda([](FClothEditorSimulationVisualization& EditorVisualization, FMenuBuilder& MenuBuilder, TSharedRef<FChaosClothAssetEditor3DViewportClient> ViewportClient)
 		{
 			EditorVisualization.ExtendViewportShowMenuWeightMapSelector(MenuBuilder, ViewportClient);
@@ -406,6 +417,10 @@ const uint32 FVisualizationOption::Count = sizeof(OptionData) / sizeof(FVisualiz
 FClothEditorSimulationVisualization::FClothEditorSimulationVisualization()
 	: Flags(false, Private::FVisualizationOption::Count)
 {
+	for (uint32 OptionIndex = 0; OptionIndex < Private::FVisualizationOption::Count; ++OptionIndex)
+	{
+		Flags[OptionIndex] = Private::FVisualizationOption::OptionData[OptionIndex].bDefaultFlagValue;
+	}
 }
 
 void FClothEditorSimulationVisualization::ExtendViewportShowMenu(FMenuBuilder& MenuBuilder, TSharedRef<FChaosClothAssetEditor3DViewportClient> ViewportClient)
@@ -445,6 +460,7 @@ void FClothEditorSimulationVisualization::ExtendViewportShowMenu(FMenuBuilder& M
 
 			const FUIAction Action(ExecuteAction, FCanExecuteAction(), IsActionChecked);
 
+			// Add menu entry
 			MenuBuilder.AddMenuEntry(Private::FVisualizationOption::OptionData[OptionIndex].DisplayName, Private::FVisualizationOption::OptionData[OptionIndex].ToolTip, FSlateIcon(), Action, NAME_None, EUserInterfaceActionType::ToggleButton);
 			Private::FVisualizationOption::OptionData[OptionIndex].AdditionalMenus.ExecuteIfBound(*this, MenuBuilder, ViewportClient);
 		}
diff --git a/Engine/Plugins/ChaosClothAssetEditor/Source/ChaosClothAssetEditor/Private/ChaosClothAsset/ClothEditorSimulationVisualization.h b/Engine/Plugins/ChaosClothAssetEditor/Source/ChaosClothAssetEditor/Private/ChaosClothAsset/ClothEditorSimulationVisualization.h
--- a/Engine/Plugins/ChaosClothAssetEditor/Source/ChaosClothAssetEditor/Private/ChaosClothAsset/ClothEditorSimulationVisualization.h
+++ b/Engine/Plugins/ChaosClothAssetEditor/Source/ChaosClothAssetEditor/Private/ChaosClothAsset/ClothEditorSimulationVisualization.h
@@ -44,8 +44,13 @@ public:
 	{
 		ExtendViewportShowMenuSpinBox(MenuBuilder, AnimatedNormalLength, 0.f, FLT_MAX, 0.f, 40.f);
 	}
+	void ExtendViewportShowMenuAerodynamicsLengthScale(FMenuBuilder& MenuBuilder)
+	{
+		ExtendViewportShowMenuSpinBox(MenuBuilder, AerodynamicsLengthScale, 0.f, FLT_MAX, 0.f, 40.f);
+	}
 	float GetPointNormalLength() const { return PointNormalLength; }
 	float GetAnimatedNormalLength() const { return AnimatedNormalLength; }
+	float GetAerodynamicsLengthScale() const { return AerodynamicsLengthScale; }
 private:
 	/** Return whether or not - given the current enabled options - the simulation should be disabled. */
 	bool ShouldDisableSimulation() const;
@@ -62,6 +67,7 @@ private:
 	TSharedPtr<FString> CurrentlySelectedWeightMap;
 	float PointNormalLength = 20.f;
 	float AnimatedNormalLength = 20.f;
+	float AerodynamicsLengthScale = 10.f;
 	
 };
 } // namespace UE::Chaos::ClothAsset
diff --git a/Engine/Source/Runtime/Experimental/Chaos/Private/Chaos/VelocityField.cpp b/Engine/Source/Runtime/Experimental/Chaos/Private/Chaos/VelocityField.cpp
--- a/Engine/Source/Runtime/Experimental/Chaos/Private/Chaos/VelocityField.cpp
+++ b/Engine/Source/Runtime/Experimental/Chaos/Private/Chaos/VelocityField.cpp
@@ -405,6 +405,18 @@ void FVelocityAndPressureField::SetMultipliers(
 	Pressure = FPBDFlatWeightMap(PressureValues, PressureMultipliers, TConstArrayView<TVec3<int32>>(Elements), Offset, NumParticles);
 }
 
+FSolverVec3 FVelocityAndPressureField::CalculateForce(const TConstArrayView<FSolverVec3>& Xs, const TConstArrayView<FSolverVec3>& Vs, int32 ElementIndex) const
+{
+	if (Private::VelocityFieldMaxVelocity <= 0.f)
+	{
+		return CalculateForce(Xs, Vs, ElementIndex, Velocity, Drag.GetValue(ElementIndex), OuterDrag.GetValue(ElementIndex), Lift.GetValue(ElementIndex), OuterLift.GetValue(ElementIndex), Pressure.GetValue(ElementIndex));
+	}
+	else
+	{
+		return CalculateForce(Xs, Vs, ElementIndex, Velocity, Drag.GetValue(ElementIndex), OuterDrag.GetValue(ElementIndex), Lift.GetValue(ElementIndex), OuterLift.GetValue(ElementIndex), Pressure.GetValue(ElementIndex), FMath::Square((FSolverReal)Private::VelocityFieldMaxVelocity));
+	}
+}
+
 void FVelocityAndPressureField::UpdateForces(const FSolverParticles& InParticles, const FSolverReal /*Dt*/)
 {
 	TRACE_CPUPROFILER_EVENT_SCOPE(FVelocityAndPressureField_UpdateForces);
@@ -644,7 +656,7 @@ void FVelocityAndPressureField::Apply(FSolverParticlesRange& InParticles, const
 			{
 				for (int32 ElementIndex = 0; ElementIndex < Elements.Num(); ++ElementIndex)
 				{
-					const FSolverVec3 Force = CalculateForce(InParticles, ElementIndex, Velocity,
+					const FSolverVec3 Force = CalculateForce(InParticles.XArray(), InParticles.GetV(), ElementIndex, Velocity,
 						(FSolverReal)Drag,
 						(FSolverReal)OuterDrag,
 						(FSolverReal)Lift,
@@ -659,7 +671,7 @@ void FVelocityAndPressureField::Apply(FSolverParticlesRange& InParticles, const
 			{
 				for (int32 ElementIndex = 0; ElementIndex < Elements.Num(); ++ElementIndex)
 				{
-					const FSolverVec3 Force = CalculateForce(InParticles, ElementIndex, Velocity,
+					const FSolverVec3 Force = CalculateForce(InParticles.XArray(), InParticles.GetV(), ElementIndex, Velocity,
 						(FSolverReal)Drag,
 						(FSolverReal)OuterDrag,
 						(FSolverReal)Lift,
@@ -757,7 +769,7 @@ void FVelocityAndPressureField::Apply(FSolverParticlesRange& InParticles, const
 					const FSolverReal ClO = OuterLift.GetValue(ElementIndex);
 					const FSolverReal Cp = Pressure.GetValue(ElementIndex);
 
-					const FSolverVec3 Force = CalculateForce(InParticles, ElementIndex, Velocity, CdI, CdO, ClI, ClO, Cp);
+					const FSolverVec3 Force = CalculateForce(InParticles.XArray(), InParticles.GetV(), ElementIndex, Velocity, CdI, CdO, ClI, ClO, Cp);
 					InParticles.Acceleration(Elements[ElementIndex][0]) += InParticles.InvM(Elements[ElementIndex][0]) * Force;
 					InParticles.Acceleration(Elements[ElementIndex][1]) += InParticles.InvM(Elements[ElementIndex][1]) * Force;
 					InParticles.Acceleration(Elements[ElementIndex][2]) += InParticles.InvM(Elements[ElementIndex][2]) * Force;
@@ -773,7 +785,7 @@ void FVelocityAndPressureField::Apply(FSolverParticlesRange& InParticles, const
 					const FSolverReal ClO = OuterLift.GetValue(ElementIndex);
 					const FSolverReal Cp = Pressure.GetValue(ElementIndex);
 
-					const FSolverVec3 Force = CalculateForce(InParticles, ElementIndex, Velocity, CdI, CdO, ClI, ClO, Cp, MaxVelocitySquared);
+					const FSolverVec3 Force = CalculateForce(InParticles.XArray(), InParticles.GetV(), ElementIndex, Velocity, CdI, CdO, ClI, ClO, Cp, MaxVelocitySquared);
 					InParticles.Acceleration(Elements[ElementIndex][0]) += InParticles.InvM(Elements[ElementIndex][0]) * Force;
 					InParticles.Acceleration(Elements[ElementIndex][1]) += InParticles.InvM(Elements[ElementIndex][1]) * Force;
 					InParticles.Acceleration(Elements[ElementIndex][2]) += InParticles.InvM(Elements[ElementIndex][2]) * Force;
diff --git a/Engine/Source/Runtime/Experimental/Chaos/Public/Chaos/VelocityField.h b/Engine/Source/Runtime/Experimental/Chaos/Public/Chaos/VelocityField.h
--- a/Engine/Source/Runtime/Experimental/Chaos/Public/Chaos/VelocityField.h
+++ b/Engine/Source/Runtime/Experimental/Chaos/Public/Chaos/VelocityField.h
@@ -195,6 +195,9 @@ public:
 	TConstArrayView<TVector<int32, 3>> GetElements() const { return TConstArrayView<TVector<int32, 3>>(Elements); }
 	TConstArrayView<FSolverVec3> GetForces() const { return TConstArrayView<FSolverVec3>(Forces); }
 
+	// This method is currently used for debug drawing.
+	CHAOS_API FSolverVec3 CalculateForce(const TConstArrayView<FSolverVec3>& Xs, const TConstArrayView<FSolverVec3>& Vs, int32 ElementIndex) const;
+
 private:
 	bool AreAerodynamicsEnabled() const { return QuarterRho > (FSolverReal)0.; }
 
@@ -214,22 +217,21 @@ private:
 		const TConstArrayView<FRealSingle>& PressureMultipliers);
 
 
-	template<typename SolverParticlesOrRange>
-	FSolverVec3 CalculateForce(const SolverParticlesOrRange& InParticles, int32 ElementIndex, const FSolverVec3& InVelocity, const FSolverReal CdI, const FSolverReal CdO, const FSolverReal ClI, const FSolverReal ClO, const FSolverReal Cp) const 
+	FSolverVec3 CalculateForce(const TConstArrayView<FSolverVec3>& Xs, const TConstArrayView<FSolverVec3>& Vs, int32 ElementIndex, const FSolverVec3& InVelocity, const FSolverReal CdI, const FSolverReal CdO, const FSolverReal ClI, const FSolverReal ClO, const FSolverReal Cp) const
 	{
 		const TVec3<int32>& Element = Elements[ElementIndex];
 
 		// Calculate the normal and the area of the surface exposed to the flow
 		FSolverVec3 N = FSolverVec3::CrossProduct(
-			InParticles.GetX(Element[2]) - InParticles.GetX(Element[0]),
-			InParticles.GetX(Element[1]) - InParticles.GetX(Element[0]));
+			Xs[Element[2]] - Xs[Element[0]],
+			Xs[Element[1]] - Xs[Element[0]]);
 		const FSolverReal DoubleArea = N.SafeNormalize();
 
 		// Calculate the direction and the relative velocity of the triangle to the flow
 		const FSolverVec3& SurfaceVelocity = (FSolverReal)(1. / 3.) * (
-			InParticles.V(Element[0]) +
-			InParticles.V(Element[1]) +
-			InParticles.V(Element[2]));
+			Vs[Element[0]] +
+			Vs[Element[1]] +
+			Vs[Element[2]]);
 		const FSolverVec3 V = InVelocity - SurfaceVelocity;
 
 		// Set the aerodynamic forces
@@ -243,11 +245,10 @@ private:
 
 	void UpdateField(const FSolverParticles& InParticles, int32 ElementIndex, const FSolverVec3& InVelocity, const FSolverReal CdI, const FSolverReal CdO, const FSolverReal ClI, const FSolverReal ClO, const FSolverReal Cp)
 	{
-		Forces[ElementIndex] = CalculateForce(InParticles, ElementIndex, InVelocity, CdI, CdO, ClI, ClO, Cp);
+		Forces[ElementIndex] = CalculateForce(TConstArrayView<FSolverVec3>(InParticles.XArray()), TConstArrayView<FSolverVec3>(InParticles.GetV()), ElementIndex, InVelocity, CdI, CdO, ClI, ClO, Cp);
 	}
 
-	template<typename SolverParticlesOrRange>
-	FSolverVec3 CalculateForce(const SolverParticlesOrRange& InParticles, int32 ElementIndex, const FSolverVec3& InVelocity, const FSolverReal CdI, const FSolverReal CdO, const FSolverReal ClI, const FSolverReal ClO, const FSolverReal Cp, const FSolverReal MaxVelocitySquared) const
+	FSolverVec3 CalculateForce(const TConstArrayView<FSolverVec3>& Xs, const TConstArrayView<FSolverVec3>& Vs, int32 ElementIndex, const FSolverVec3& InVelocity, const FSolverReal CdI, const FSolverReal CdO, const FSolverReal ClI, const FSolverReal ClO, const FSolverReal Cp, const FSolverReal MaxVelocitySquared) const
 	{
 		checkSlow(MaxVelocitySquared > (FSolverReal)0);
 
@@ -255,15 +256,15 @@ private:
 
 		// Calculate the normal and the area of the surface exposed to the flow
 		FSolverVec3 N = FSolverVec3::CrossProduct(
-			InParticles.GetX(Element[2]) - InParticles.GetX(Element[0]),
-			InParticles.GetX(Element[1]) - InParticles.GetX(Element[0]));
+			Xs[Element[2]] - Xs[Element[0]],
+			Xs[Element[1]] - Xs[Element[0]]);
 		const FSolverReal DoubleArea = N.SafeNormalize();
 
 		// Calculate the direction and the relative velocity of the triangle to the flow
 		const FSolverVec3& SurfaceVelocity = (FSolverReal)(1. / 3.) * (
-			InParticles.V(Element[0]) +
-			InParticles.V(Element[1]) +
-			InParticles.V(Element[2]));
+			Vs[Element[0]] +
+			Vs[Element[1]] +
+			Vs[Element[2]]);
 		FSolverVec3 V = InVelocity - SurfaceVelocity;
 
 		// Clamp the velocity
@@ -284,7 +285,7 @@ private:
 
 	void UpdateField(const FSolverParticles& InParticles, int32 ElementIndex, const FSolverVec3& InVelocity, const FSolverReal CdI, const FSolverReal CdO, const FSolverReal ClI, const FSolverReal ClO, const FSolverReal Cp, const FSolverReal MaxVelocitySquared)
 	{
-		Forces[ElementIndex] = CalculateForce(InParticles, ElementIndex, InVelocity, CdI, CdO, ClI, ClO, Cp, MaxVelocitySquared);
+		Forces[ElementIndex] = CalculateForce(TConstArrayView<FSolverVec3>(InParticles.XArray()), TConstArrayView<FSolverVec3>(InParticles.GetV()), ElementIndex, InVelocity, CdI, CdO, ClI, ClO, Cp, MaxVelocitySquared);
 	}
 
 private:
