CIS - Fixed Issue 590180: Compile warnings in Module.ChaosClothAssetDataflowNodes.cpp, ImportFilePath.gen.cpp and ImportFilePath.h

diff --git a/Engine/Plugins/ChaosClothAssetEditor/Source/ChaosClothAssetDataflowNodes/Public/ChaosClothAsset/ImportFilePath.h b/Engine/Plugins/ChaosClothAssetEditor/Source/ChaosClothAssetDataflowNodes/Public/ChaosClothAsset/ImportFilePath.h
--- a/Engine/Plugins/ChaosClothAssetEditor/Source/ChaosClothAssetDataflowNodes/Public/ChaosClothAsset/ImportFilePath.h
+++ b/Engine/Plugins/ChaosClothAssetEditor/Source/ChaosClothAssetDataflowNodes/Public/ChaosClothAsset/ImportFilePath.h
@@ -5,9 +5,11 @@
 #include "Delegates/Delegate.h"
 #include "ImportFilePath.generated.h"
 
+PRAGMA_DISABLE_DEPRECATION_WARNINGS  // For bForceReimport
 USTRUCT()
 struct FChaosClothAssetImportFilePath
 {
+PRAGMA_ENABLE_DEPRECATION_WARNINGS
 	GENERATED_USTRUCT_BODY()
 
 public:
