Change Dataflow namespace to UE::Dataflow to be ok with coding standard

diff --git a/Engine/Plugins/ChaosClothAssetEditor/Source/ChaosClothAssetDataflowNodes/Public/ChaosClothAsset/TransferSkinWeightsNode.h b/Engine/Plugins/ChaosClothAssetEditor/Source/ChaosClothAssetDataflowNodes/Public/ChaosClothAsset/TransferSkinWeightsNode.h
--- a/Engine/Plugins/ChaosClothAssetEditor/Source/ChaosClothAssetDataflowNodes/Public/ChaosClothAsset/TransferSkinWeightsNode.h
+++ b/Engine/Plugins/ChaosClothAssetEditor/Source/ChaosClothAssetDataflowNodes/Public/ChaosClothAsset/TransferSkinWeightsNode.h
@@ -152,7 +152,7 @@ public:
 	UPROPERTY(EditAnywhere, Category = "Transfer Skin Weights", Meta = (DisplayName = "Max Bone Influences"))
 	EChaosClothAssetMaxNumInfluences MaxNumInfluences = EChaosClothAssetMaxNumInfluences::Eight;
 	
-	FChaosClothAssetTransferSkinWeightsNode(const Dataflow::FNodeParameters& InParam, FGuid InGuid = FGuid::NewGuid());
+	FChaosClothAssetTransferSkinWeightsNode(const UE::Dataflow::FNodeParameters& InParam, FGuid InGuid = FGuid::NewGuid());
 
 private:
 
@@ -178,5 +178,5 @@ private:
 	 * different normals which could potentially lead to different weights being computed. You can either try to 
 	 * increase the normal threshold or switch to the closest point method.
 	 */
-	virtual void Evaluate(Dataflow::FContext& Context, const FDataflowOutput* Out) const override;
+	virtual void Evaluate(UE::Dataflow::FContext& Context, const FDataflowOutput* Out) const override;
 };
diff --git a/Engine/Plugins/ChaosClothAssetEditor/Source/ChaosClothAssetDataflowNodes/Public/ChaosClothAsset/TransformPositionsNode.h b/Engine/Plugins/ChaosClothAssetEditor/Source/ChaosClothAssetDataflowNodes/Public/ChaosClothAsset/TransformPositionsNode.h
--- a/Engine/Plugins/ChaosClothAssetEditor/Source/ChaosClothAssetDataflowNodes/Public/ChaosClothAsset/TransformPositionsNode.h
+++ b/Engine/Plugins/ChaosClothAssetEditor/Source/ChaosClothAssetDataflowNodes/Public/ChaosClothAsset/TransformPositionsNode.h
@@ -74,8 +74,8 @@ public:
 	UPROPERTY(EditAnywhere, Category = "TransformPositions - Render", Meta = (UIMax = 10, ClampMin = -1, EditCondition = "bTransformRenderPositions"))
 	int32 RenderPattern = INDEX_NONE;
 
-	FChaosClothAssetTransformPositionsNode(const Dataflow::FNodeParameters& InParam, FGuid InGuid = FGuid::NewGuid());
+	FChaosClothAssetTransformPositionsNode(const UE::Dataflow::FNodeParameters& InParam, FGuid InGuid = FGuid::NewGuid());
 
 private:
-	virtual void Evaluate(Dataflow::FContext& Context, const FDataflowOutput* Out) const override;
+	virtual void Evaluate(UE::Dataflow::FContext& Context, const FDataflowOutput* Out) const override;
 };
diff --git a/Engine/Plugins/ChaosClothAssetEditor/Source/ChaosClothAssetDataflowNodes/Public/ChaosClothAsset/TransformUVsNode.h b/Engine/Plugins/ChaosClothAssetEditor/Source/ChaosClothAssetDataflowNodes/Public/ChaosClothAsset/TransformUVsNode.h
--- a/Engine/Plugins/ChaosClothAssetEditor/Source/ChaosClothAssetDataflowNodes/Public/ChaosClothAsset/TransformUVsNode.h
+++ b/Engine/Plugins/ChaosClothAssetEditor/Source/ChaosClothAssetDataflowNodes/Public/ChaosClothAsset/TransformUVsNode.h
@@ -38,7 +38,7 @@ public:
 	UPROPERTY(EditAnywhere, Category = "Transform UVs", Meta = (UIMax = 5, ClampMin = -1))
 	int32 UVChannel = INDEX_NONE;
 
-	FChaosClothAssetTransformUVsNode(const Dataflow::FNodeParameters& InParam, FGuid InGuid = FGuid::NewGuid());
+	FChaosClothAssetTransformUVsNode(const UE::Dataflow::FNodeParameters& InParam, FGuid InGuid = FGuid::NewGuid());
 
-	virtual void Evaluate(Dataflow::FContext& Context, const FDataflowOutput* Out) const override;
+	virtual void Evaluate(UE::Dataflow::FContext& Context, const FDataflowOutput* Out) const override;
 };
diff --git a/Engine/Plugins/ChaosClothAssetEditor/Source/ChaosClothAssetDataflowNodes/Public/ChaosClothAsset/USDImportNode.h b/Engine/Plugins/ChaosClothAssetEditor/Source/ChaosClothAssetDataflowNodes/Public/ChaosClothAsset/USDImportNode.h
--- a/Engine/Plugins/ChaosClothAssetEditor/Source/ChaosClothAssetDataflowNodes/Public/ChaosClothAsset/USDImportNode.h
+++ b/Engine/Plugins/ChaosClothAssetEditor/Source/ChaosClothAssetDataflowNodes/Public/ChaosClothAsset/USDImportNode.h
@@ -23,7 +23,7 @@ public:
 	UPROPERTY(EditAnywhere, Category = "USD Import", Meta = (DisplayName = "USD File"))
 	FChaosClothAssetImportFilePath UsdFile;
 
-	FChaosClothAssetUSDImportNode(const Dataflow::FNodeParameters& InParam, FGuid InGuid = FGuid::NewGuid());
+	FChaosClothAssetUSDImportNode(const UE::Dataflow::FNodeParameters& InParam, FGuid InGuid = FGuid::NewGuid());
 
 private:
 	friend struct FChaosClothAssetUSDImportNode_v2;  // For ImportFromFile
@@ -37,8 +37,8 @@ private:
 		class FText& OutErrorText);
 
 	//~ Begin FDataflowNode interface
-	virtual void SetAssetValue(TObjectPtr<UObject> Asset, Dataflow::FContext& Context) const override;
-	virtual void Evaluate(Dataflow::FContext& Context, const FDataflowOutput* Out) const override;
+	virtual void SetAssetValue(TObjectPtr<UObject> Asset, UE::Dataflow::FContext& Context) const override;
+	virtual void Evaluate(UE::Dataflow::FContext& Context, const FDataflowOutput* Out) const override;
 	virtual void Serialize(FArchive& Archive) override;
 	//~ End FDataflowNode interface
 
diff --git a/Engine/Plugins/ChaosClothAssetEditor/Source/ChaosClothAssetDataflowNodes/Public/ChaosClothAsset/USDImportNode_v2.h b/Engine/Plugins/ChaosClothAssetEditor/Source/ChaosClothAssetDataflowNodes/Public/ChaosClothAsset/USDImportNode_v2.h
--- a/Engine/Plugins/ChaosClothAssetEditor/Source/ChaosClothAssetDataflowNodes/Public/ChaosClothAsset/USDImportNode_v2.h
+++ b/Engine/Plugins/ChaosClothAssetEditor/Source/ChaosClothAssetDataflowNodes/Public/ChaosClothAsset/USDImportNode_v2.h
@@ -49,11 +49,11 @@ public:
 	UPROPERTY(EditAnywhere, Category = "USD Import", Meta = (ButtonImage = "Icons.Refresh", EditCondition = "ImportedRenderStaticMesh != nullptr && bImportRenderMesh", EditConditionHides))
 	FDataflowFunctionProperty ReloadRenderStaticMesh;
 
-	FChaosClothAssetUSDImportNode_v2(const Dataflow::FNodeParameters& InParam, FGuid InGuid = FGuid::NewGuid());
+	FChaosClothAssetUSDImportNode_v2(const UE::Dataflow::FNodeParameters& InParam, FGuid InGuid = FGuid::NewGuid());
 
 private:
 	//~ Begin FDataflowNode interface
-	virtual void Evaluate(Dataflow::FContext& Context, const FDataflowOutput* Out) const override;
+	virtual void Evaluate(UE::Dataflow::FContext& Context, const FDataflowOutput* Out) const override;
 	virtual void Serialize(FArchive& Archive) override;
 	//~ End FDataflowNode interface
 
diff --git a/Engine/Plugins/ChaosClothAssetEditor/Source/ChaosClothAssetDataflowNodes/Public/ChaosClothAsset/WeightMapNode.h b/Engine/Plugins/ChaosClothAssetEditor/Source/ChaosClothAssetDataflowNodes/Public/ChaosClothAsset/WeightMapNode.h
--- a/Engine/Plugins/ChaosClothAssetEditor/Source/ChaosClothAssetDataflowNodes/Public/ChaosClothAsset/WeightMapNode.h
+++ b/Engine/Plugins/ChaosClothAssetEditor/Source/ChaosClothAssetDataflowNodes/Public/ChaosClothAsset/WeightMapNode.h
@@ -63,7 +63,7 @@ public:
 	UPROPERTY(EditAnywhere, Category = "Weight Map", Meta = (ButtonImage = "Icons.Convert"))
 	FDataflowFunctionProperty Transfer;
 
-	FChaosClothAssetWeightMapNode(const Dataflow::FNodeParameters& InParam, FGuid InGuid = FGuid::NewGuid());
+	FChaosClothAssetWeightMapNode(const UE::Dataflow::FNodeParameters& InParam, FGuid InGuid = FGuid::NewGuid());
 
 private:
 	friend class UClothEditorWeightMapPaintTool;
@@ -75,7 +75,7 @@ private:
 	TArray<float>& GetVertexWeights() { return VertexWeights; }
 
 	// These methods are exported for UClothEditorWeightMapPaintTool which lives in a different module.
-	FName CHAOSCLOTHASSETDATAFLOWNODES_API GetInputName(Dataflow::FContext& Context) const;
+	FName CHAOSCLOTHASSETDATAFLOWNODES_API GetInputName(UE::Dataflow::FContext& Context) const;
 
 	void CHAOSCLOTHASSETDATAFLOWNODES_API SetVertexWeights(const TConstArrayView<float> InputMap, const TArray<float>& FinalValues);
 
@@ -83,7 +83,7 @@ private:
 	void CHAOSCLOTHASSETDATAFLOWNODES_API CalculateFinalVertexWeightValues(const TConstArrayView<float> InputMap, TArrayView<float> FinalOutputMap) const;
 
 	//~ Begin FDataflowNode interface
-	virtual void Evaluate(Dataflow::FContext& Context, const FDataflowOutput* Out) const override;
+	virtual void Evaluate(UE::Dataflow::FContext& Context, const FDataflowOutput* Out) const override;
 	virtual void Serialize(FArchive& Ar) override;
 	virtual FDataflowOutput* RedirectSerializedOutput(const FName& MissingOutputName) override;
 	//~ End FDataflowNode interface
diff --git a/Engine/Plugins/ChaosClothAssetEditor/Source/ChaosClothAssetDataflowNodes/Public/ChaosClothAsset/WeightMapToSelectionNode.h b/Engine/Plugins/ChaosClothAssetEditor/Source/ChaosClothAssetDataflowNodes/Public/ChaosClothAsset/WeightMapToSelectionNode.h
--- a/Engine/Plugins/ChaosClothAssetEditor/Source/ChaosClothAssetDataflowNodes/Public/ChaosClothAsset/WeightMapToSelectionNode.h
+++ b/Engine/Plugins/ChaosClothAssetEditor/Source/ChaosClothAssetDataflowNodes/Public/ChaosClothAsset/WeightMapToSelectionNode.h
@@ -52,8 +52,8 @@ public:
 	UPROPERTY(EditAnywhere, Category = "Weight Map To Selection", Meta = (ClampMin = "0", ClampMax = "1"))
 	float SelectionThreshold = 0.95f;
 
-	FChaosClothAssetWeightMapToSelectionNode(const Dataflow::FNodeParameters& InParam, FGuid InGuid = FGuid::NewGuid());
+	FChaosClothAssetWeightMapToSelectionNode(const UE::Dataflow::FNodeParameters& InParam, FGuid InGuid = FGuid::NewGuid());
 private:
 
-	virtual void Evaluate(Dataflow::FContext& Context, const FDataflowOutput* Out) const override;
+	virtual void Evaluate(UE::Dataflow::FContext& Context, const FDataflowOutput* Out) const override;
 };
diff --git a/Engine/Plugins/ChaosClothAssetEditor/Source/ChaosClothAssetEditor/Private/ChaosClothAsset/ClothEditor3DViewportClient.h b/Engine/Plugins/ChaosClothAssetEditor/Source/ChaosClothAssetEditor/Private/ChaosClothAsset/ClothEditor3DViewportClient.h
--- a/Engine/Plugins/ChaosClothAssetEditor/Source/ChaosClothAssetEditor/Private/ChaosClothAsset/ClothEditor3DViewportClient.h
+++ b/Engine/Plugins/ChaosClothAssetEditor/Source/ChaosClothAssetEditor/Private/ChaosClothAsset/ClothEditor3DViewportClient.h
@@ -128,7 +128,7 @@ private:
 	bool bRenderMeshWireframe = false;
 
 	// Dataflow render support
-	Dataflow::FTimestamp LastModifiedTimestamp = Dataflow::FTimestamp::Invalid;
+	UE::Dataflow::FTimestamp LastModifiedTimestamp = UE::Dataflow::FTimestamp::Invalid;
 
 	// Gizmo support
 	TObjectPtr<UTransformProxy> TransformProxy = nullptr;
diff --git a/Engine/Plugins/ChaosClothAssetEditor/Source/ChaosClothAssetEditor/Private/ChaosClothAsset/ClothEditorMode.cpp b/Engine/Plugins/ChaosClothAssetEditor/Source/ChaosClothAssetEditor/Private/ChaosClothAsset/ClothEditorMode.cpp
--- a/Engine/Plugins/ChaosClothAssetEditor/Source/ChaosClothAssetEditor/Private/ChaosClothAsset/ClothEditorMode.cpp
+++ b/Engine/Plugins/ChaosClothAssetEditor/Source/ChaosClothAssetEditor/Private/ChaosClothAsset/ClothEditorMode.cpp
@@ -1559,7 +1559,7 @@ void UChaosClothAssetEditorMode::InitializeContextObject()
 		RestSpaceToolsContext->ContextObjectStore->AddContextObject(DataflowContextObject);
 	}
 
-	DataflowContextObject->SetConstructionViewMode(Dataflow::FRenderingViewModeFactory::GetInstance().GetViewMode(ClothViewModeToDataflowViewModeName(ConstructionViewMode)));
+	DataflowContextObject->SetConstructionViewMode(UE::Dataflow::FRenderingViewModeFactory::GetInstance().GetViewMode(ClothViewModeToDataflowViewModeName(ConstructionViewMode)));
 
 	if (const TSharedPtr<SDataflowGraphEditor> GraphEditor = DataflowGraphEditor.Pin())
 	{
@@ -1581,7 +1581,7 @@ void UChaosClothAssetEditorMode::UpdateContextObject(const TSharedPtr<FManagedAr
 	
 	if (UDataflowContextObject* DataflowContextObject = RestSpaceToolsContext->ContextObjectStore->FindContext<UDataflowContextObject>())
 	{
-		DataflowContextObject->SetConstructionViewMode(Dataflow::FRenderingViewModeFactory::GetInstance().GetViewMode(ClothViewModeToDataflowViewModeName(ConstructionViewMode)));
+		DataflowContextObject->SetConstructionViewMode(UE::Dataflow::FRenderingViewModeFactory::GetInstance().GetViewMode(ClothViewModeToDataflowViewModeName(ConstructionViewMode)));
 
 		DataflowContextObject->SetSelectedCollection(Collection, bDynamicMeshUseInputCollection);
 	}
@@ -1596,7 +1596,7 @@ void UChaosClothAssetEditorMode::DeleteContextObject()
 	}
 }
 
-void UChaosClothAssetEditorMode::SetDataflowContext(TWeakPtr<Dataflow::FEngineContext> InDataflowContext)
+void UChaosClothAssetEditorMode::SetDataflowContext(TWeakPtr<UE::Dataflow::FEngineContext> InDataflowContext)
 {
 	DataflowContext = InDataflowContext;
 	UEditorInteractiveToolsContext* const RestSpaceToolsContext = GetInteractiveToolsContext();
diff --git a/Engine/Plugins/ChaosClothAssetEditor/Source/ChaosClothAssetEditor/Private/ChaosClothAsset/ClothEditorMode.h b/Engine/Plugins/ChaosClothAssetEditor/Source/ChaosClothAssetEditor/Private/ChaosClothAsset/ClothEditorMode.h
--- a/Engine/Plugins/ChaosClothAssetEditor/Source/ChaosClothAssetEditor/Private/ChaosClothAsset/ClothEditorMode.h
+++ b/Engine/Plugins/ChaosClothAssetEditor/Source/ChaosClothAssetEditor/Private/ChaosClothAsset/ClothEditorMode.h
@@ -186,7 +186,7 @@ private:
 	TSharedPtr<FManagedArrayCollection> GetClothCollection();
 	TSharedPtr<FManagedArrayCollection> GetInputClothCollection();
 
-	void SetDataflowContext(TWeakPtr<Dataflow::FEngineContext> InDataflowContext);
+	void SetDataflowContext(TWeakPtr<UE::Dataflow::FEngineContext> InDataflowContext);
 	void SetDataflowGraphEditor(TSharedPtr<SDataflowGraphEditor> InGraphEditor);
 	
 	void StartToolForSelectedNode(const UObject* SelectedNode);
@@ -310,7 +310,7 @@ private:
 
 	UPROPERTY()
 	TObjectPtr<UEditorInteractiveToolsContext> ActiveToolsContext = nullptr;
-	TWeakPtr<Dataflow::FEngineContext> DataflowContext;
+	TWeakPtr<UE::Dataflow::FEngineContext> DataflowContext;
 	TSharedPtr<FManagedArrayCollection> SelectedClothCollection = nullptr;
 	TSharedPtr<FManagedArrayCollection> SelectedInputClothCollection = nullptr;
 
diff --git a/Engine/Plugins/ChaosClothAssetEditor/Source/ChaosClothAssetEditor/Private/ChaosClothAsset/ClothEditorToolkit.cpp b/Engine/Plugins/ChaosClothAssetEditor/Source/ChaosClothAssetEditor/Private/ChaosClothAsset/ClothEditorToolkit.cpp
--- a/Engine/Plugins/ChaosClothAssetEditor/Source/ChaosClothAssetEditor/Private/ChaosClothAsset/ClothEditorToolkit.cpp
+++ b/Engine/Plugins/ChaosClothAssetEditor/Source/ChaosClothAssetEditor/Private/ChaosClothAsset/ClothEditorToolkit.cpp
@@ -205,7 +205,7 @@ FChaosClothAssetEditorToolkit::~FChaosClothAssetEditorToolkit()
 	EditorModeManager->DestroyMode(UChaosClothAssetEditorMode::EM_ChaosClothAssetEditorModeId);
 }
 
-TSharedPtr<Dataflow::FEngineContext> FChaosClothAssetEditorToolkit::GetDataflowContext() const
+TSharedPtr<UE::Dataflow::FEngineContext> FChaosClothAssetEditorToolkit::GetDataflowContext() const
 {
 	return DataflowContext;
 }
@@ -469,7 +469,7 @@ void FChaosClothAssetEditorToolkit::PostInitAssetEditor()
 
 	ClothMode->DataflowGraph = GetDataflow();
 	ClothMode->SetDataflowGraphEditor(GraphEditor);
-	DataflowContext = TSharedPtr<Dataflow::FEngineContext>(new Dataflow::FClothAssetDataflowContext(GetAsset(), GetDataflow()));
+	DataflowContext = TSharedPtr<UE::Dataflow::FEngineContext>(new UE::Dataflow::FClothAssetDataflowContext(GetAsset(), GetDataflow()));
 	ClothMode->SetDataflowContext(DataflowContext);
 
 	// Handle Dataflow asset reload event
@@ -858,7 +858,7 @@ void FChaosClothAssetEditorToolkit::EvaluateNode(const FDataflowNode* Node, cons
 		SlowTask.EnterProgressFrame(1.f);
 
 		// Evaluate the node and update LastDataflowNodeTimestamp if doing so
-		const Dataflow::FTimestamp OldTimestamp = LastDataflowNodeTimestamp;
+		const UE::Dataflow::FTimestamp OldTimestamp = LastDataflowNodeTimestamp;
 		Node = FDataflowEditorCommands::EvaluateNode(
 			*DataflowContext,
 			LastDataflowNodeTimestamp,
@@ -893,7 +893,7 @@ TSharedRef<SDataflowGraphEditor> FChaosClothAssetEditorToolkit::CreateGraphEdito
 {
 	UDataflow* const Dataflow = GetDataflow();
 	ensure(Dataflow);
-	using namespace Dataflow;
+	using namespace UE::Dataflow;
 
 	const auto EvalLambda = [this](const FDataflowNode* Node, const FDataflowOutput* Output)
 	{
@@ -1022,7 +1022,7 @@ void FChaosClothAssetEditorToolkit::OnNodeTitleCommitted(const FText& InNewText,
 }
 
 
-TSharedPtr<FManagedArrayCollection> FChaosClothAssetEditorToolkit::GetClothCollectionIfPossible(const TSharedPtr<FDataflowNode> InDataflowNode, const TSharedPtr<Dataflow::FEngineContext> Context)
+TSharedPtr<FManagedArrayCollection> FChaosClothAssetEditorToolkit::GetClothCollectionIfPossible(const TSharedPtr<FDataflowNode> InDataflowNode, const TSharedPtr<UE::Dataflow::FEngineContext> Context)
 {
 	if (InDataflowNode && Context)
 	{
@@ -1050,7 +1050,7 @@ TSharedPtr<FManagedArrayCollection> FChaosClothAssetEditorToolkit::GetClothColle
 }
 
 
-TSharedPtr<FManagedArrayCollection> FChaosClothAssetEditorToolkit::GetInputClothCollectionIfPossible(const TSharedPtr<FDataflowNode> InDataflowNode, const TSharedPtr<Dataflow::FEngineContext> Context)
+TSharedPtr<FManagedArrayCollection> FChaosClothAssetEditorToolkit::GetInputClothCollectionIfPossible(const TSharedPtr<FDataflowNode> InDataflowNode, const TSharedPtr<UE::Dataflow::FEngineContext> Context)
 {
 	if (InDataflowNode && Context)
 	{
@@ -1083,7 +1083,7 @@ TSharedPtr<FDataflowNode> FChaosClothAssetEditorToolkit::GetSelectedDataflowNode
 	{
 		if (UDataflow* const Dataflow = GetDataflow())
 		{
-			if (TSharedPtr<Dataflow::FGraph> Graph = Dataflow->GetDataflow())
+			if (TSharedPtr<UE::Dataflow::FGraph> Graph = Dataflow->GetDataflow())
 			{
 				return Graph->FindBaseNode(SelectedDataflowNodeGuid);
 			}
@@ -1098,7 +1098,7 @@ TSharedPtr<const FDataflowNode> FChaosClothAssetEditorToolkit::GetSelectedDatafl
 	{
 		if (const UDataflow* const Dataflow = GetDataflow())
 		{
-			if (TSharedPtr<const Dataflow::FGraph> Graph = Dataflow->GetDataflow())
+			if (TSharedPtr<const UE::Dataflow::FGraph> Graph = Dataflow->GetDataflow())
 			{
 				return Graph->FindBaseNode(SelectedDataflowNodeGuid);
 			}
@@ -1226,7 +1226,7 @@ void FChaosClothAssetEditorToolkit::OnNodeSelectionChanged(const TSet<UObject*>&
 			}
 		}
 
-		Dataflow->LastModifiedRenderTarget = Dataflow::FTimestamp::Current();
+		Dataflow->LastModifiedRenderTarget = UE::Dataflow::FTimestamp::Current();
 	}
 
 
diff --git a/Engine/Plugins/ChaosClothAssetEditor/Source/ChaosClothAssetEditor/Private/ChaosClothAsset/ClothEditorToolkit.h b/Engine/Plugins/ChaosClothAssetEditor/Source/ChaosClothAssetEditor/Private/ChaosClothAsset/ClothEditorToolkit.h
--- a/Engine/Plugins/ChaosClothAssetEditor/Source/ChaosClothAssetEditor/Private/ChaosClothAsset/ClothEditorToolkit.h
+++ b/Engine/Plugins/ChaosClothAssetEditor/Source/ChaosClothAssetEditor/Private/ChaosClothAsset/ClothEditorToolkit.h
@@ -24,7 +24,7 @@ class FClothEditorSimulationVisualization;
 class FChaosClothAssetEditor3DViewportClient;
 }
 
-namespace Dataflow
+namespace UE::Dataflow
 {
 	class CHAOSCLOTHASSETEDITOR_API FClothAssetDataflowContext final : public TEngineContext<FContextSingle>
 	{
@@ -54,7 +54,7 @@ public:
 	explicit FChaosClothAssetEditorToolkit(UAssetEditor* InOwningAssetEditor);
 	virtual ~FChaosClothAssetEditorToolkit();
 
-	TSharedPtr<Dataflow::FEngineContext> GetDataflowContext() const;
+	TSharedPtr<UE::Dataflow::FEngineContext> GetDataflowContext() const;
 	const UDataflow* GetDataflow() const;
 
 private:
@@ -129,8 +129,8 @@ private:
 	void ReinitializeGraphEditorWidget();
 	TSharedPtr<IStructureDetailsView> CreateNodeDetailsEditorWidget(UObject* ObjectToEdit);
 
-	TSharedPtr<FManagedArrayCollection> GetClothCollectionIfPossible(const TSharedPtr<FDataflowNode> InDataflowNode, const TSharedPtr<Dataflow::FEngineContext> Context);
-	TSharedPtr<FManagedArrayCollection> GetInputClothCollectionIfPossible(const TSharedPtr<FDataflowNode> InDataflowNode, const TSharedPtr<Dataflow::FEngineContext> Context);
+	TSharedPtr<FManagedArrayCollection> GetClothCollectionIfPossible(const TSharedPtr<FDataflowNode> InDataflowNode, const TSharedPtr<UE::Dataflow::FEngineContext> Context);
+	TSharedPtr<FManagedArrayCollection> GetInputClothCollectionIfPossible(const TSharedPtr<FDataflowNode> InDataflowNode, const TSharedPtr<UE::Dataflow::FEngineContext> Context);
 	TSharedPtr<FDataflowNode> GetSelectedDataflowNode();
 	TSharedPtr<const FDataflowNode> GetSelectedDataflowNode() const;
 
@@ -165,8 +165,8 @@ private:
 	TSharedPtr<SClothCollectionOutliner> Outliner;
 
 	// Dataflow
-	TSharedPtr<Dataflow::FEngineContext> DataflowContext;
-	Dataflow::FTimestamp LastDataflowNodeTimestamp = Dataflow::FTimestamp::Invalid;
+	TSharedPtr<UE::Dataflow::FEngineContext> DataflowContext;
+	UE::Dataflow::FTimestamp LastDataflowNodeTimestamp = UE::Dataflow::FTimestamp::Invalid;
 	FDelegateHandle OnNodeInvalidatedDelegateHandle;
 	FGuid SelectedDataflowNodeGuid;
 
diff --git a/Engine/Plugins/ChaosClothAssetEditor/Source/ChaosClothAssetEditorTools/Private/ChaosClothAsset/ChaosClothAssetEditorToolsModule.cpp b/Engine/Plugins/ChaosClothAssetEditor/Source/ChaosClothAssetEditorTools/Private/ChaosClothAsset/ChaosClothAssetEditorToolsModule.cpp
--- a/Engine/Plugins/ChaosClothAssetEditor/Source/ChaosClothAssetEditorTools/Private/ChaosClothAsset/ChaosClothAssetEditorToolsModule.cpp
+++ b/Engine/Plugins/ChaosClothAssetEditor/Source/ChaosClothAssetEditorTools/Private/ChaosClothAsset/ChaosClothAssetEditorToolsModule.cpp
@@ -16,7 +16,7 @@ namespace UE::Chaos::ClothAsset
 	public:
 		virtual void StartupModule() override
 		{
-			Dataflow::FDataflowToolRegistry& ToolRegistry = Dataflow::FDataflowToolRegistry::Get();
+			UE::Dataflow::FDataflowToolRegistry& ToolRegistry = UE::Dataflow::FDataflowToolRegistry::Get();
 		
 			const TSharedRef<const FClothToolActionCommandBindings> ClothToolActions = MakeShared<FClothToolActionCommandBindings>();
 
@@ -27,7 +27,7 @@ namespace UE::Chaos::ClothAsset
 
 		virtual void ShutdownModule() override
 		{
-			Dataflow::FDataflowToolRegistry& ToolRegistry = Dataflow::FDataflowToolRegistry::Get();
+			UE::Dataflow::FDataflowToolRegistry& ToolRegistry = UE::Dataflow::FDataflowToolRegistry::Get();
 			ToolRegistry.RemoveNodeToToolMapping(FChaosClothAssetWeightMapNode::StaticType());
 			ToolRegistry.RemoveNodeToToolMapping(FChaosClothAssetSelectionNode_v2::StaticType());
 			ToolRegistry.RemoveNodeToToolMapping(FChaosClothAssetTransferSkinWeightsNode::StaticType());
diff --git a/Engine/Plugins/ChaosClothAssetEditor/Source/ChaosClothAssetEditorTools/Private/ChaosClothAsset/ClothEditorContextObject.cpp b/Engine/Plugins/ChaosClothAssetEditor/Source/ChaosClothAssetEditorTools/Private/ChaosClothAsset/ClothEditorContextObject.cpp
--- a/Engine/Plugins/ChaosClothAssetEditor/Source/ChaosClothAssetEditorTools/Private/ChaosClothAsset/ClothEditorContextObject.cpp
+++ b/Engine/Plugins/ChaosClothAssetEditor/Source/ChaosClothAssetEditorTools/Private/ChaosClothAsset/ClothEditorContextObject.cpp
@@ -8,7 +8,7 @@
 
 #include UE_INLINE_GENERATED_CPP_BY_NAME(ClothEditorContextObject)
 
-void UClothEditorContextObject::Init(TWeakPtr<SDataflowGraphEditor> InDataflowGraphEditor, TWeakPtr<Dataflow::FEngineContext> InDataflowContext, UE::Chaos::ClothAsset::EClothPatternVertexType InConstructionViewMode, TWeakPtr<FManagedArrayCollection> InSelectedClothCollection, bool bInUsingInputCollection)
+void UClothEditorContextObject::Init(TWeakPtr<SDataflowGraphEditor> InDataflowGraphEditor, TWeakPtr<UE::Dataflow::FEngineContext> InDataflowContext, UE::Chaos::ClothAsset::EClothPatternVertexType InConstructionViewMode, TWeakPtr<FManagedArrayCollection> InSelectedClothCollection, bool bInUsingInputCollection)
 {
 	DataflowGraphEditor = InDataflowGraphEditor;
 	DataflowContext = InDataflowContext;
diff --git a/Engine/Plugins/ChaosClothAssetEditor/Source/ChaosClothAssetEditorTools/Private/ChaosClothAsset/ClothEditorToolBuilders.cpp b/Engine/Plugins/ChaosClothAssetEditor/Source/ChaosClothAssetEditorTools/Private/ChaosClothAsset/ClothEditorToolBuilders.cpp
--- a/Engine/Plugins/ChaosClothAssetEditor/Source/ChaosClothAssetEditorTools/Private/ChaosClothAsset/ClothEditorToolBuilders.cpp
+++ b/Engine/Plugins/ChaosClothAssetEditor/Source/ChaosClothAssetEditorTools/Private/ChaosClothAsset/ClothEditorToolBuilders.cpp
@@ -20,14 +20,14 @@
 
 // ------------------- Weight Map Paint Tool -------------------
 
-void UClothEditorWeightMapPaintToolBuilder::GetSupportedConstructionViewModes(const UDataflowContextObject& ContextObject, TArray<const Dataflow::IDataflowConstructionViewMode*>& Modes) const
+void UClothEditorWeightMapPaintToolBuilder::GetSupportedConstructionViewModes(const UDataflowContextObject& ContextObject, TArray<const UE::Dataflow::IDataflowConstructionViewMode*>& Modes) const
 {
 	using namespace UE::Chaos::ClothAsset;
 
-	const Dataflow::FRenderingViewModeFactory& Factory = Dataflow::FRenderingViewModeFactory::GetInstance();
-	const Dataflow::IDataflowConstructionViewMode* const Sim2DMode = Factory.GetViewMode(ClothViewModeToDataflowViewModeName(EClothPatternVertexType::Sim2D));
-	const Dataflow::IDataflowConstructionViewMode* const Sim3DMode = Factory.GetViewMode(ClothViewModeToDataflowViewModeName(EClothPatternVertexType::Sim3D));
-	const Dataflow::IDataflowConstructionViewMode* const RenderMode = Factory.GetViewMode(ClothViewModeToDataflowViewModeName(EClothPatternVertexType::Render));
+	const UE::Dataflow::FRenderingViewModeFactory& Factory = UE::Dataflow::FRenderingViewModeFactory::GetInstance();
+	const UE::Dataflow::IDataflowConstructionViewMode* const Sim2DMode = Factory.GetViewMode(ClothViewModeToDataflowViewModeName(EClothPatternVertexType::Sim2D));
+	const UE::Dataflow::IDataflowConstructionViewMode* const Sim3DMode = Factory.GetViewMode(ClothViewModeToDataflowViewModeName(EClothPatternVertexType::Sim3D));
+	const UE::Dataflow::IDataflowConstructionViewMode* const RenderMode = Factory.GetViewMode(ClothViewModeToDataflowViewModeName(EClothPatternVertexType::Render));
 
 	checkf(Sim2DMode, TEXT("Couldn't find DataflowConstructionViewMode corresponding to EClothPatternVertexType::Sim2D"));
 	checkf(Sim3DMode, TEXT("Couldn't find DataflowConstructionViewMode corresponding to EClothPatternVertexType::Sim3D"));
@@ -141,14 +141,14 @@ UMeshSurfacePointTool* UClothEditorWeightMapPaintToolBuilder::CreateNewTool(cons
 // ------------------- Selection Tool -------------------
 
 
-void UClothMeshSelectionToolBuilder::GetSupportedConstructionViewModes(const UDataflowContextObject& ContextObject, TArray<const Dataflow::IDataflowConstructionViewMode*>& Modes) const
+void UClothMeshSelectionToolBuilder::GetSupportedConstructionViewModes(const UDataflowContextObject& ContextObject, TArray<const UE::Dataflow::IDataflowConstructionViewMode*>& Modes) const
 {
 	using namespace UE::Chaos::ClothAsset;
 
-	const Dataflow::FRenderingViewModeFactory& Factory = Dataflow::FRenderingViewModeFactory::GetInstance();
-	const Dataflow::IDataflowConstructionViewMode* const Sim2DMode = Factory.GetViewMode(ClothViewModeToDataflowViewModeName(EClothPatternVertexType::Sim2D));
-	const Dataflow::IDataflowConstructionViewMode* const Sim3DMode = Factory.GetViewMode(ClothViewModeToDataflowViewModeName(EClothPatternVertexType::Sim3D));
-	const Dataflow::IDataflowConstructionViewMode* const RenderMode = Factory.GetViewMode(ClothViewModeToDataflowViewModeName(EClothPatternVertexType::Render));
+	const UE::Dataflow::FRenderingViewModeFactory& Factory = UE::Dataflow::FRenderingViewModeFactory::GetInstance();
+	const UE::Dataflow::IDataflowConstructionViewMode* const Sim2DMode = Factory.GetViewMode(ClothViewModeToDataflowViewModeName(EClothPatternVertexType::Sim2D));
+	const UE::Dataflow::IDataflowConstructionViewMode* const Sim3DMode = Factory.GetViewMode(ClothViewModeToDataflowViewModeName(EClothPatternVertexType::Sim3D));
+	const UE::Dataflow::IDataflowConstructionViewMode* const RenderMode = Factory.GetViewMode(ClothViewModeToDataflowViewModeName(EClothPatternVertexType::Render));
 
 	checkf(Sim2DMode, TEXT("Couldn't find DataflowConstructionViewMode corresponding to EClothPatternVertexType::Sim2D"));
 	checkf(Sim3DMode, TEXT("Couldn't find DataflowConstructionViewMode corresponding to EClothPatternVertexType::Sim3D"));
@@ -241,11 +241,11 @@ UInteractiveTool* UClothMeshSelectionToolBuilder::BuildTool(const FToolBuilderSt
 // ------------------- Skin Weight Transfer Tool -------------------
 
 
-void UClothTransferSkinWeightsToolBuilder::GetSupportedConstructionViewModes(const UDataflowContextObject& ContextObject, TArray<const Dataflow::IDataflowConstructionViewMode*>& Modes) const
+void UClothTransferSkinWeightsToolBuilder::GetSupportedConstructionViewModes(const UDataflowContextObject& ContextObject, TArray<const UE::Dataflow::IDataflowConstructionViewMode*>& Modes) const
 {
 	using namespace UE::Chaos::ClothAsset;
-	const Dataflow::FRenderingViewModeFactory& Factory = Dataflow::FRenderingViewModeFactory::GetInstance();
-	const Dataflow::IDataflowConstructionViewMode* const Sim3DMode = Factory.GetViewMode(ClothViewModeToDataflowViewModeName(EClothPatternVertexType::Sim3D));
+	const UE::Dataflow::FRenderingViewModeFactory& Factory = UE::Dataflow::FRenderingViewModeFactory::GetInstance();
+	const UE::Dataflow::IDataflowConstructionViewMode* const Sim3DMode = Factory.GetViewMode(ClothViewModeToDataflowViewModeName(EClothPatternVertexType::Sim3D));
 	Modes.Add(Sim3DMode);
 }
 
@@ -288,7 +288,7 @@ namespace UE::Chaos::ClothAsset
 		ToolCDOs.Add(GetMutableDefault<UClothMeshSelectionTool>());
 	}
 
-	EClothPatternVertexType DataflowViewModeToClothViewMode(const Dataflow::IDataflowConstructionViewMode* DataflowViewMode)
+	EClothPatternVertexType DataflowViewModeToClothViewMode(const UE::Dataflow::IDataflowConstructionViewMode* DataflowViewMode)
 	{
 		const FName ViewModeName = DataflowViewMode->GetName();
 		if (ViewModeName == FName("Cloth2DSimView"))
diff --git a/Engine/Plugins/ChaosClothAssetEditor/Source/ChaosClothAssetEditorTools/Private/ChaosClothAsset/ClothMeshSelectionTool.cpp b/Engine/Plugins/ChaosClothAssetEditor/Source/ChaosClothAssetEditorTools/Private/ChaosClothAsset/ClothMeshSelectionTool.cpp
--- a/Engine/Plugins/ChaosClothAssetEditor/Source/ChaosClothAssetEditorTools/Private/ChaosClothAsset/ClothMeshSelectionTool.cpp
+++ b/Engine/Plugins/ChaosClothAssetEditor/Source/ChaosClothAssetEditorTools/Private/ChaosClothAsset/ClothMeshSelectionTool.cpp
@@ -492,7 +492,7 @@ bool UClothMeshSelectionTool::GetSelectedNodeInfo(FString& OutSelectionName, UE:
 		ensure(DataflowContextObject->IsUsingInputCollection());
 		if (const TSharedPtr<const FManagedArrayCollection> ClothCollection = DataflowContextObject->GetSelectedCollection())
 		{
-			if (const TSharedPtr<Dataflow::FEngineContext> DataflowContext = DataflowContextObject->GetDataflowContext())
+			if (const TSharedPtr<UE::Dataflow::FEngineContext> DataflowContext = DataflowContextObject->GetDataflowContext())
 			{
 				using namespace UE::Chaos::ClothAsset;
 				const FName InputName = SelectionNodeToUpdate->GetInputName(*DataflowContext);
diff --git a/Engine/Plugins/ChaosClothAssetEditor/Source/ChaosClothAssetEditorTools/Private/ChaosClothAsset/ClothToolActionCommandBindings.h b/Engine/Plugins/ChaosClothAssetEditor/Source/ChaosClothAssetEditorTools/Private/ChaosClothAsset/ClothToolActionCommandBindings.h
--- a/Engine/Plugins/ChaosClothAssetEditor/Source/ChaosClothAssetEditorTools/Private/ChaosClothAsset/ClothToolActionCommandBindings.h
+++ b/Engine/Plugins/ChaosClothAssetEditor/Source/ChaosClothAssetEditorTools/Private/ChaosClothAsset/ClothToolActionCommandBindings.h
@@ -55,7 +55,7 @@ public:
 
 // Tool action registry entry
 
-class FClothToolActionCommandBindings : public Dataflow::FDataflowToolRegistry::IDataflowToolActionCommands
+class FClothToolActionCommandBindings : public UE::Dataflow::FDataflowToolRegistry::IDataflowToolActionCommands
 {
 
 public:
diff --git a/Engine/Plugins/ChaosClothAssetEditor/Source/ChaosClothAssetEditorTools/Private/ChaosClothAsset/ClothWeightMapPaintTool.cpp b/Engine/Plugins/ChaosClothAssetEditor/Source/ChaosClothAssetEditorTools/Private/ChaosClothAsset/ClothWeightMapPaintTool.cpp
--- a/Engine/Plugins/ChaosClothAssetEditor/Source/ChaosClothAssetEditorTools/Private/ChaosClothAsset/ClothWeightMapPaintTool.cpp
+++ b/Engine/Plugins/ChaosClothAssetEditor/Source/ChaosClothAssetEditorTools/Private/ChaosClothAsset/ClothWeightMapPaintTool.cpp
@@ -458,7 +458,7 @@ void UClothEditorWeightMapPaintTool::InitializeSculptMeshFromTarget()
 			const bool bIsRenderMode = ViewMode == EClothPatternVertexType::Render;
 
 			// Find the map if it exists.
-			if (TSharedPtr<Dataflow::FEngineContext> DataflowContext = DataflowContextObject->GetDataflowContext())
+			if (TSharedPtr<UE::Dataflow::FEngineContext> DataflowContext = DataflowContextObject->GetDataflowContext())
 			{
 				const FName InputName = WeightMapNodeToUpdate->GetInputName(*DataflowContext);
 				if (bIsRenderMode)
diff --git a/Engine/Plugins/ChaosClothAssetEditor/Source/ChaosClothAssetEditorTools/Public/ChaosClothAsset/ClothEditorContextObject.h b/Engine/Plugins/ChaosClothAssetEditor/Source/ChaosClothAssetEditorTools/Public/ChaosClothAsset/ClothEditorContextObject.h
--- a/Engine/Plugins/ChaosClothAssetEditor/Source/ChaosClothAssetEditorTools/Public/ChaosClothAsset/ClothEditorContextObject.h
+++ b/Engine/Plugins/ChaosClothAssetEditor/Source/ChaosClothAssetEditorTools/Public/ChaosClothAsset/ClothEditorContextObject.h
@@ -24,7 +24,7 @@ class CHAOSCLOTHASSETEDITORTOOLS_API UClothEditorContextObject : public UObject
 	GENERATED_BODY()
 
 public:
-	void Init(TWeakPtr<SDataflowGraphEditor> DataflowGraphEditor, TWeakPtr<Dataflow::FEngineContext> DataflowContext, UE::Chaos::ClothAsset::EClothPatternVertexType InConstructionViewMode, TWeakPtr<FManagedArrayCollection> SelectedClothCollection, bool bInUsingInputCollection);
+	void Init(TWeakPtr<SDataflowGraphEditor> DataflowGraphEditor, TWeakPtr<UE::Dataflow::FEngineContext> DataflowContext, UE::Chaos::ClothAsset::EClothPatternVertexType InConstructionViewMode, TWeakPtr<FManagedArrayCollection> SelectedClothCollection, bool bInUsingInputCollection);
 
 	UE_DEPRECATED(5.5, "Use the Init with a DataflowContext")
 	void Init(TWeakPtr<SDataflowGraphEditor> InDataflowGraphEditor, UE::Chaos::ClothAsset::EClothPatternVertexType InConstructionViewMode, TWeakPtr<FManagedArrayCollection> InSelectedClothCollection, TWeakPtr<FManagedArrayCollection> InSelectedInputClothCollection = nullptr)
@@ -55,12 +55,12 @@ public:
 		return nullptr;
 	}
 	
-	TWeakPtr<Dataflow::FEngineContext> GetDataflowContext() const 
+	TWeakPtr<UE::Dataflow::FEngineContext> GetDataflowContext() const 
 	{
 		return DataflowContext;
 	}
 
-	void SetDataflowContext(TWeakPtr<Dataflow::FEngineContext> InDataflowContext)
+	void SetDataflowContext(TWeakPtr<UE::Dataflow::FEngineContext> InDataflowContext)
 	{
 		DataflowContext = InDataflowContext;
 	}
@@ -96,7 +96,7 @@ public:
 private:
 
 	TWeakPtr<SDataflowGraphEditor> DataflowGraphEditor;
-	TWeakPtr<Dataflow::FEngineContext> DataflowContext;
+	TWeakPtr<UE::Dataflow::FEngineContext> DataflowContext;
 
 	UE::Chaos::ClothAsset::EClothPatternVertexType ConstructionViewMode;
 	TWeakPtr<const FManagedArrayCollection> SelectedClothCollection;
diff --git a/Engine/Plugins/ChaosClothAssetEditor/Source/ChaosClothAssetEditorTools/Public/ChaosClothAsset/ClothEditorToolBuilders.h b/Engine/Plugins/ChaosClothAssetEditor/Source/ChaosClothAssetEditorTools/Public/ChaosClothAsset/ClothEditorToolBuilders.h
--- a/Engine/Plugins/ChaosClothAssetEditor/Source/ChaosClothAssetEditorTools/Public/ChaosClothAsset/ClothEditorToolBuilders.h
+++ b/Engine/Plugins/ChaosClothAssetEditor/Source/ChaosClothAssetEditorTools/Public/ChaosClothAsset/ClothEditorToolBuilders.h
@@ -15,7 +15,7 @@ class UDataflowContextObject;
 class UClothEditorContextObject;
 #endif
 
-namespace Dataflow
+namespace UE::Dataflow
 {
 	class IDataflowConstructionViewMode;
 }
@@ -58,7 +58,7 @@ class CHAOSCLOTHASSETEDITORTOOLS_API UClothEditorWeightMapPaintToolBuilder : pub
 private:
 
 	// IDataflowEditorToolBuilder
-	virtual void GetSupportedConstructionViewModes(const UDataflowContextObject& ContextObject, TArray<const Dataflow::IDataflowConstructionViewMode*>& Modes) const override;
+	virtual void GetSupportedConstructionViewModes(const UDataflowContextObject& ContextObject, TArray<const UE::Dataflow::IDataflowConstructionViewMode*>& Modes) const override;
 	virtual bool CanSceneStateChange(const UInteractiveTool* ActiveTool, const FToolBuilderState& SceneState) const override;
 	virtual void SceneStateChanged(UInteractiveTool* ActiveTool, const FToolBuilderState& SceneState) override;
 
@@ -81,7 +81,7 @@ class CHAOSCLOTHASSETEDITORTOOLS_API UClothMeshSelectionToolBuilder : public UIn
 private:
 
 	// IDataflowEditorToolBuilder
-	virtual void GetSupportedConstructionViewModes(const UDataflowContextObject& ContextObject, TArray<const Dataflow::IDataflowConstructionViewMode*>& Modes) const override;
+	virtual void GetSupportedConstructionViewModes(const UDataflowContextObject& ContextObject, TArray<const UE::Dataflow::IDataflowConstructionViewMode*>& Modes) const override;
 	virtual bool CanSceneStateChange(const UInteractiveTool* ActiveTool, const FToolBuilderState& SceneState) const override;
 	virtual void SceneStateChanged(UInteractiveTool* ActiveTool, const FToolBuilderState& SceneState) override;
 
@@ -103,7 +103,7 @@ class CHAOSCLOTHASSETEDITORTOOLS_API UClothTransferSkinWeightsToolBuilder : publ
 private:
 
 	// IDataflowEditorToolBuilder
-	virtual void GetSupportedConstructionViewModes(const UDataflowContextObject& ContextObject, TArray<const Dataflow::IDataflowConstructionViewMode*>& Modes) const override;
+	virtual void GetSupportedConstructionViewModes(const UDataflowContextObject& ContextObject, TArray<const UE::Dataflow::IDataflowConstructionViewMode*>& Modes) const override;
 	virtual bool CanSceneStateChange(const UInteractiveTool* ActiveTool, const FToolBuilderState& SceneState) const override;
 	virtual void SceneStateChanged(UInteractiveTool* ActiveTool, const FToolBuilderState& SceneState) override;
 
@@ -121,7 +121,7 @@ namespace UE::Chaos::ClothAsset
 	void CHAOSCLOTHASSETEDITORTOOLS_API GetClothEditorToolDefaultObjectList(TArray<UInteractiveTool*>& ToolCDOs);
 
 	// Mapping from Dataflow View Mode to Cloth View Mode. Input object must be one of FCloth2DSimViewMode, FCloth3DSimViewMode, or FClothRenderViewMode, defined in ClothDataflowViewModes.h
-	EClothPatternVertexType CHAOSCLOTHASSETEDITORTOOLS_API DataflowViewModeToClothViewMode(const Dataflow::IDataflowConstructionViewMode* DataflowViewMode);
+	EClothPatternVertexType CHAOSCLOTHASSETEDITORTOOLS_API DataflowViewModeToClothViewMode(const UE::Dataflow::IDataflowConstructionViewMode* DataflowViewMode);
 
 	// Mapping from Cloth View Mode to Dataflow View Mode name. Ouptut will be one of "Cloth2DSimView", "Cloth3DSimView", or "ClothRenderView", as defined in ClothDataflowViewModes.cpp
 	FName CHAOSCLOTHASSETEDITORTOOLS_API ClothViewModeToDataflowViewModeName(EClothPatternVertexType ClothViewMode);
diff --git a/Engine/Plugins/Experimental/ChaosFlesh/Source/ChaosFleshDeprecatedNodes/Private/Dataflow/ChaosFleshConstructTetGridNode.cpp b/Engine/Plugins/Experimental/ChaosFlesh/Source/ChaosFleshDeprecatedNodes/Private/Dataflow/ChaosFleshConstructTetGridNode.cpp
--- a/Engine/Plugins/Experimental/ChaosFlesh/Source/ChaosFleshDeprecatedNodes/Private/Dataflow/ChaosFleshConstructTetGridNode.cpp
+++ b/Engine/Plugins/Experimental/ChaosFlesh/Source/ChaosFleshDeprecatedNodes/Private/Dataflow/ChaosFleshConstructTetGridNode.cpp
@@ -13,7 +13,7 @@
 // FConstructTetGridNode
 //=============================================================================
 
-void FConstructTetGridNode::Evaluate(Dataflow::FContext& Context, const FDataflowOutput* Out) const
+void FConstructTetGridNode::Evaluate(UE::Dataflow::FContext& Context, const FDataflowOutput* Out) const
 {
 	if (Out->IsA<FManagedArrayCollection>(&Collection))
 	{
@@ -31,7 +31,7 @@ void FConstructTetGridNode::Evaluate(Dataflow::FContext& Context, const FDataflo
 
 		UE_LOG(LogChaosFlesh, Display, TEXT("TetGrid generated %d points and %d tetrahedra."), X.Num(), Tets.Num());
 
-		TArray<FIntVector3> Tris = Dataflow::GetSurfaceTriangles(Tets, !bDiscardInteriorTriangles);
+		TArray<FIntVector3> Tris = UE::Dataflow::GetSurfaceTriangles(Tets, !bDiscardInteriorTriangles);
 		TUniquePtr<FTetrahedralCollection> TetCollection(
 			FTetrahedralCollection::NewTetrahedralCollection(X, Tris, Tets));
 		InCollection->AppendGeometry(*TetCollection.Get());
diff --git a/Engine/Plugins/Experimental/ChaosFlesh/Source/ChaosFleshDeprecatedNodes/Private/Dataflow/ChaosFleshCreateTetrahedralCollectionNode.cpp b/Engine/Plugins/Experimental/ChaosFlesh/Source/ChaosFleshDeprecatedNodes/Private/Dataflow/ChaosFleshCreateTetrahedralCollectionNode.cpp
--- a/Engine/Plugins/Experimental/ChaosFlesh/Source/ChaosFleshDeprecatedNodes/Private/Dataflow/ChaosFleshCreateTetrahedralCollectionNode.cpp
+++ b/Engine/Plugins/Experimental/ChaosFlesh/Source/ChaosFleshDeprecatedNodes/Private/Dataflow/ChaosFleshCreateTetrahedralCollectionNode.cpp
@@ -30,7 +30,7 @@
 //=============================================================================
 
 
-void FGenerateTetrahedralCollectionDataflowNodes::Evaluate(Dataflow::FContext& Context, const FDataflowOutput* Out) const
+void FGenerateTetrahedralCollectionDataflowNodes::Evaluate(UE::Dataflow::FContext& Context, const FDataflowOutput* Out) const
 {
 	if (Out->IsA<DataType>(&Collection))
 	{
@@ -145,7 +145,7 @@ void FGenerateTetrahedralCollectionDataflowNodes::Evaluate(Dataflow::FContext& C
 }
 
 void FGenerateTetrahedralCollectionDataflowNodes::EvaluateIsoStuffing(
-	Dataflow::FContext& Context, 
+	UE::Dataflow::FContext& Context, 
 	TUniquePtr<FFleshCollection>& InCollection,
 	const UE::Geometry::FDynamicMesh3& DynamicMesh) const
 {
@@ -174,7 +174,7 @@ void FGenerateTetrahedralCollectionDataflowNodes::EvaluateIsoStuffing(
 		{
 			TArray<FVector> Vertices; Vertices.SetNumUninitialized(IsosurfaceStuffing.Vertices.Num());
 			TArray<FIntVector4> Elements; Elements.SetNumUninitialized(IsosurfaceStuffing.Tets.Num());
-			TArray<FIntVector3> SurfaceElements = Dataflow::GetSurfaceTriangles(IsosurfaceStuffing.Tets, !bDiscardInteriorTriangles);
+			TArray<FIntVector3> SurfaceElements = UE::Dataflow::GetSurfaceTriangles(IsosurfaceStuffing.Tets, !bDiscardInteriorTriangles);
 
 			for (int32 Tdx = 0; Tdx < IsosurfaceStuffing.Tets.Num(); ++Tdx)
 			{
@@ -202,7 +202,7 @@ void FGenerateTetrahedralCollectionDataflowNodes::EvaluateIsoStuffing(
 }
 
 void FGenerateTetrahedralCollectionDataflowNodes::EvaluateTetWild(
-	Dataflow::FContext& Context, 
+	UE::Dataflow::FContext& Context, 
 	TUniquePtr<FFleshCollection>& InCollection,
 	const UE::Geometry::FDynamicMesh3& DynamicMesh) const
 {
@@ -240,7 +240,7 @@ void FGenerateTetrahedralCollectionDataflowNodes::EvaluateTetWild(
 		UE_LOG(LogChaosFlesh, Display,TEXT("Generating tet mesh via TetWild..."));
 		if (UE::Geometry::FTetWild::ComputeTetMesh(Params, Verts, Tris, TetVerts, Tets, &Progress))
 		{
-			TArray<FIntVector3> SurfaceElements = Dataflow::GetSurfaceTriangles(Tets, !bDiscardInteriorTriangles);
+			TArray<FIntVector3> SurfaceElements = UE::Dataflow::GetSurfaceTriangles(Tets, !bDiscardInteriorTriangles);
 			TUniquePtr<FTetrahedralCollection> TetCollection(FTetrahedralCollection::NewTetrahedralCollection(TetVerts, SurfaceElements, Tets));
 			InCollection->AppendGeometry(*TetCollection.Get());
 
diff --git a/Engine/Plugins/Experimental/ChaosFlesh/Source/ChaosFleshDeprecatedNodes/Private/Dataflow/ChaosFleshDeprecatedNodesPlugin.cpp b/Engine/Plugins/Experimental/ChaosFlesh/Source/ChaosFleshDeprecatedNodes/Private/Dataflow/ChaosFleshDeprecatedNodesPlugin.cpp
--- a/Engine/Plugins/Experimental/ChaosFlesh/Source/ChaosFleshDeprecatedNodes/Private/Dataflow/ChaosFleshDeprecatedNodesPlugin.cpp
+++ b/Engine/Plugins/Experimental/ChaosFlesh/Source/ChaosFleshDeprecatedNodes/Private/Dataflow/ChaosFleshDeprecatedNodesPlugin.cpp
@@ -14,7 +14,7 @@
 
 void IChaosFleshDeprecatedNodesPlugin::StartupModule()
 {
-	Dataflow::RegisterChaosFleshImportGEONodes();
+	UE::Dataflow::RegisterChaosFleshImportGEONodes();
 	DATAFLOW_NODE_REGISTER_CREATION_FACTORY(FGenerateTetrahedralCollectionDataflowNodes);
 	DATAFLOW_NODE_REGISTER_CREATION_FACTORY(FConstructTetGridNode);
 	DATAFLOW_NODE_REGISTER_CREATION_FACTORY(FKinematicOriginInsertionInitializationDataflowNode);
diff --git a/Engine/Plugins/Experimental/ChaosFlesh/Source/ChaosFleshDeprecatedNodes/Private/Dataflow/ChaosFleshImportGEO.cpp b/Engine/Plugins/Experimental/ChaosFlesh/Source/ChaosFleshDeprecatedNodes/Private/Dataflow/ChaosFleshImportGEO.cpp
--- a/Engine/Plugins/Experimental/ChaosFlesh/Source/ChaosFleshDeprecatedNodes/Private/Dataflow/ChaosFleshImportGEO.cpp
+++ b/Engine/Plugins/Experimental/ChaosFlesh/Source/ChaosFleshDeprecatedNodes/Private/Dataflow/ChaosFleshImportGEO.cpp
@@ -13,7 +13,7 @@
 
 DEFINE_LOG_CATEGORY(LogImportGEO);
 
-namespace Dataflow
+namespace UE::Dataflow
 {
 	void RegisterChaosFleshImportGEONodes()
 	{
@@ -30,7 +30,7 @@ namespace Dataflow
 
 
 void
-FExtractGEOInt::Evaluate(Dataflow::FContext& Context, const FDataflowOutput* Out) const
+FExtractGEOInt::Evaluate(UE::Dataflow::FContext& Context, const FDataflowOutput* Out) const
 {
 	if (Out->IsA<int32>(&Value))
 	{
@@ -45,7 +45,7 @@ FExtractGEOInt::Evaluate(Dataflow::FContext& Context, const FDataflowOutput* Out
 }
 
 void
-FExtractGEOIntVector::Evaluate(Dataflow::FContext& Context, const FDataflowOutput* Out) const
+FExtractGEOIntVector::Evaluate(UE::Dataflow::FContext& Context, const FDataflowOutput* Out) const
 {
 	if (Out->IsA<TArray<int32>>(&Value))
 	{
@@ -60,7 +60,7 @@ FExtractGEOIntVector::Evaluate(Dataflow::FContext& Context, const FDataflowOutpu
 }
 
 void
-FExtractGEOFloatVector::Evaluate(Dataflow::FContext& Context, const FDataflowOutput* Out) const
+FExtractGEOFloatVector::Evaluate(UE::Dataflow::FContext& Context, const FDataflowOutput* Out) const
 {
 	if (Out->IsA<TArray<float>>(&Value))
 	{
@@ -178,7 +178,7 @@ FImportGEO::ReadGEOFile(const bool PrintStats) const
 
 void
 FImportGEO::Evaluate(
-	Dataflow::FContext& Context,
+	UE::Dataflow::FContext& Context,
 	const FDataflowOutput* Out) const
 {
 	if (Out->IsA<FManagedArrayCollection>(&Collection))
@@ -242,7 +242,7 @@ FImportGEO::Evaluate(
 
 			if (Vertices.Num() && Elements.Num())
 			{
-				TArray<FIntVector3> SurfaceElements = Dataflow::GetSurfaceTriangles(Elements, !bDiscardInteriorTriangles);
+				TArray<FIntVector3> SurfaceElements = UE::Dataflow::GetSurfaceTriangles(Elements, !bDiscardInteriorTriangles);
 				TUniquePtr<FTetrahedralCollection> OBJCollection(
 					FTetrahedralCollection::NewTetrahedralCollection(Vertices, SurfaceElements, Elements));
 				InCollection->AppendGeometry(*OBJCollection.Get());
diff --git a/Engine/Plugins/Experimental/ChaosFlesh/Source/ChaosFleshDeprecatedNodes/Private/Dataflow/ChaosFleshKinematicOriginInsertionInitializationNode.cpp b/Engine/Plugins/Experimental/ChaosFlesh/Source/ChaosFleshDeprecatedNodes/Private/Dataflow/ChaosFleshKinematicOriginInsertionInitializationNode.cpp
--- a/Engine/Plugins/Experimental/ChaosFlesh/Source/ChaosFleshDeprecatedNodes/Private/Dataflow/ChaosFleshKinematicOriginInsertionInitializationNode.cpp
+++ b/Engine/Plugins/Experimental/ChaosFlesh/Source/ChaosFleshDeprecatedNodes/Private/Dataflow/ChaosFleshKinematicOriginInsertionInitializationNode.cpp
@@ -11,7 +11,7 @@
 #include UE_INLINE_GENERATED_CPP_BY_NAME(ChaosFleshKinematicOriginInsertionInitializationNode)
 
 
-void FKinematicOriginInsertionInitializationDataflowNode::Evaluate(Dataflow::FContext& Context, const FDataflowOutput* Out) const
+void FKinematicOriginInsertionInitializationDataflowNode::Evaluate(UE::Dataflow::FContext& Context, const FDataflowOutput* Out) const
 {
 	if (Out->IsA<DataType>(&Collection))
 	{
@@ -51,7 +51,7 @@ void FKinematicOriginInsertionInitializationDataflowNode::Evaluate(Dataflow::FCo
 						const FSkinWeightVertexBuffer* SkinWeightVertexBuffer = LODRenderData->GetSkinWeightVertexBuffer();
 						const int32 MaxBoneInfluences = SkinWeightVertexBuffer->GetMaxBoneInfluences();
 						TArray<FTransform> ComponentPose;
-						Dataflow::Animation::GlobalTransforms(BoneSkeletalMesh->GetRefSkeleton(), ComponentPose);
+						UE::Dataflow::Animation::GlobalTransforms(BoneSkeletalMesh->GetRefSkeleton(), ComponentPose);
 						TArray<TArray<int32>> BoneBoundVerts;
 						TArray<TArray<float>> BoneBoundWeights;
 						BoneBoundVerts.SetNum(ComponentPose.Num());
diff --git a/Engine/Plugins/Experimental/ChaosFlesh/Source/ChaosFleshDeprecatedNodes/Private/Dataflow/ChaosFleshKinematicTetrahedralConstraintNode.cpp b/Engine/Plugins/Experimental/ChaosFlesh/Source/ChaosFleshDeprecatedNodes/Private/Dataflow/ChaosFleshKinematicTetrahedralConstraintNode.cpp
--- a/Engine/Plugins/Experimental/ChaosFlesh/Source/ChaosFleshDeprecatedNodes/Private/Dataflow/ChaosFleshKinematicTetrahedralConstraintNode.cpp
+++ b/Engine/Plugins/Experimental/ChaosFlesh/Source/ChaosFleshDeprecatedNodes/Private/Dataflow/ChaosFleshKinematicTetrahedralConstraintNode.cpp
@@ -10,7 +10,7 @@
 #include UE_INLINE_GENERATED_CPP_BY_NAME(ChaosFleshKinematicTetrahedralConstraintNode)
 
 
-void FKinematicTetrahedralBindingsDataflowNode::Evaluate(Dataflow::FContext& Context, const FDataflowOutput* Out) const
+void FKinematicTetrahedralBindingsDataflowNode::Evaluate(UE::Dataflow::FContext& Context, const FDataflowOutput* Out) const
 {
 	if (Out->IsA<DataType>(&Collection))
 	{
@@ -26,7 +26,7 @@ void FKinematicTetrahedralBindingsDataflowNode::Evaluate(Dataflow::FContext& Con
 			ExclusionList.ParseIntoArray(StrArray, *FString(" "));				
 			int32 NumTets = Tetrahedron->Num();
 			TArray<FTransform> ComponentPose;
-			Dataflow::Animation::GlobalTransforms(SkeletalMesh->GetRefSkeleton(), ComponentPose);
+			UE::Dataflow::Animation::GlobalTransforms(SkeletalMesh->GetRefSkeleton(), ComponentPose);
 			TArray<bool> VertAdded;
 			VertAdded.Init(false, Vertex->Num());
 			auto DoSkipBoneIndex = [&StrArray, &SkeletalMesh](int32 BoneIndex)
diff --git a/Engine/Plugins/Experimental/ChaosFlesh/Source/ChaosFleshDeprecatedNodes/Public/Dataflow/ChaosFleshConstructTetGridNode.h b/Engine/Plugins/Experimental/ChaosFlesh/Source/ChaosFleshDeprecatedNodes/Public/Dataflow/ChaosFleshConstructTetGridNode.h
--- a/Engine/Plugins/Experimental/ChaosFlesh/Source/ChaosFleshDeprecatedNodes/Public/Dataflow/ChaosFleshConstructTetGridNode.h
+++ b/Engine/Plugins/Experimental/ChaosFlesh/Source/ChaosFleshDeprecatedNodes/Public/Dataflow/ChaosFleshConstructTetGridNode.h
@@ -32,14 +32,14 @@ public:
 	UPROPERTY(EditAnywhere, Category = "Dataflow")
 	bool bDiscardInteriorTriangles = true;
 
-	FConstructTetGridNode(const Dataflow::FNodeParameters& InParam, FGuid InGuid = FGuid::NewGuid())
+	FConstructTetGridNode(const UE::Dataflow::FNodeParameters& InParam, FGuid InGuid = FGuid::NewGuid())
 		: FDataflowNode(InParam, InGuid)
 	{
 		RegisterInputConnection(&Collection);
 		RegisterOutputConnection(&Collection);
 	}
 
-	virtual void Evaluate(Dataflow::FContext& Context, const FDataflowOutput* Out) const override;
+	virtual void Evaluate(UE::Dataflow::FContext& Context, const FDataflowOutput* Out) const override;
 };
 
 
diff --git a/Engine/Plugins/Experimental/ChaosFlesh/Source/ChaosFleshDeprecatedNodes/Public/Dataflow/ChaosFleshCreateTetrahedralCollectionNode.h b/Engine/Plugins/Experimental/ChaosFlesh/Source/ChaosFleshDeprecatedNodes/Public/Dataflow/ChaosFleshCreateTetrahedralCollectionNode.h
--- a/Engine/Plugins/Experimental/ChaosFlesh/Source/ChaosFleshDeprecatedNodes/Public/Dataflow/ChaosFleshCreateTetrahedralCollectionNode.h
+++ b/Engine/Plugins/Experimental/ChaosFlesh/Source/ChaosFleshDeprecatedNodes/Public/Dataflow/ChaosFleshCreateTetrahedralCollectionNode.h
@@ -99,7 +99,7 @@ public:
 	UPROPERTY(EditAnywhere, Category = "Dataflow")
 	bool bDiscardInteriorTriangles = true;
 
-	FGenerateTetrahedralCollectionDataflowNodes(const Dataflow::FNodeParameters& InParam, FGuid InGuid = FGuid::NewGuid())
+	FGenerateTetrahedralCollectionDataflowNodes(const UE::Dataflow::FNodeParameters& InParam, FGuid InGuid = FGuid::NewGuid())
 		: FDataflowNode(InParam, InGuid)
 	{
 		RegisterInputConnection(&StaticMesh);
@@ -108,11 +108,11 @@ public:
 		RegisterOutputConnection(&Collection);
 	}
 
-	virtual void Evaluate(Dataflow::FContext& Context, const FDataflowOutput* Out) const override;
+	virtual void Evaluate(UE::Dataflow::FContext& Context, const FDataflowOutput* Out) const override;
 
 protected:
-	void EvaluateIsoStuffing(Dataflow::FContext& Context, TUniquePtr<FFleshCollection>& InCollection, const UE::Geometry::FDynamicMesh3& DynamicMesh) const;
-	void EvaluateTetWild(Dataflow::FContext& Context, TUniquePtr<FFleshCollection>& InCollection, const UE::Geometry::FDynamicMesh3& DynamicMesh) const;
+	void EvaluateIsoStuffing(UE::Dataflow::FContext& Context, TUniquePtr<FFleshCollection>& InCollection, const UE::Geometry::FDynamicMesh3& DynamicMesh) const;
+	void EvaluateTetWild(UE::Dataflow::FContext& Context, TUniquePtr<FFleshCollection>& InCollection, const UE::Geometry::FDynamicMesh3& DynamicMesh) const;
 };
 
 
diff --git a/Engine/Plugins/Experimental/ChaosFlesh/Source/ChaosFleshDeprecatedNodes/Public/Dataflow/ChaosFleshImportGEO.h b/Engine/Plugins/Experimental/ChaosFlesh/Source/ChaosFleshDeprecatedNodes/Public/Dataflow/ChaosFleshImportGEO.h
--- a/Engine/Plugins/Experimental/ChaosFlesh/Source/ChaosFleshDeprecatedNodes/Public/Dataflow/ChaosFleshImportGEO.h
+++ b/Engine/Plugins/Experimental/ChaosFlesh/Source/ChaosFleshDeprecatedNodes/Public/Dataflow/ChaosFleshImportGEO.h
@@ -55,14 +55,14 @@ public:
 	UPROPERTY(EditAnywhere, Category = "Dataflow", meta = (DataflowOutput))
 	int32 Value = 0;
 
-	FExtractGEOInt(const Dataflow::FNodeParameters& InParam, FGuid InGuid = FGuid::NewGuid())
+	FExtractGEOInt(const UE::Dataflow::FNodeParameters& InParam, FGuid InGuid = FGuid::NewGuid())
 		: FDataflowNode(InParam, InGuid)
 	{
 		RegisterInputConnection(&IntVars);
 		RegisterOutputConnection(&Value);
 	}
 
-	virtual void Evaluate(Dataflow::FContext& Context, const FDataflowOutput* Out) const override;
+	virtual void Evaluate(UE::Dataflow::FContext& Context, const FDataflowOutput* Out) const override;
 };
 
 /** Extract a named integer array from the results of an ImportGEO node. */
@@ -85,14 +85,14 @@ public:
 	UPROPERTY(EditAnywhere, Category = "Dataflow", meta = (DataflowOutput))
 	TArray<int32> Value;
 
-	FExtractGEOIntVector(const Dataflow::FNodeParameters& InParam, FGuid InGuid = FGuid::NewGuid())
+	FExtractGEOIntVector(const UE::Dataflow::FNodeParameters& InParam, FGuid InGuid = FGuid::NewGuid())
 		: FDataflowNode(InParam, InGuid)
 	{
 		RegisterInputConnection(&IntVectorVars);
 		RegisterOutputConnection(&Value);
 	}
 
-	virtual void Evaluate(Dataflow::FContext& Context, const FDataflowOutput* Out) const override;
+	virtual void Evaluate(UE::Dataflow::FContext& Context, const FDataflowOutput* Out) const override;
 };
 
 /** Extract a named float array from the results of an ImportGEO node. */
@@ -115,14 +115,14 @@ public:
 	UPROPERTY(EditAnywhere, Category = "Dataflow", meta = (DataflowOutput))
 	TArray<float> Value;
 
-	FExtractGEOFloatVector(const Dataflow::FNodeParameters& InParam, FGuid InGuid = FGuid::NewGuid())
+	FExtractGEOFloatVector(const UE::Dataflow::FNodeParameters& InParam, FGuid InGuid = FGuid::NewGuid())
 		: FDataflowNode(InParam, InGuid)
 	{
 		RegisterInputConnection(&FloatVectorVars);
 		RegisterOutputConnection(&Value);
 	}
 
-	virtual void Evaluate(Dataflow::FContext& Context, const FDataflowOutput* Out) const override;
+	virtual void Evaluate(UE::Dataflow::FContext& Context, const FDataflowOutput* Out) const override;
 };
 
 /** Import data from GEO file. */
@@ -172,7 +172,7 @@ public:
 	UPROPERTY(EditAnywhere, Category = "Dataflow", meta = (DataflowOutput, DisplayName = "FloatVectorVars"))
 	FGEOMapStringArrayFloat FloatVectorVarsOutput;
 
-	FImportGEO(const Dataflow::FNodeParameters& InParam, FGuid InGuid = FGuid::NewGuid())
+	FImportGEO(const UE::Dataflow::FNodeParameters& InParam, FGuid InGuid = FGuid::NewGuid())
 		: FDataflowNode(InParam, InGuid)
 	{
 		RegisterInputConnection(&Collection);
@@ -182,7 +182,7 @@ public:
 		RegisterOutputConnection(&FloatVectorVarsOutput);
 	}
 
-	virtual void Evaluate(Dataflow::FContext& Context, const FDataflowOutput* Out) const override;
+	virtual void Evaluate(UE::Dataflow::FContext& Context, const FDataflowOutput* Out) const override;
 
 protected:
 	bool ReadGEOFile(const bool PrintStats=true) const;
@@ -195,7 +195,7 @@ protected:
 	mutable TMap<FString, TPair<TArray<std::string>, TArray<int32>>> IndexedStringVars;
 };
 
-namespace Dataflow
+namespace UE::Dataflow
 {
 	void RegisterChaosFleshImportGEONodes();
 }
diff --git a/Engine/Plugins/Experimental/ChaosFlesh/Source/ChaosFleshDeprecatedNodes/Public/Dataflow/ChaosFleshKinematicOriginInsertionInitializationNode.h b/Engine/Plugins/Experimental/ChaosFlesh/Source/ChaosFleshDeprecatedNodes/Public/Dataflow/ChaosFleshKinematicOriginInsertionInitializationNode.h
--- a/Engine/Plugins/Experimental/ChaosFlesh/Source/ChaosFleshDeprecatedNodes/Public/Dataflow/ChaosFleshKinematicOriginInsertionInitializationNode.h
+++ b/Engine/Plugins/Experimental/ChaosFlesh/Source/ChaosFleshDeprecatedNodes/Public/Dataflow/ChaosFleshKinematicOriginInsertionInitializationNode.h
@@ -33,7 +33,7 @@ public:
 		TObjectPtr<USkeletalMesh> BoneSkeletalMeshIn;
 
 
-	FKinematicOriginInsertionInitializationDataflowNode(const Dataflow::FNodeParameters& InParam, FGuid InGuid = FGuid::NewGuid())
+	FKinematicOriginInsertionInitializationDataflowNode(const UE::Dataflow::FNodeParameters& InParam, FGuid InGuid = FGuid::NewGuid())
 		: FDataflowNode(InParam, InGuid)
 	{
 		RegisterInputConnection(&Collection);
@@ -43,6 +43,6 @@ public:
 		RegisterInputConnection(&BoneSkeletalMeshIn);
 	}
 
-	virtual void Evaluate(Dataflow::FContext& Context, const FDataflowOutput* Out) const override;
+	virtual void Evaluate(UE::Dataflow::FContext& Context, const FDataflowOutput* Out) const override;
 };
 
diff --git a/Engine/Plugins/Experimental/ChaosFlesh/Source/ChaosFleshDeprecatedNodes/Public/Dataflow/ChaosFleshKinematicTetrahedralConstraintNode.h b/Engine/Plugins/Experimental/ChaosFlesh/Source/ChaosFleshDeprecatedNodes/Public/Dataflow/ChaosFleshKinematicTetrahedralConstraintNode.h
--- a/Engine/Plugins/Experimental/ChaosFlesh/Source/ChaosFleshDeprecatedNodes/Public/Dataflow/ChaosFleshKinematicTetrahedralConstraintNode.h
+++ b/Engine/Plugins/Experimental/ChaosFlesh/Source/ChaosFleshDeprecatedNodes/Public/Dataflow/ChaosFleshKinematicTetrahedralConstraintNode.h
@@ -30,7 +30,7 @@ public:
 	UPROPERTY(EditAnywhere, Category = "Dataflow")
 	FString ExclusionList = "twist foo";
 
-	FKinematicTetrahedralBindingsDataflowNode(const Dataflow::FNodeParameters& InParam, FGuid InGuid = FGuid::NewGuid())
+	FKinematicTetrahedralBindingsDataflowNode(const UE::Dataflow::FNodeParameters& InParam, FGuid InGuid = FGuid::NewGuid())
 		: FDataflowNode(InParam, InGuid)
 	{
 		RegisterInputConnection(&SkeletalMeshIn);
@@ -38,6 +38,6 @@ public:
 		RegisterOutputConnection(&Collection, &Collection);
 	}
 
-	virtual void Evaluate(Dataflow::FContext& Context, const FDataflowOutput* Out) const override;
+	virtual void Evaluate(UE::Dataflow::FContext& Context, const FDataflowOutput* Out) const override;
 };
 
diff --git a/Engine/Plugins/Experimental/ChaosFlesh/Source/ChaosFleshEngine/Private/ChaosFlesh/ChaosDeformableTetrahedralComponent.cpp b/Engine/Plugins/Experimental/ChaosFlesh/Source/ChaosFleshEngine/Private/ChaosFlesh/ChaosDeformableTetrahedralComponent.cpp
--- a/Engine/Plugins/Experimental/ChaosFlesh/Source/ChaosFleshEngine/Private/ChaosFlesh/ChaosDeformableTetrahedralComponent.cpp
+++ b/Engine/Plugins/Experimental/ChaosFlesh/Source/ChaosFleshEngine/Private/ChaosFlesh/ChaosDeformableTetrahedralComponent.cpp
@@ -182,7 +182,7 @@ UDeformablePhysicsComponent::FDataMapValue UDeformableTetrahedralComponent::NewD
 									if (!Roots.IsEmpty() && ensureMsgf(Roots.Num() == 1, TEXT("Error: Only supports a single root per skeleton.(%s)"), *Skeleton->GetName()))
 									{
 										TArray<FTransform> ComponentLocalPose;
-										Dataflow::Animation::GlobalTransforms(SkeletalMesh->GetRefSkeleton(), ComponentLocalPose);
+										UE::Dataflow::Animation::GlobalTransforms(SkeletalMesh->GetRefSkeleton(), ComponentLocalPose);
 
 										const TArray<FTransform>& ComponentTransforms = SkeletalMeshComponent->GetComponentSpaceTransforms();
 										if (ComponentLocalPose.Num() == ComponentTransforms.Num())
@@ -763,7 +763,7 @@ TArray<FVector> UDeformableTetrahedralComponent::GetSkeletalMeshEmbeddedPosition
 		Format == ChaosDeformableBindingOption::ComponentDelta) // BoneDelta handled below
 	{
 		TArray<FTransform> ComponentPose;
-		Dataflow::Animation::GlobalTransforms(GetRestCollection()->TargetDeformationSkeleton->GetRefSkeleton(), ComponentPose);
+		UE::Dataflow::Animation::GlobalTransforms(GetRestCollection()->TargetDeformationSkeleton->GetRefSkeleton(), ComponentPose);
 		TransformPositions.SetNumUninitialized(ComponentPose.Num());
 		if (TargetDeformationSkeletonOffset.Equals(FTransform::Identity))
 		{
@@ -1053,7 +1053,7 @@ TArray<FVector> UDeformableTetrahedralComponent::GetSkeletalMeshBindingPositions
 				{
 					// Component relative transforms, not world.
 					TArray<FTransform> ComponentPose;
-					Dataflow::Animation::GlobalTransforms(InSkeletalMesh->GetRefSkeleton(), ComponentPose);
+					UE::Dataflow::Animation::GlobalTransforms(InSkeletalMesh->GetRefSkeleton(), ComponentPose);
 
 					TransformPositions.SetNumUninitialized(ComponentPose.Num());
 					for (int32 i = 0; i < ComponentPose.Num(); i++)
diff --git a/Engine/Plugins/Experimental/ChaosFlesh/Source/ChaosFleshNodes/Private/Dataflow/ChaosFleshAddKinematicParticlesNode.cpp b/Engine/Plugins/Experimental/ChaosFlesh/Source/ChaosFleshNodes/Private/Dataflow/ChaosFleshAddKinematicParticlesNode.cpp
--- a/Engine/Plugins/Experimental/ChaosFlesh/Source/ChaosFleshNodes/Private/Dataflow/ChaosFleshAddKinematicParticlesNode.cpp
+++ b/Engine/Plugins/Experimental/ChaosFlesh/Source/ChaosFleshNodes/Private/Dataflow/ChaosFleshAddKinematicParticlesNode.cpp
@@ -13,7 +13,7 @@
 
 //DEFINE_LOG_CATEGORY_STATIC(ChaosFleshAddKinematicParticlesDataflowNodeLog, Log, All);
 
-void FAddKinematicParticlesDataflowNode::Evaluate(Dataflow::FContext& Context, const FDataflowOutput* Out) const
+void FAddKinematicParticlesDataflowNode::Evaluate(UE::Dataflow::FContext& Context, const FDataflowOutput* Out) const
 {
 	if (Out->IsA<DataType>(&Collection) || Out->IsA<TArray<int32>>(&TargetIndicesOut))
 	{
@@ -46,7 +46,7 @@ void FAddKinematicParticlesDataflowNode::Evaluate(Dataflow::FContext& Context, c
 				if (IndexValue != INDEX_NONE)
 				{
 					TArray<FTransform> ComponentPose;
-					Dataflow::Animation::GlobalTransforms(SkeletalMesh->GetRefSkeleton(), ComponentPose);
+					UE::Dataflow::Animation::GlobalTransforms(SkeletalMesh->GetRefSkeleton(), ComponentPose);
 
 					TArray<int32> BranchIndices;
 					if (SkeletalSelectionMode == ESkeletalSeletionMode::Dataflow_SkeletalSelection_Branch)
diff --git a/Engine/Plugins/Experimental/ChaosFlesh/Source/ChaosFleshNodes/Private/Dataflow/ChaosFleshAppendTetrahedralCollectionNode.cpp b/Engine/Plugins/Experimental/ChaosFlesh/Source/ChaosFleshNodes/Private/Dataflow/ChaosFleshAppendTetrahedralCollectionNode.cpp
--- a/Engine/Plugins/Experimental/ChaosFlesh/Source/ChaosFleshNodes/Private/Dataflow/ChaosFleshAppendTetrahedralCollectionNode.cpp
+++ b/Engine/Plugins/Experimental/ChaosFlesh/Source/ChaosFleshNodes/Private/Dataflow/ChaosFleshAppendTetrahedralCollectionNode.cpp
@@ -5,7 +5,7 @@
 #include "ChaosFlesh/FleshCollection.h"
 #include "Dataflow/DataflowNodeFactory.h"
 
-void FAppendTetrahedralCollectionDataflowNode::Evaluate(Dataflow::FContext& Context, const FDataflowOutput* Out) const
+void FAppendTetrahedralCollectionDataflowNode::Evaluate(UE::Dataflow::FContext& Context, const FDataflowOutput* Out) const
 {
 	if (Out->IsA<FManagedArrayCollection>(&Collection1))
 	{
diff --git a/Engine/Plugins/Experimental/ChaosFlesh/Source/ChaosFleshNodes/Private/Dataflow/ChaosFleshAuthorSceneCollisionCandidatesNode.cpp b/Engine/Plugins/Experimental/ChaosFlesh/Source/ChaosFleshNodes/Private/Dataflow/ChaosFleshAuthorSceneCollisionCandidatesNode.cpp
--- a/Engine/Plugins/Experimental/ChaosFlesh/Source/ChaosFleshNodes/Private/Dataflow/ChaosFleshAuthorSceneCollisionCandidatesNode.cpp
+++ b/Engine/Plugins/Experimental/ChaosFlesh/Source/ChaosFleshNodes/Private/Dataflow/ChaosFleshAuthorSceneCollisionCandidatesNode.cpp
@@ -11,7 +11,7 @@
 DEFINE_LOG_CATEGORY(LogAuthorSceneCollisionCandidates);
 
 void
-FAuthorSceneCollisionCandidates::Evaluate(Dataflow::FContext& Context, const FDataflowOutput* Out) const
+FAuthorSceneCollisionCandidates::Evaluate(UE::Dataflow::FContext& Context, const FDataflowOutput* Out) const
 {
 	if (Out->IsA<DataType>(&Collection))
 	{
diff --git a/Engine/Plugins/Experimental/ChaosFlesh/Source/ChaosFleshNodes/Private/Dataflow/ChaosFleshCalculateTetrehedralMetricsNode.cpp b/Engine/Plugins/Experimental/ChaosFlesh/Source/ChaosFleshNodes/Private/Dataflow/ChaosFleshCalculateTetrehedralMetricsNode.cpp
--- a/Engine/Plugins/Experimental/ChaosFlesh/Source/ChaosFleshNodes/Private/Dataflow/ChaosFleshCalculateTetrehedralMetricsNode.cpp
+++ b/Engine/Plugins/Experimental/ChaosFlesh/Source/ChaosFleshNodes/Private/Dataflow/ChaosFleshCalculateTetrehedralMetricsNode.cpp
@@ -14,7 +14,7 @@
 //=============================================================================
 
 void
-FCalculateTetMetrics::Evaluate(Dataflow::FContext& Context, const FDataflowOutput* Out) const
+FCalculateTetMetrics::Evaluate(UE::Dataflow::FContext& Context, const FDataflowOutput* Out) const
 {
 	if (Out->IsA<DataType>(&Collection))
 	{
diff --git a/Engine/Plugins/Experimental/ChaosFlesh/Source/ChaosFleshNodes/Private/Dataflow/ChaosFleshCollisionBodyConstraintNode.cpp b/Engine/Plugins/Experimental/ChaosFlesh/Source/ChaosFleshNodes/Private/Dataflow/ChaosFleshCollisionBodyConstraintNode.cpp
--- a/Engine/Plugins/Experimental/ChaosFlesh/Source/ChaosFleshNodes/Private/Dataflow/ChaosFleshCollisionBodyConstraintNode.cpp
+++ b/Engine/Plugins/Experimental/ChaosFlesh/Source/ChaosFleshNodes/Private/Dataflow/ChaosFleshCollisionBodyConstraintNode.cpp
@@ -11,7 +11,7 @@
 
 #include UE_INLINE_GENERATED_CPP_BY_NAME(ChaosFleshCollisionBodyConstraintNode)
 
-void FKinematicBodySetupInitializationDataflowNode::Evaluate(Dataflow::FContext& Context, const FDataflowOutput* Out) const
+void FKinematicBodySetupInitializationDataflowNode::Evaluate(UE::Dataflow::FContext& Context, const FDataflowOutput* Out) const
 {
 	if (Out->IsA<DataType>(&Collection))
 	{
@@ -25,7 +25,7 @@ void FKinematicBodySetupInitializationDataflowNode::Evaluate(Dataflow::FContext&
 				{
 					TArray<TObjectPtr<USkeletalBodySetup>> SkeletalBodySetups = PhysicsAsset->SkeletalBodySetups;
 					TArray<FTransform> ComponentPose;
-					Dataflow::Animation::GlobalTransforms(SkeletalMesh->GetRefSkeleton(), ComponentPose);
+					UE::Dataflow::Animation::GlobalTransforms(SkeletalMesh->GetRefSkeleton(), ComponentPose);
 					for (const TObjectPtr<USkeletalBodySetup>& BodySetup : SkeletalBodySetups)
 					{	
 						TArray<FKSphylElem> SphylElems = BodySetup->AggGeom.SphylElems;
diff --git a/Engine/Plugins/Experimental/ChaosFlesh/Source/ChaosFleshNodes/Private/Dataflow/ChaosFleshComputeFiberFieldNode.cpp b/Engine/Plugins/Experimental/ChaosFlesh/Source/ChaosFleshNodes/Private/Dataflow/ChaosFleshComputeFiberFieldNode.cpp
--- a/Engine/Plugins/Experimental/ChaosFlesh/Source/ChaosFleshNodes/Private/Dataflow/ChaosFleshComputeFiberFieldNode.cpp
+++ b/Engine/Plugins/Experimental/ChaosFlesh/Source/ChaosFleshNodes/Private/Dataflow/ChaosFleshComputeFiberFieldNode.cpp
@@ -9,7 +9,7 @@
 
 #include UE_INLINE_GENERATED_CPP_BY_NAME(ChaosFleshComputeFiberFieldNode)
 
-void FComputeFiberFieldNode::Evaluate(Dataflow::FContext& Context, const FDataflowOutput* Out) const
+void FComputeFiberFieldNode::Evaluate(UE::Dataflow::FContext& Context, const FDataflowOutput* Out) const
 {
 	if (Out->IsA<FManagedArrayCollection>(&Collection))
 	{
@@ -204,7 +204,7 @@ void FComputeFiberFieldNode::ComputeFiberField(
 		Tolerance);
 }
 
-void FComputeFiberStreamlineNode::Evaluate(Dataflow::FContext& Context, const FDataflowOutput* Out) const
+void FComputeFiberStreamlineNode::Evaluate(UE::Dataflow::FContext& Context, const FDataflowOutput* Out) const
 {
 	//
 	// Gather inputs
diff --git a/Engine/Plugins/Experimental/ChaosFlesh/Source/ChaosFleshNodes/Private/Dataflow/ChaosFleshComputeIslandsNode.cpp b/Engine/Plugins/Experimental/ChaosFlesh/Source/ChaosFleshNodes/Private/Dataflow/ChaosFleshComputeIslandsNode.cpp
--- a/Engine/Plugins/Experimental/ChaosFlesh/Source/ChaosFleshNodes/Private/Dataflow/ChaosFleshComputeIslandsNode.cpp
+++ b/Engine/Plugins/Experimental/ChaosFlesh/Source/ChaosFleshNodes/Private/Dataflow/ChaosFleshComputeIslandsNode.cpp
@@ -7,7 +7,7 @@
 
 #include UE_INLINE_GENERATED_CPP_BY_NAME(ChaosFleshComputeIslandsNode)
 
-void FComputeIslandsNode::Evaluate(Dataflow::FContext& Context, const FDataflowOutput* Out) const
+void FComputeIslandsNode::Evaluate(UE::Dataflow::FContext& Context, const FDataflowOutput* Out) const
 {
 	if (Out->IsA<FManagedArrayCollection>(&Collection))
 	{
diff --git a/Engine/Plugins/Experimental/ChaosFlesh/Source/ChaosFleshNodes/Private/Dataflow/ChaosFleshComputeMuscleActivationNode.cpp b/Engine/Plugins/Experimental/ChaosFlesh/Source/ChaosFleshNodes/Private/Dataflow/ChaosFleshComputeMuscleActivationNode.cpp
--- a/Engine/Plugins/Experimental/ChaosFlesh/Source/ChaosFleshNodes/Private/Dataflow/ChaosFleshComputeMuscleActivationNode.cpp
+++ b/Engine/Plugins/Experimental/ChaosFlesh/Source/ChaosFleshNodes/Private/Dataflow/ChaosFleshComputeMuscleActivationNode.cpp
@@ -8,7 +8,7 @@
 
 #include UE_INLINE_GENERATED_CPP_BY_NAME(ChaosFleshComputeMuscleActivationNode)
 
-void FComputeMuscleActivationDataNode::Evaluate(Dataflow::FContext& Context, const FDataflowOutput* Out) const
+void FComputeMuscleActivationDataNode::Evaluate(UE::Dataflow::FContext& Context, const FDataflowOutput* Out) const
 {
 	if (Out->IsA<FManagedArrayCollection>(&Collection))
 	{
diff --git a/Engine/Plugins/Experimental/ChaosFlesh/Source/ChaosFleshNodes/Private/Dataflow/ChaosFleshCreateTetrahedronNode.cpp b/Engine/Plugins/Experimental/ChaosFlesh/Source/ChaosFleshNodes/Private/Dataflow/ChaosFleshCreateTetrahedronNode.cpp
--- a/Engine/Plugins/Experimental/ChaosFlesh/Source/ChaosFleshNodes/Private/Dataflow/ChaosFleshCreateTetrahedronNode.cpp
+++ b/Engine/Plugins/Experimental/ChaosFlesh/Source/ChaosFleshNodes/Private/Dataflow/ChaosFleshCreateTetrahedronNode.cpp
@@ -37,7 +37,7 @@
 // FCreateTetrahedronDataflowNode
 //=============================================================================
 
-void FCreateTetrahedronDataflowNode::Evaluate(Dataflow::FContext& Context, const FDataflowOutput* Out) const
+void FCreateTetrahedronDataflowNode::Evaluate(UE::Dataflow::FContext& Context, const FDataflowOutput* Out) const
 {
 	if (Out->IsA<DataType>(&Collection))
 	{
@@ -71,7 +71,7 @@ void FCreateTetrahedronDataflowNode::Evaluate(Dataflow::FContext& Context, const
 				return (int32)INDEX_NONE;
 			};
 
-			TArray<int32> ProcessGeometryIndices = Dataflow::GetMatchingMeshIndices(Selection, InSourceCollectionVal.Get());
+			TArray<int32> ProcessGeometryIndices = UE::Dataflow::GetMatchingMeshIndices(Selection, InSourceCollectionVal.Get());
 
 			TArray<TUniquePtr<FFleshCollection>> CollectionBuffer;
 			for (int32 Gdx = 0; Gdx < Source.NumGeometry(); Gdx++)
@@ -190,7 +190,7 @@ void FCreateTetrahedronDataflowNode::Evaluate(Dataflow::FContext& Context, const
 }
 
 void FCreateTetrahedronDataflowNode::EvaluateIsoStuffing(
-	Dataflow::FContext& Context, 
+	UE::Dataflow::FContext& Context, 
 	FFleshCollection& InCollection,
 	const UE::Geometry::FDynamicMesh3& DynamicMesh) const
 {
@@ -219,7 +219,7 @@ void FCreateTetrahedronDataflowNode::EvaluateIsoStuffing(
 		{
 			TArray<FVector> Vertices; Vertices.SetNumUninitialized(IsosurfaceStuffing.Vertices.Num());
 			TArray<FIntVector4> Elements; Elements.SetNumUninitialized(IsosurfaceStuffing.Tets.Num());
-			TArray<FIntVector3> SurfaceElements = Dataflow::GetSurfaceTriangles(IsosurfaceStuffing.Tets, !bDiscardInteriorTriangles);
+			TArray<FIntVector3> SurfaceElements = UE::Dataflow::GetSurfaceTriangles(IsosurfaceStuffing.Tets, !bDiscardInteriorTriangles);
 
 			for (int32 Tdx = 0; Tdx < IsosurfaceStuffing.Tets.Num(); ++Tdx)
 			{
@@ -247,7 +247,7 @@ void FCreateTetrahedronDataflowNode::EvaluateIsoStuffing(
 }
 
 void FCreateTetrahedronDataflowNode::EvaluateTetWild(
-	Dataflow::FContext& Context, 
+	UE::Dataflow::FContext& Context, 
 	FFleshCollection& InCollection,
 	const UE::Geometry::FDynamicMesh3& DynamicMesh) const
 {
@@ -285,7 +285,7 @@ void FCreateTetrahedronDataflowNode::EvaluateTetWild(
 		UE_LOG(LogChaosFlesh, Display,TEXT("Generating tet mesh via TetWild..."));
 		if (UE::Geometry::FTetWild::ComputeTetMesh(Params, Verts, Tris, TetVerts, Tets, &Progress))
 		{
-			TArray<FIntVector3> SurfaceElements = Dataflow::GetSurfaceTriangles(Tets, !bDiscardInteriorTriangles);
+			TArray<FIntVector3> SurfaceElements = UE::Dataflow::GetSurfaceTriangles(Tets, !bDiscardInteriorTriangles);
 			TUniquePtr<FTetrahedralCollection> TetCollection(FTetrahedralCollection::NewTetrahedralCollection(TetVerts, SurfaceElements, Tets));
 			InCollection.AppendGeometry(*TetCollection.Get());
 
diff --git a/Engine/Plugins/Experimental/ChaosFlesh/Source/ChaosFleshNodes/Private/Dataflow/ChaosFleshFleshAssetTerminalNode.cpp b/Engine/Plugins/Experimental/ChaosFlesh/Source/ChaosFleshNodes/Private/Dataflow/ChaosFleshFleshAssetTerminalNode.cpp
--- a/Engine/Plugins/Experimental/ChaosFlesh/Source/ChaosFleshNodes/Private/Dataflow/ChaosFleshFleshAssetTerminalNode.cpp
+++ b/Engine/Plugins/Experimental/ChaosFlesh/Source/ChaosFleshNodes/Private/Dataflow/ChaosFleshFleshAssetTerminalNode.cpp
@@ -5,7 +5,7 @@
 #include UE_INLINE_GENERATED_CPP_BY_NAME(ChaosFleshFleshAssetTerminalNode)
 
 
-void FFleshAssetTerminalDataflowNode::SetAssetValue(TObjectPtr<UObject> Asset, Dataflow::FContext& Context) const
+void FFleshAssetTerminalDataflowNode::SetAssetValue(TObjectPtr<UObject> Asset, UE::Dataflow::FContext& Context) const
 {
 	if (UFleshAsset* InFleshAsset = Cast<UFleshAsset>(Asset.Get()))
 	{
@@ -14,7 +14,7 @@ void FFleshAssetTerminalDataflowNode::SetAssetValue(TObjectPtr<UObject> Asset, D
 	}
 }
 
-void FFleshAssetTerminalDataflowNode::Evaluate(Dataflow::FContext& Context, const FDataflowOutput* Out) const
+void FFleshAssetTerminalDataflowNode::Evaluate(UE::Dataflow::FContext& Context, const FDataflowOutput* Out) const
 {
 	const FManagedArrayCollection& InCollection = GetValue<FManagedArrayCollection>(Context, &Collection);
 	SetValue(Context, InCollection, &Collection);
diff --git a/Engine/Plugins/Experimental/ChaosFlesh/Source/ChaosFleshNodes/Private/Dataflow/ChaosFleshGenerateOriginInsertionNode.cpp b/Engine/Plugins/Experimental/ChaosFlesh/Source/ChaosFleshNodes/Private/Dataflow/ChaosFleshGenerateOriginInsertionNode.cpp
--- a/Engine/Plugins/Experimental/ChaosFlesh/Source/ChaosFleshNodes/Private/Dataflow/ChaosFleshGenerateOriginInsertionNode.cpp
+++ b/Engine/Plugins/Experimental/ChaosFlesh/Source/ChaosFleshNodes/Private/Dataflow/ChaosFleshGenerateOriginInsertionNode.cpp
@@ -8,7 +8,7 @@
 
 #include UE_INLINE_GENERATED_CPP_BY_NAME(ChaosFleshGenerateOriginInsertionNode)
 
-void FGenerateOriginInsertionNode::Evaluate(Dataflow::FContext& Context, const FDataflowOutput* Out) const
+void FGenerateOriginInsertionNode::Evaluate(UE::Dataflow::FContext& Context, const FDataflowOutput* Out) const
 {
 	if (Out->IsA<FManagedArrayCollection>(&Collection))
 	{
diff --git a/Engine/Plugins/Experimental/ChaosFlesh/Source/ChaosFleshNodes/Private/Dataflow/ChaosFleshGenerateSkeletalBindingsNode.cpp b/Engine/Plugins/Experimental/ChaosFlesh/Source/ChaosFleshNodes/Private/Dataflow/ChaosFleshGenerateSkeletalBindingsNode.cpp
--- a/Engine/Plugins/Experimental/ChaosFlesh/Source/ChaosFleshNodes/Private/Dataflow/ChaosFleshGenerateSkeletalBindingsNode.cpp
+++ b/Engine/Plugins/Experimental/ChaosFlesh/Source/ChaosFleshNodes/Private/Dataflow/ChaosFleshGenerateSkeletalBindingsNode.cpp
@@ -19,7 +19,7 @@
 DEFINE_LOG_CATEGORY(LogSkeletalBindings);
 
 
-namespace Dataflow
+namespace UE::Dataflow
 {
 
 	TArray<int32> ChildIndices(TObjectPtr<const USkeletalMesh> SkeletalMesh, int32 StartIndex)
@@ -56,7 +56,7 @@ namespace Dataflow
 }
 
 void
-FGenerateSkeletalBindings::Evaluate(Dataflow::FContext& Context, const FDataflowOutput* Out) const
+FGenerateSkeletalBindings::Evaluate(UE::Dataflow::FContext& Context, const FDataflowOutput* Out) const
 {
 	auto UEVertd = [](FVector3f V) { return FVector3d(V.X, V.Y, V.Z); };
 	auto UEVertf = [](FVector3d V) { return FVector3f(V.X, V.Y, V.Z); };
@@ -106,8 +106,8 @@ FGenerateSkeletalBindings::Evaluate(Dataflow::FContext& Context, const FDataflow
 
 				FName SkeletalMeshName = FName(SkeletalMesh->GetName());
 				TArray<FTransform> ComponentPose;
-				Dataflow::Animation::GlobalTransforms(SkeletalMesh->GetRefSkeleton(), ComponentPose);
-				TArray<int32> Indices = Dataflow::ChildIndices(SkeletalMesh, GetValue<int32>(Context, &BoneIndexIn));
+				UE::Dataflow::Animation::GlobalTransforms(SkeletalMesh->GetRefSkeleton(), ComponentPose);
+				TArray<int32> Indices = UE::Dataflow::ChildIndices(SkeletalMesh, GetValue<int32>(Context, &BoneIndexIn));
 
 				//
 				// Do intersection tests against tets
diff --git a/Engine/Plugins/Experimental/ChaosFlesh/Source/ChaosFleshNodes/Private/Dataflow/ChaosFleshGenerateSurfaceBindingsNode.cpp b/Engine/Plugins/Experimental/ChaosFlesh/Source/ChaosFleshNodes/Private/Dataflow/ChaosFleshGenerateSurfaceBindingsNode.cpp
--- a/Engine/Plugins/Experimental/ChaosFlesh/Source/ChaosFleshNodes/Private/Dataflow/ChaosFleshGenerateSurfaceBindingsNode.cpp
+++ b/Engine/Plugins/Experimental/ChaosFlesh/Source/ChaosFleshNodes/Private/Dataflow/ChaosFleshGenerateSurfaceBindingsNode.cpp
@@ -128,7 +128,7 @@ BuildVertexToVertexAdjacencyBuffer(
 
 
 void
-FGenerateSurfaceBindings::Evaluate(Dataflow::FContext& Context, const FDataflowOutput* Out) const
+FGenerateSurfaceBindings::Evaluate(UE::Dataflow::FContext& Context, const FDataflowOutput* Out) const
 {
 	if (Out->IsA<FManagedArrayCollection>(&Collection))
 	{
diff --git a/Engine/Plugins/Experimental/ChaosFlesh/Source/ChaosFleshNodes/Private/Dataflow/ChaosFleshGetFleshAssetNode.cpp b/Engine/Plugins/Experimental/ChaosFlesh/Source/ChaosFleshNodes/Private/Dataflow/ChaosFleshGetFleshAssetNode.cpp
--- a/Engine/Plugins/Experimental/ChaosFlesh/Source/ChaosFleshNodes/Private/Dataflow/ChaosFleshGetFleshAssetNode.cpp
+++ b/Engine/Plugins/Experimental/ChaosFlesh/Source/ChaosFleshNodes/Private/Dataflow/ChaosFleshGetFleshAssetNode.cpp
@@ -4,7 +4,7 @@
 
 #include UE_INLINE_GENERATED_CPP_BY_NAME(ChaosFleshGetFleshAssetNode)
 
-void FGetFleshAssetDataflowNode::Evaluate(Dataflow::FContext& Context, const FDataflowOutput* Out) const
+void FGetFleshAssetDataflowNode::Evaluate(UE::Dataflow::FContext& Context, const FDataflowOutput* Out) const
 {
 	if (Out->IsA<FManagedArrayCollection>(&Output))
 	{
@@ -14,7 +14,7 @@ void FGetFleshAssetDataflowNode::Evaluate(Dataflow::FContext& Context, const FDa
 		const UFleshAsset* FleshAssetValue = FleshAsset;
 		if (!FleshAssetValue)
 		{
-			if (const Dataflow::FEngineContext* EngineContext = Context.AsType<Dataflow::FEngineContext>())
+			if (const UE::Dataflow::FEngineContext* EngineContext = Context.AsType<UE::Dataflow::FEngineContext>())
 			{
 				FleshAssetValue = Cast<UFleshAsset>(EngineContext->Owner);
 			}
diff --git a/Engine/Plugins/Experimental/ChaosFlesh/Source/ChaosFleshNodes/Private/Dataflow/ChaosFleshGetSurfaceIndicesNode.cpp b/Engine/Plugins/Experimental/ChaosFlesh/Source/ChaosFleshNodes/Private/Dataflow/ChaosFleshGetSurfaceIndicesNode.cpp
--- a/Engine/Plugins/Experimental/ChaosFlesh/Source/ChaosFleshNodes/Private/Dataflow/ChaosFleshGetSurfaceIndicesNode.cpp
+++ b/Engine/Plugins/Experimental/ChaosFlesh/Source/ChaosFleshNodes/Private/Dataflow/ChaosFleshGetSurfaceIndicesNode.cpp
@@ -4,7 +4,7 @@
 
 #include UE_INLINE_GENERATED_CPP_BY_NAME(ChaosFleshGetSurfaceIndicesNode)
 
-void FGetSurfaceIndicesNode::Evaluate(Dataflow::FContext& Context, const FDataflowOutput* Out) const
+void FGetSurfaceIndicesNode::Evaluate(UE::Dataflow::FContext& Context, const FDataflowOutput* Out) const
 {
 	const FManagedArrayCollection InCollection = GetValue<FManagedArrayCollection>(Context, &Collection);
 	if (const TManagedArray<FVector3f>* Vertex = InCollection.FindAttribute<FVector3f>("Vertex", FGeometryCollection::VerticesGroup))
diff --git a/Engine/Plugins/Experimental/ChaosFlesh/Source/ChaosFleshNodes/Private/Dataflow/ChaosFleshIsolateComponentNode.cpp b/Engine/Plugins/Experimental/ChaosFlesh/Source/ChaosFleshNodes/Private/Dataflow/ChaosFleshIsolateComponentNode.cpp
--- a/Engine/Plugins/Experimental/ChaosFlesh/Source/ChaosFleshNodes/Private/Dataflow/ChaosFleshIsolateComponentNode.cpp
+++ b/Engine/Plugins/Experimental/ChaosFlesh/Source/ChaosFleshNodes/Private/Dataflow/ChaosFleshIsolateComponentNode.cpp
@@ -4,7 +4,7 @@
 
 #include UE_INLINE_GENERATED_CPP_BY_NAME(ChaosFleshIsolateComponentNode)
 
-void FIsolateComponentNode::Evaluate(Dataflow::FContext& Context, const FDataflowOutput* Out) const
+void FIsolateComponentNode::Evaluate(UE::Dataflow::FContext& Context, const FDataflowOutput* Out) const
 {
 	if (Out->IsA<FManagedArrayCollection>(&Collection))
 	{	
diff --git a/Engine/Plugins/Experimental/ChaosFlesh/Source/ChaosFleshNodes/Private/Dataflow/ChaosFleshKinematicConstraintNode.cpp b/Engine/Plugins/Experimental/ChaosFlesh/Source/ChaosFleshNodes/Private/Dataflow/ChaosFleshKinematicConstraintNode.cpp
--- a/Engine/Plugins/Experimental/ChaosFlesh/Source/ChaosFleshNodes/Private/Dataflow/ChaosFleshKinematicConstraintNode.cpp
+++ b/Engine/Plugins/Experimental/ChaosFlesh/Source/ChaosFleshNodes/Private/Dataflow/ChaosFleshKinematicConstraintNode.cpp
@@ -9,7 +9,7 @@
 #include UE_INLINE_GENERATED_CPP_BY_NAME(ChaosFleshKinematicConstraintNode)
 
 
-void FKinematicInitializationDataflowNode::Evaluate(Dataflow::FContext& Context, const FDataflowOutput* Out) const
+void FKinematicInitializationDataflowNode::Evaluate(UE::Dataflow::FContext& Context, const FDataflowOutput* Out) const
 {
 	if (Out->IsA<DataType>(&Collection))
 	{
@@ -42,7 +42,7 @@ void FKinematicInitializationDataflowNode::Evaluate(Dataflow::FContext& Context,
 				if (IndexValue != INDEX_NONE)
 				{
 					TArray<FTransform> ComponentPose;
-					Dataflow::Animation::GlobalTransforms(SkeletalMesh->GetRefSkeleton(), ComponentPose);
+					UE::Dataflow::Animation::GlobalTransforms(SkeletalMesh->GetRefSkeleton(), ComponentPose);
 
 					TArray<int32> BranchIndices;
 					if (SkeletalSelectionMode == ESkeletalSeletionMode::Dataflow_SkeletalSelection_Branch)
diff --git a/Engine/Plugins/Experimental/ChaosFlesh/Source/ChaosFleshNodes/Private/Dataflow/ChaosFleshKinematicMuscleAttachmentsNode.cpp b/Engine/Plugins/Experimental/ChaosFlesh/Source/ChaosFleshNodes/Private/Dataflow/ChaosFleshKinematicMuscleAttachmentsNode.cpp
--- a/Engine/Plugins/Experimental/ChaosFlesh/Source/ChaosFleshNodes/Private/Dataflow/ChaosFleshKinematicMuscleAttachmentsNode.cpp
+++ b/Engine/Plugins/Experimental/ChaosFlesh/Source/ChaosFleshNodes/Private/Dataflow/ChaosFleshKinematicMuscleAttachmentsNode.cpp
@@ -11,7 +11,7 @@
 #include UE_INLINE_GENERATED_CPP_BY_NAME(ChaosFleshKinematicMuscleAttachmentsNode)
 
 
-void FKinematicMuscleAttachmentsDataflowNode::Evaluate(Dataflow::FContext& Context, const FDataflowOutput* Out) const
+void FKinematicMuscleAttachmentsDataflowNode::Evaluate(UE::Dataflow::FContext& Context, const FDataflowOutput* Out) const
 {
 	if (Out->IsA<DataType>(&Collection))
 	{
@@ -51,7 +51,7 @@ void FKinematicMuscleAttachmentsDataflowNode::Evaluate(Dataflow::FContext& Conte
 						const FSkinWeightVertexBuffer* SkinWeightVertexBuffer = LODRenderData->GetSkinWeightVertexBuffer();
 						const int32 MaxBoneInfluences = SkinWeightVertexBuffer->GetMaxBoneInfluences();
 						TArray<FTransform> ComponentPose;
-						Dataflow::Animation::GlobalTransforms(SkeletalMesh->GetRefSkeleton(), ComponentPose);
+						UE::Dataflow::Animation::GlobalTransforms(SkeletalMesh->GetRefSkeleton(), ComponentPose);
 						TArray<TArray<int32>> BoneBoundVerts;
 						TArray<TArray<float>> BoneBoundWeights;
 						BoneBoundVerts.SetNum(ComponentPose.Num());
diff --git a/Engine/Plugins/Experimental/ChaosFlesh/Source/ChaosFleshNodes/Private/Dataflow/ChaosFleshKinematicSkeletonConstraintNode.cpp b/Engine/Plugins/Experimental/ChaosFlesh/Source/ChaosFleshNodes/Private/Dataflow/ChaosFleshKinematicSkeletonConstraintNode.cpp
--- a/Engine/Plugins/Experimental/ChaosFlesh/Source/ChaosFleshNodes/Private/Dataflow/ChaosFleshKinematicSkeletonConstraintNode.cpp
+++ b/Engine/Plugins/Experimental/ChaosFlesh/Source/ChaosFleshNodes/Private/Dataflow/ChaosFleshKinematicSkeletonConstraintNode.cpp
@@ -11,7 +11,7 @@
 
 #include UE_INLINE_GENERATED_CPP_BY_NAME(ChaosFleshKinematicSkeletonConstraintNode)
 
-void FKinematicSkeletonConstraintDataflowNode::Evaluate(Dataflow::FContext& Context, const FDataflowOutput* Out) const
+void FKinematicSkeletonConstraintDataflowNode::Evaluate(UE::Dataflow::FContext& Context, const FDataflowOutput* Out) const
 {
 	if (Out->IsA<DataType>(&Collection))
 	{
@@ -27,7 +27,7 @@ void FKinematicSkeletonConstraintDataflowNode::Evaluate(Dataflow::FContext& Cont
 			int32 NumTets = TetCollection.Tetrahedron.Num();
 
 			TArray<FTransform> ComponentPose;
-			Dataflow::Animation::GlobalTransforms(Skeleton->GetReferenceSkeleton(), ComponentPose);
+			UE::Dataflow::Animation::GlobalTransforms(Skeleton->GetReferenceSkeleton(), ComponentPose);
 
 			TArray<FVector3f> Vertex = TetCollection.Vertex.Get().GetConstArray();
 			TetCollection.ComponentSpaceVertices(Vertex);
diff --git a/Engine/Plugins/Experimental/ChaosFlesh/Source/ChaosFleshNodes/Private/Dataflow/ChaosFleshNodesUtility.cpp b/Engine/Plugins/Experimental/ChaosFlesh/Source/ChaosFleshNodes/Private/Dataflow/ChaosFleshNodesUtility.cpp
--- a/Engine/Plugins/Experimental/ChaosFlesh/Source/ChaosFleshNodes/Private/Dataflow/ChaosFleshNodesUtility.cpp
+++ b/Engine/Plugins/Experimental/ChaosFlesh/Source/ChaosFleshNodes/Private/Dataflow/ChaosFleshNodesUtility.cpp
@@ -6,7 +6,7 @@
 #include "Dataflow/DataflowEngine.h"
 #include "GeometryCollection/GeometryCollectionAlgo.h"
 
-namespace Dataflow
+namespace UE::Dataflow
 {
 	// Helper to get the boundary of a tet mesh, useful for debugging / verifying output
 	TArray<FIntVector3> GetSurfaceTriangles(const TArray<FIntVector4>& Tets, const bool bKeepInterior)
diff --git a/Engine/Plugins/Experimental/ChaosFlesh/Source/ChaosFleshNodes/Private/Dataflow/ChaosFleshRadialTetrahedronNode.cpp b/Engine/Plugins/Experimental/ChaosFlesh/Source/ChaosFleshNodes/Private/Dataflow/ChaosFleshRadialTetrahedronNode.cpp
--- a/Engine/Plugins/Experimental/ChaosFlesh/Source/ChaosFleshNodes/Private/Dataflow/ChaosFleshRadialTetrahedronNode.cpp
+++ b/Engine/Plugins/Experimental/ChaosFlesh/Source/ChaosFleshNodes/Private/Dataflow/ChaosFleshRadialTetrahedronNode.cpp
@@ -27,7 +27,7 @@
 //=============================================================================
 
 
-void FRadialTetrahedronDataflowNodes::Evaluate(Dataflow::FContext& Context, const FDataflowOutput* Out) const
+void FRadialTetrahedronDataflowNodes::Evaluate(UE::Dataflow::FContext& Context, const FDataflowOutput* Out) const
 {
 	if (Out->IsA<DataType>(&Collection))
 	{
@@ -35,7 +35,7 @@ void FRadialTetrahedronDataflowNodes::Evaluate(Dataflow::FContext& Context, cons
 		TArray<FIntVector4> TetElements;
 		TArray<FVector> TetVertices;
 		RadialTetMesh(InnerRadius, OuterRadius, Height, RadialSample, AngularSample, VerticalSample, BulgeRatio * (OuterRadius - InnerRadius) / (Chaos::FReal)2., TetElements, TetVertices);
-		TArray<FIntVector3> SurfaceElements = Dataflow::GetSurfaceTriangles(TetElements, !bDiscardInteriorTriangles);
+		TArray<FIntVector3> SurfaceElements = UE::Dataflow::GetSurfaceTriangles(TetElements, !bDiscardInteriorTriangles);
 		TUniquePtr<FTetrahedralCollection> RadialCollection(FTetrahedralCollection::NewTetrahedralCollection(TetVertices, SurfaceElements, TetElements));
 		InCollection->AppendGeometry(*RadialCollection.Get());
 
diff --git a/Engine/Plugins/Experimental/ChaosFlesh/Source/ChaosFleshNodes/Private/Dataflow/ChaosFleshSetFleshBonePositionTargetBindingNode.cpp b/Engine/Plugins/Experimental/ChaosFlesh/Source/ChaosFleshNodes/Private/Dataflow/ChaosFleshSetFleshBonePositionTargetBindingNode.cpp
--- a/Engine/Plugins/Experimental/ChaosFlesh/Source/ChaosFleshNodes/Private/Dataflow/ChaosFleshSetFleshBonePositionTargetBindingNode.cpp
+++ b/Engine/Plugins/Experimental/ChaosFlesh/Source/ChaosFleshNodes/Private/Dataflow/ChaosFleshSetFleshBonePositionTargetBindingNode.cpp
@@ -16,7 +16,7 @@
 
 //DEFINE_LOG_CATEGORY_STATIC(ChaosFleshSetFleshBonePositionTargetBindingNodeLog, Log, All);
 
-void FSetFleshBonePositionTargetBindingDataflowNode::Evaluate(Dataflow::FContext& Context, const FDataflowOutput* Out) const
+void FSetFleshBonePositionTargetBindingDataflowNode::Evaluate(UE::Dataflow::FContext& Context, const FDataflowOutput* Out) const
 {
 	if (Out->IsA<DataType>(&Collection))
 	{
@@ -45,7 +45,7 @@ void FSetFleshBonePositionTargetBindingDataflowNode::Evaluate(Dataflow::FContext
 					const FSkinWeightVertexBuffer* SkinWeightVertexBuffer = LODRenderData->GetSkinWeightVertexBuffer();
 					const int32 MaxBoneInfluences = SkinWeightVertexBuffer->GetMaxBoneInfluences();
 					TArray<FTransform> ComponentPose;
-					Dataflow::Animation::GlobalTransforms(BoneSkeletalMesh->GetRefSkeleton(), ComponentPose);
+					UE::Dataflow::Animation::GlobalTransforms(BoneSkeletalMesh->GetRefSkeleton(), ComponentPose);
 					TArray<TArray<int32>> BoneBoundVerts;
 					TArray<TArray<float>> BoneBoundWeights;
 					BoneBoundVerts.SetNum(ComponentPose.Num());
diff --git a/Engine/Plugins/Experimental/ChaosFlesh/Source/ChaosFleshNodes/Private/Dataflow/ChaosFleshSetFleshDefaultPropertiesNode.cpp b/Engine/Plugins/Experimental/ChaosFlesh/Source/ChaosFleshNodes/Private/Dataflow/ChaosFleshSetFleshDefaultPropertiesNode.cpp
--- a/Engine/Plugins/Experimental/ChaosFlesh/Source/ChaosFleshNodes/Private/Dataflow/ChaosFleshSetFleshDefaultPropertiesNode.cpp
+++ b/Engine/Plugins/Experimental/ChaosFlesh/Source/ChaosFleshNodes/Private/Dataflow/ChaosFleshSetFleshDefaultPropertiesNode.cpp
@@ -23,7 +23,7 @@
 
 template <class T> using MType = FManagedArrayCollection::TManagedType<T>;
 
-void FSetFleshDefaultPropertiesNode::Evaluate(Dataflow::FContext& Context, const FDataflowOutput* Out) const
+void FSetFleshDefaultPropertiesNode::Evaluate(UE::Dataflow::FContext& Context, const FDataflowOutput* Out) const
 {
 	if (Out->IsA<FManagedArrayCollection>(&Collection))
 	{
diff --git a/Engine/Plugins/Experimental/ChaosFlesh/Source/ChaosFleshNodes/Private/Dataflow/ChaosFleshSetVertexTetrahedraPositionTargetBindingNode.cpp b/Engine/Plugins/Experimental/ChaosFlesh/Source/ChaosFleshNodes/Private/Dataflow/ChaosFleshSetVertexTetrahedraPositionTargetBindingNode.cpp
--- a/Engine/Plugins/Experimental/ChaosFlesh/Source/ChaosFleshNodes/Private/Dataflow/ChaosFleshSetVertexTetrahedraPositionTargetBindingNode.cpp
+++ b/Engine/Plugins/Experimental/ChaosFlesh/Source/ChaosFleshNodes/Private/Dataflow/ChaosFleshSetVertexTetrahedraPositionTargetBindingNode.cpp
@@ -12,7 +12,7 @@
 //DEFINE_LOG_CATEGORY_STATIC(ChaosFleshSetVertexTetrahedraPositionTargetBindingNodeLog, Log, All);
 
 
-void FSetVertexTetrahedraPositionTargetBindingDataflowNode::Evaluate(Dataflow::FContext& Context, const FDataflowOutput* Out) const
+void FSetVertexTetrahedraPositionTargetBindingDataflowNode::Evaluate(UE::Dataflow::FContext& Context, const FDataflowOutput* Out) const
 {
 	if (Out->IsA<DataType>(&Collection))
 	{
diff --git a/Engine/Plugins/Experimental/ChaosFlesh/Source/ChaosFleshNodes/Private/Dataflow/ChaosFleshSetVertexTrianglePositionTargetBindingNode.cpp b/Engine/Plugins/Experimental/ChaosFlesh/Source/ChaosFleshNodes/Private/Dataflow/ChaosFleshSetVertexTrianglePositionTargetBindingNode.cpp
--- a/Engine/Plugins/Experimental/ChaosFlesh/Source/ChaosFleshNodes/Private/Dataflow/ChaosFleshSetVertexTrianglePositionTargetBindingNode.cpp
+++ b/Engine/Plugins/Experimental/ChaosFlesh/Source/ChaosFleshNodes/Private/Dataflow/ChaosFleshSetVertexTrianglePositionTargetBindingNode.cpp
@@ -12,7 +12,7 @@
 
 #include UE_INLINE_GENERATED_CPP_BY_NAME(ChaosFleshSetVertexTrianglePositionTargetBindingNode)
 
-void FSetVertexTrianglePositionTargetBindingDataflowNode::Evaluate(Dataflow::FContext& Context, const FDataflowOutput* Out) const
+void FSetVertexTrianglePositionTargetBindingDataflowNode::Evaluate(UE::Dataflow::FContext& Context, const FDataflowOutput* Out) const
 {
 	if (Out->IsA<DataType>(&Collection))
 	{
diff --git a/Engine/Plugins/Experimental/ChaosFlesh/Source/ChaosFleshNodes/Private/Dataflow/ChaosFleshSetVertexVertexPositionTargetBindingNode.cpp b/Engine/Plugins/Experimental/ChaosFlesh/Source/ChaosFleshNodes/Private/Dataflow/ChaosFleshSetVertexVertexPositionTargetBindingNode.cpp
--- a/Engine/Plugins/Experimental/ChaosFlesh/Source/ChaosFleshNodes/Private/Dataflow/ChaosFleshSetVertexVertexPositionTargetBindingNode.cpp
+++ b/Engine/Plugins/Experimental/ChaosFlesh/Source/ChaosFleshNodes/Private/Dataflow/ChaosFleshSetVertexVertexPositionTargetBindingNode.cpp
@@ -11,7 +11,7 @@
 
 
 
-void FSetVertexVertexPositionTargetBindingDataflowNode::Evaluate(Dataflow::FContext& Context, const FDataflowOutput* Out) const
+void FSetVertexVertexPositionTargetBindingDataflowNode::Evaluate(UE::Dataflow::FContext& Context, const FDataflowOutput* Out) const
 {
 	if (Out->IsA<DataType>(&Collection))
 	{
diff --git a/Engine/Plugins/Experimental/ChaosFlesh/Source/ChaosFleshNodes/Private/Dataflow/ChaosFleshSkeletalMeshConstraintNode.cpp b/Engine/Plugins/Experimental/ChaosFlesh/Source/ChaosFleshNodes/Private/Dataflow/ChaosFleshSkeletalMeshConstraintNode.cpp
--- a/Engine/Plugins/Experimental/ChaosFlesh/Source/ChaosFleshNodes/Private/Dataflow/ChaosFleshSkeletalMeshConstraintNode.cpp
+++ b/Engine/Plugins/Experimental/ChaosFlesh/Source/ChaosFleshNodes/Private/Dataflow/ChaosFleshSkeletalMeshConstraintNode.cpp
@@ -11,7 +11,7 @@
 #include UE_INLINE_GENERATED_CPP_BY_NAME(ChaosFleshSkeletalMeshConstraintNode)
 
 
-void FKinematicSkeletalMeshInitializationDataflowNode::Evaluate(Dataflow::FContext& Context, const FDataflowOutput* Out) const
+void FKinematicSkeletalMeshInitializationDataflowNode::Evaluate(UE::Dataflow::FContext& Context, const FDataflowOutput* Out) const
 {
 	if (Out->IsA<DataType>(&Collection) || Out->IsA<TArray<int32>>(&IndicesOut))
 	{
@@ -38,7 +38,7 @@ void FKinematicSkeletalMeshInitializationDataflowNode::Evaluate(Dataflow::FConte
 				const FSkinWeightVertexBuffer* SkinWeightVertexBuffer = LODRenderData->GetSkinWeightVertexBuffer();
 				const int32 MaxBoneInfluences = SkinWeightVertexBuffer->GetMaxBoneInfluences();
 				TArray<FTransform> ComponentPose;
-				Dataflow::Animation::GlobalTransforms(SkeletalMesh->GetRefSkeleton(), ComponentPose);
+				UE::Dataflow::Animation::GlobalTransforms(SkeletalMesh->GetRefSkeleton(), ComponentPose);
 				TArray<TArray<int32>> BoundVerts;
 				TArray<TArray<float>> BoundWeights;
 				BoundVerts.SetNum(ComponentPose.Num());
diff --git a/Engine/Plugins/Experimental/ChaosFlesh/Source/ChaosFleshNodes/Private/Dataflow/ChaosFleshSkinSimulationPropertiesNode.cpp b/Engine/Plugins/Experimental/ChaosFlesh/Source/ChaosFleshNodes/Private/Dataflow/ChaosFleshSkinSimulationPropertiesNode.cpp
--- a/Engine/Plugins/Experimental/ChaosFlesh/Source/ChaosFleshNodes/Private/Dataflow/ChaosFleshSkinSimulationPropertiesNode.cpp
+++ b/Engine/Plugins/Experimental/ChaosFlesh/Source/ChaosFleshNodes/Private/Dataflow/ChaosFleshSkinSimulationPropertiesNode.cpp
@@ -30,7 +30,7 @@
 #include "Spatial/MeshAABBTree3.h"
 
 
-void FSkinSimulationPropertiesDataflowNodes::Evaluate(Dataflow::FContext& Context, const FDataflowOutput* Out) const
+void FSkinSimulationPropertiesDataflowNodes::Evaluate(UE::Dataflow::FContext& Context, const FDataflowOutput* Out) const
 {
 	TUniquePtr<FFleshCollection> InCollection(GetValue<DataType>(Context, &Collection).NewCopy<FFleshCollection>());
 	if (Out->IsA<DataType>(&Collection))
diff --git a/Engine/Plugins/Experimental/ChaosFlesh/Source/ChaosFleshNodes/Private/Dataflow/ChaosFleshTriangleMeshSimulationPropertiesNode.cpp b/Engine/Plugins/Experimental/ChaosFlesh/Source/ChaosFleshNodes/Private/Dataflow/ChaosFleshTriangleMeshSimulationPropertiesNode.cpp
--- a/Engine/Plugins/Experimental/ChaosFlesh/Source/ChaosFleshNodes/Private/Dataflow/ChaosFleshTriangleMeshSimulationPropertiesNode.cpp
+++ b/Engine/Plugins/Experimental/ChaosFlesh/Source/ChaosFleshNodes/Private/Dataflow/ChaosFleshTriangleMeshSimulationPropertiesNode.cpp
@@ -31,7 +31,7 @@
 
 template <class T> using MType = FManagedArrayCollection::TManagedType<T>;
 
-void FTriangleMeshSimulationPropertiesDataflowNodes::Evaluate(Dataflow::FContext& Context, const FDataflowOutput* Out) const
+void FTriangleMeshSimulationPropertiesDataflowNodes::Evaluate(UE::Dataflow::FContext& Context, const FDataflowOutput* Out) const
 {
 	TUniquePtr<FFleshCollection> InCollection(GetValue<DataType>(Context, &Collection).NewCopy<FFleshCollection>());
 	if (Out->IsA<DataType>(&Collection))
diff --git a/Engine/Plugins/Experimental/ChaosFlesh/Source/ChaosFleshNodes/Private/Dataflow/ChaosFleshVertexConstraintNode.cpp b/Engine/Plugins/Experimental/ChaosFlesh/Source/ChaosFleshNodes/Private/Dataflow/ChaosFleshVertexConstraintNode.cpp
--- a/Engine/Plugins/Experimental/ChaosFlesh/Source/ChaosFleshNodes/Private/Dataflow/ChaosFleshVertexConstraintNode.cpp
+++ b/Engine/Plugins/Experimental/ChaosFlesh/Source/ChaosFleshNodes/Private/Dataflow/ChaosFleshVertexConstraintNode.cpp
@@ -6,7 +6,7 @@
 
 #include UE_INLINE_GENERATED_CPP_BY_NAME(ChaosFleshVertexConstraintNode)
 
-void FSetVerticesKinematicDataflowNode::Evaluate(Dataflow::FContext& Context, const FDataflowOutput* Out) const
+void FSetVerticesKinematicDataflowNode::Evaluate(UE::Dataflow::FContext& Context, const FDataflowOutput* Out) const
 {
 	if (Out->IsA<DataType>(&Collection))
 	{
