Dataflow Editor: PDI rendering callbacks for Contruction and Simulation viewports
- Also implement Cloth viewport visualizations for the Dataflow Editor using this system
- Fix UDataflowContextObject having a stale Selected Collection member

diff --git a/Engine/Plugins/ChaosClothAssetEditor/Source/ChaosClothAssetEditor/Private/ChaosClothAsset/ChaosClothAssetEditorModule.cpp b/Engine/Plugins/ChaosClothAssetEditor/Source/ChaosClothAssetEditor/Private/ChaosClothAsset/ChaosClothAssetEditorModule.cpp
--- a/Engine/Plugins/ChaosClothAssetEditor/Source/ChaosClothAssetEditor/Private/ChaosClothAsset/ChaosClothAssetEditorModule.cpp
+++ b/Engine/Plugins/ChaosClothAssetEditor/Source/ChaosClothAssetEditor/Private/ChaosClothAsset/ChaosClothAssetEditorModule.cpp
@@ -3,6 +3,8 @@
 #include "ChaosClothAsset/ChaosClothAssetEditorModule.h"
 #include "ChaosClothAsset/AssetDefinition_ClothAsset.h"
 #include "ChaosClothAsset/ClothAsset.h"
+#include "ChaosClothAsset/ClothDataflowConstructionVisualization.h"
+#include "ChaosClothAsset/ClothDataflowSimulationVisualization.h"
 #include "ChaosClothAsset/ClothEditorCommands.h"
 #include "ChaosClothAsset/ClothEditorMode.h"
 #include "ChaosClothAsset/ClothEditorStyle.h"
@@ -94,6 +96,9 @@ namespace UE::Chaos::ClothAsset
 		// Menus need to be registered in a callback to make sure the system is ready for them.
 		StartupCallbackDelegateHandle = UToolMenus::RegisterStartupCallback(
 			FSimpleMulticastDelegate::FDelegate::CreateRaw(this, &FChaosClothAssetEditorModule::RegisterMenus));
+
+		Dataflow::FDataflowConstructionVisualizationRegistry::GetInstance().RegisterVisualization(MakeUnique<FClothDataflowConstructionVisualization>());
+		Dataflow::FDataflowSimulationVisualizationRegistry::GetInstance().RegisterVisualization(MakeUnique<FClothDataflowSimulationVisualization>());
 	}
 
 	void FChaosClothAssetEditorModule::ShutdownModule()
@@ -103,6 +108,9 @@ namespace UE::Chaos::ClothAsset
 		FEditorModeRegistry::Get().UnregisterMode(UChaosClothAssetEditorMode::EM_ChaosClothAssetEditorModeId);
 
 		UToolMenus::UnRegisterStartupCallback(StartupCallbackDelegateHandle);
+
+		Dataflow::FDataflowConstructionVisualizationRegistry::GetInstance().DeregisterVisualization(FClothDataflowConstructionVisualization::Name);
+		Dataflow::FDataflowSimulationVisualizationRegistry::GetInstance().DeregisterVisualization(FClothDataflowSimulationVisualization::Name);
 	}
 
 } // namespace UE::Chaos::ClothAsset
diff --git a/Engine/Plugins/ChaosClothAssetEditor/Source/ChaosClothAssetEditor/Private/ChaosClothAsset/ClothDataflowConstructionVisualization.cpp b/Engine/Plugins/ChaosClothAssetEditor/Source/ChaosClothAssetEditor/Private/ChaosClothAsset/ClothDataflowConstructionVisualization.cpp
new file mode 100644
--- /dev/null
+++ b/Engine/Plugins/ChaosClothAssetEditor/Source/ChaosClothAssetEditor/Private/ChaosClothAsset/ClothDataflowConstructionVisualization.cpp
@@ -0,0 +1,320 @@
+// Copyright Epic Games, Inc. All Rights Reserved.
+
+#include "ChaosClothAsset/ClothDataflowConstructionVisualization.h"
+#include "ChaosClothAsset/ClothGeometryTools.h"
+#include "ChaosClothAsset/CollectionClothFacade.h"
+#include "Dataflow/DataflowConstructionScene.h"
+#include "Dataflow/DataflowRenderingViewMode.h"
+#include "Dataflow/DataflowConstructionViewportClient.h"
+#include "GeometryCollection/ManagedArrayCollection.h"
+#include "Components/DynamicMeshComponent.h"
+#include "DynamicMeshBuilder.h"
+#include "SceneView.h"
+#include "Materials/Material.h"
+#include "EditorModes.h"
+#include "Framework/MultiBox/MultiBoxBuilder.h"
+
+#define LOCTEXT_NAMESPACE "ChaosClothAssetDataflowConstructionVisualization"
+
+namespace UE::Chaos::ClothAsset
+{
+	namespace Private
+	{
+		extern FLinearColor PseudoRandomColor(int32 NumColorRotations);
+
+		void RenderSeams(const UE::Chaos::ClothAsset::FCollectionClothConstFacade& ClothFacade,
+			const UE::Geometry::FDynamicMesh3* const Mesh,
+			const UE::Dataflow::IDataflowConstructionViewMode* const ViewMode,
+			const TSharedRef<const FManagedArrayCollection>& ClothCollection,
+			bool bCollapseSeams,
+			FPrimitiveDrawInterface* PDI)
+		{
+			const bool bIs2DSim = ViewMode->GetName() == FName("Cloth2DSimView");
+			const bool bIs3DSim = ViewMode->GetName() == FName("Cloth3DSimView");
+
+			if (!bIs2DSim && !bIs3DSim)
+			{
+				return;
+			}
+
+			const float LineThickness = bIs3DSim ? 4.0f : 2.0f;
+			constexpr float PointSize = 4.0f;
+
+			int32 ConnectedSeamIndex = 0;		// Used to generate different colors for each connected seam, if multiple connected seams are found per input seam
+
+			for (int32 SeamIndex = 0; SeamIndex < ClothFacade.GetNumSeams(); ++SeamIndex)
+			{
+				const UE::Chaos::ClothAsset::FCollectionClothSeamConstFacade SeamFacade = ClothFacade.GetSeam(SeamIndex);
+
+				if (bIs2DSim)
+				{
+					// Stitches are given in random order, so first construct paths of connected stitches
+					// Note one SeamFacade can contain multiple disjoint paths
+					TArray<TArray<FIntVector2>> ConnectedSeams;
+					UE::Chaos::ClothAsset::FClothGeometryTools::BuildConnectedSeams2D(ClothCollection, SeamIndex, *Mesh, ConnectedSeams);
+
+					for (const TArray<FIntVector2>& ConnectedSeam : ConnectedSeams)
+					{
+						const FColor SeamColor = UE::Chaos::ClothAsset::Private::PseudoRandomColor(ConnectedSeamIndex++).ToFColor(true);
+
+						// draw connected edge on each side of the seam
+						for (int32 StitchID = 0; StitchID < ConnectedSeam.Num() - 1; ++StitchID)
+						{
+							const FVector3d Stitch0Point0 = Mesh->GetVertex(ConnectedSeam[StitchID][0]);
+							const FVector3d Stitch0Point1 = Mesh->GetVertex(ConnectedSeam[StitchID][1]);
+							const FVector3d Stitch1Point0 = Mesh->GetVertex(ConnectedSeam[StitchID + 1][0]);
+							const FVector3d Stitch1Point1 = Mesh->GetVertex(ConnectedSeam[StitchID + 1][1]);
+
+							PDI->DrawLine(Stitch0Point0, Stitch1Point0, SeamColor, SDPG_World);
+							PDI->DrawLine(Stitch0Point1, Stitch1Point1, SeamColor, SDPG_World);
+						}
+
+						// draw connection between stitch points
+
+						if (bCollapseSeams)
+						{
+							const int32 StitchID = ConnectedSeam.Num() / 2;
+							const int32 StitchVertex0 = ConnectedSeam[StitchID][0];
+							const int32 StitchVertex1 = ConnectedSeam[StitchID][1];
+							const FVector StitchPoint0 = Mesh->GetVertex(StitchVertex0);
+							const FVector StitchPoint1 = Mesh->GetVertex(StitchVertex1);
+
+							PDI->DrawLine(StitchPoint0, StitchPoint1, SeamColor, SDPG_World);
+							PDI->DrawPoint(StitchPoint0, SeamColor, PointSize, SDPG_World);
+							PDI->DrawPoint(StitchPoint1, SeamColor, PointSize, SDPG_World);
+						}
+						else
+						{
+							for (int32 StitchID = 0; StitchID < ConnectedSeam.Num(); ++StitchID)
+							{
+								const int32 StitchVertex0 = ConnectedSeam[StitchID][0];
+								const int32 StitchVertex1 = ConnectedSeam[StitchID][1];
+								const FVector StitchPoint0 = Mesh->GetVertex(StitchVertex0);
+								const FVector StitchPoint1 = Mesh->GetVertex(StitchVertex1);
+
+								PDI->DrawLine(StitchPoint0, StitchPoint1, SeamColor, SDPG_World);
+								PDI->DrawPoint(StitchPoint0, SeamColor, PointSize, SDPG_World);
+								PDI->DrawPoint(StitchPoint1, SeamColor, PointSize, SDPG_World);
+							}
+						}
+					}
+				}
+				else if (bIs3DSim)
+				{
+					const TArray<int32> SeamStitches(SeamFacade.GetSeamStitch3DIndex());
+					const FColor SeamColor = UE::Chaos::ClothAsset::Private::PseudoRandomColor(SeamIndex).ToFColor(true);
+
+					// In 3D we should be able to draw the seam edges in any order, doesn't need to be in connected paths
+					for (int32 StitchIndexI = 0; StitchIndexI < SeamStitches.Num(); ++StitchIndexI)
+					{
+						const int32 StitchIVertex = SeamStitches[StitchIndexI];
+						for (int32 StitchIndexJ = StitchIndexI + 1; StitchIndexJ < SeamStitches.Num(); ++StitchIndexJ)
+						{
+							const int32 StitchJVertex = SeamStitches[StitchIndexJ];
+
+							if (Mesh->FindEdge(StitchIVertex, StitchJVertex) != UE::Geometry::FDynamicMesh3::InvalidID)
+							{
+								const FVector StitchIPoint = Mesh->GetVertex(StitchIVertex);
+								const FVector StitchJPoint = Mesh->GetVertex(StitchJVertex);
+								PDI->DrawLine(StitchIPoint, StitchJPoint, SeamColor, SDPG_World);
+								PDI->DrawPoint(StitchIPoint, SeamColor, PointSize, SDPG_World);
+								PDI->DrawPoint(StitchJPoint, SeamColor, PointSize, SDPG_World);
+							}
+						}
+					}
+				}
+			}
+		}
+
+		void RenderPatterns(const UE::Chaos::ClothAsset::FCollectionClothConstFacade& ClothFacade, const UE::Geometry::FDynamicMesh3* const Mesh, FPrimitiveDrawInterface* PDI)
+		{
+
+			for (int32 PatternIndex = 0; PatternIndex < ClothFacade.GetNumSimPatterns(); ++PatternIndex)
+			{
+				const FCollectionClothSimPatternConstFacade Pattern = ClothFacade.GetSimPattern(PatternIndex);
+				const FLinearColor LinearPatternColor = UE::Chaos::ClothAsset::Private::PseudoRandomColor(PatternIndex);
+				const FColor PatternColor = LinearPatternColor.ToFColorSRGB();
+
+				FDynamicMeshBuilder MeshBuilder(PDI->View->GetFeatureLevel());
+
+				for (int32 PatternFaceID = 0; PatternFaceID < Pattern.GetNumSimFaces(); ++PatternFaceID)
+				{
+					const int32 MeshTriangleID = Pattern.GetSimFacesOffset() + PatternFaceID;
+					const FVector3f TrianglePoint0 = (FVector3f)Mesh->GetTriVertex(MeshTriangleID, 0);
+					const FVector3f TrianglePoint1 = (FVector3f)Mesh->GetTriVertex(MeshTriangleID, 1);
+					const FVector3f TrianglePoint2 = (FVector3f)Mesh->GetTriVertex(MeshTriangleID, 2);
+
+					const int32 VertexIndex0 = MeshBuilder.AddVertex(FDynamicMeshVertex(TrianglePoint0));
+					const int32 VertexIndex1 = MeshBuilder.AddVertex(FDynamicMeshVertex(TrianglePoint1));
+					const int32 VertexIndex2 = MeshBuilder.AddVertex(FDynamicMeshVertex(TrianglePoint2));
+					MeshBuilder.AddTriangle(VertexIndex0, VertexIndex1, VertexIndex2);
+				}
+
+				FDynamicColoredMaterialRenderProxy* PatternColorMaterial = new FDynamicColoredMaterialRenderProxy(GEngine->EmissiveMeshMaterial->GetRenderProxy(), UE::Chaos::ClothAsset::Private::PseudoRandomColor(PatternIndex));
+				PDI->RegisterDynamicResource(PatternColorMaterial);
+
+				MeshBuilder.Draw(PDI, FMatrix::Identity, PatternColorMaterial, SDPG_World, false, false);
+			}
+		}
+	}
+
+	const FName FClothDataflowConstructionVisualization::Name = FName("ClothDataflowConstructionVisualization");
+
+	FName FClothDataflowConstructionVisualization::GetName() const
+	{
+		return Name;
+	}
+
+	void FClothDataflowConstructionVisualization::ExtendViewportShowMenu(const TSharedPtr<FDataflowConstructionViewportClient>& ViewportClient, FMenuBuilder& MenuBuilder)
+	{
+		TSharedPtr<const FManagedArrayCollection> Collection;
+
+		if (const FDataflowConstructionScene* ConstructionScene = StaticCast<FDataflowConstructionScene*>(ViewportClient->GetPreviewScene()))
+		{
+			if (const TObjectPtr<UDataflowBaseContent> DataflowContent = ConstructionScene->GetEditorContent())
+			{
+				Collection = DataflowContent->GetSelectedCollection();
+			}
+		}
+
+		if (!Collection)
+		{
+			return;
+		}
+
+		const UE::Chaos::ClothAsset::FCollectionClothConstFacade ClothFacade(Collection.ToSharedRef());
+		if (!ClothFacade.IsValid())
+		{
+			return;
+		}
+
+		MenuBuilder.BeginSection(TEXT("ClothSeamVisualization"), LOCTEXT("ClothSeamVisualizationSectionName", "Chaos Cloth"));
+		{
+			const FUIAction SeamToggleAction(
+				FExecuteAction::CreateLambda([this, ViewportClient]()
+				{
+					bSeamVisualizationEnabled = !bSeamVisualizationEnabled;
+					if (ViewportClient)
+					{
+						ViewportClient->Invalidate();
+					}
+				}),
+				FCanExecuteAction(),
+				FIsActionChecked::CreateLambda([this]()
+				{
+					return bSeamVisualizationEnabled;
+				}));
+
+			MenuBuilder.AddMenuEntry(LOCTEXT("ClothSeamVisualizationOptionButtonText", "Seams"),
+				LOCTEXT("ClothSeamVisualizationOptionTooltipText", "Seams visualization"),
+				FSlateIcon(),
+				SeamToggleAction,
+				NAME_None,
+				EUserInterfaceActionType::ToggleButton);
+
+			const FUIAction SeamCollapseToggleAction(
+				FExecuteAction::CreateLambda([this, ViewportClient]()
+				{
+					bCollapseSeams = !bCollapseSeams;
+					if (ViewportClient)
+					{
+						ViewportClient->Invalidate();
+					}
+				}),
+				FCanExecuteAction::CreateLambda([this, ViewportClient]()
+				{
+					return bSeamVisualizationEnabled;
+				}),
+				FIsActionChecked::CreateLambda([this]()
+				{
+					return bCollapseSeams;
+				}));
+
+			MenuBuilder.AddMenuEntry(LOCTEXT("CollapseSeamsVisualizationOptionText", "Collapse Seams"),
+				LOCTEXT("CollapseSeamsVisualizationOptionTooltipText", "Collapse seams connection in seams visualization"),
+				FSlateIcon(),
+				SeamCollapseToggleAction,
+				NAME_None,
+				EUserInterfaceActionType::ToggleButton);
+
+
+			const FUIAction PatternColorToggleAction(
+				FExecuteAction::CreateLambda([this, ViewportClient]()
+				{
+					bPatternColorVisualizationEnabled = !bPatternColorVisualizationEnabled;
+					if (ViewportClient)
+					{
+						ViewportClient->Invalidate();
+					}
+				}),
+				FCanExecuteAction(),
+				FIsActionChecked::CreateLambda([this]()
+				{
+					return bPatternColorVisualizationEnabled;
+				}));
+
+			MenuBuilder.AddMenuEntry(LOCTEXT("ClothPatternColorVisualizationEnabledOptionText", "Color Patterns"),
+				LOCTEXT("ClothPatternColorVisualizationEnabledOptionTooltipText", "Draw each cloth patterns with a different color"),
+				FSlateIcon(),
+				PatternColorToggleAction,
+				NAME_None,
+				EUserInterfaceActionType::ToggleButton);
+		}
+
+		MenuBuilder.EndSection();
+	}
+
+	void FClothDataflowConstructionVisualization::Draw(const FDataflowConstructionScene* ConstructionScene, FPrimitiveDrawInterface* PDI)
+	{
+		if (!ConstructionScene)
+		{
+			return;
+		}
+
+		TArray<TObjectPtr<UDynamicMeshComponent>> SceneComponents = ConstructionScene->GetDynamicMeshComponents();
+		if (SceneComponents.Num() != 1)
+		{
+			return;
+		}
+		const UDynamicMeshComponent* const MeshComponent = SceneComponents[0];
+		checkf(MeshComponent, TEXT("Found null mesh component in the ConstructionScene components"));
+
+		const TObjectPtr<UDataflowBaseContent> DataflowContent = ConstructionScene->GetEditorContent();
+		checkf(DataflowContent, TEXT("No DataflowContent in the ConstructionScene"));
+
+		const TSharedPtr<const FManagedArrayCollection> Collection = DataflowContent->GetSelectedCollection();
+
+		if (!Collection)
+		{
+			return;
+		}
+
+		const UE::Chaos::ClothAsset::FCollectionClothConstFacade ClothFacade(Collection.ToSharedRef());
+
+		if (!ClothFacade.IsValid())
+		{
+			return;
+		}
+
+		const UE::Geometry::FDynamicMesh3* Mesh = MeshComponent->GetMesh();
+		checkf(Mesh, TEXT("Unexpected null mesh in a DynamicMeshComponent"));
+
+		const UE::Dataflow::IDataflowConstructionViewMode* const ViewMode = DataflowContent->GetConstructionViewMode();
+		checkf(ViewMode, TEXT("Found null view mode in DataflowContent"));
+
+		if (bPatternColorVisualizationEnabled)
+		{
+			Private::RenderPatterns(ClothFacade, Mesh, PDI);
+		}
+
+		if (bSeamVisualizationEnabled)
+		{
+			Private::RenderSeams(ClothFacade, Mesh, ViewMode, Collection.ToSharedRef(), bCollapseSeams, PDI);
+		}
+		
+	}
+
+}  // End namespace UE::Chaos::ClothAsset
+
+#undef LOCTEXT_NAMESPACE 
+ 
\ No newline at end of file
diff --git a/Engine/Plugins/ChaosClothAssetEditor/Source/ChaosClothAssetEditor/Private/ChaosClothAsset/ClothDataflowConstructionVisualization.h b/Engine/Plugins/ChaosClothAssetEditor/Source/ChaosClothAssetEditor/Private/ChaosClothAsset/ClothDataflowConstructionVisualization.h
new file mode 100644
--- /dev/null
+++ b/Engine/Plugins/ChaosClothAssetEditor/Source/ChaosClothAssetEditor/Private/ChaosClothAsset/ClothDataflowConstructionVisualization.h
@@ -0,0 +1,27 @@
+// Copyright Epic Games, Inc. All Rights Reserved.
+
+#pragma once
+
+#include "Dataflow/DataflowConstructionVisualization.h"
+
+namespace UE::Chaos::ClothAsset
+{
+	class FClothDataflowConstructionVisualization : public UE::Dataflow::IDataflowConstructionVisualization
+	{
+	public:
+		static const FName Name;
+		virtual ~FClothDataflowConstructionVisualization() = default;
+
+	private:
+
+		// IDataflowConstructionVisualization
+		virtual FName GetName() const override;
+		virtual void ExtendViewportShowMenu(const TSharedPtr<FDataflowConstructionViewportClient>& ViewportClient, FMenuBuilder& MenuBuilder) override;
+		virtual void Draw(const FDataflowConstructionScene* ConstructionScene, FPrimitiveDrawInterface* PDI) override;
+
+		bool bSeamVisualizationEnabled = false;
+		bool bCollapseSeams = false;
+		bool bPatternColorVisualizationEnabled = false;
+	};
+
+}  // End namespace UE::Chaos::ClothAsset
diff --git a/Engine/Plugins/ChaosClothAssetEditor/Source/ChaosClothAssetEditor/Private/ChaosClothAsset/ClothDataflowSimulationVisualization.cpp b/Engine/Plugins/ChaosClothAssetEditor/Source/ChaosClothAssetEditor/Private/ChaosClothAsset/ClothDataflowSimulationVisualization.cpp
new file mode 100644
--- /dev/null
+++ b/Engine/Plugins/ChaosClothAssetEditor/Source/ChaosClothAssetEditor/Private/ChaosClothAsset/ClothDataflowSimulationVisualization.cpp
@@ -0,0 +1,67 @@
+// Copyright Epic Games, Inc. All Rights Reserved.
+
+#include "ChaosClothAsset/ClothDataflowSimulationVisualization.h"
+#include "ChaosClothAsset/ClothComponent.h"
+#include "Dataflow/DataflowEditorToolkit.h"
+#include "Dataflow/DataflowSimulationViewportClient.h"
+#include "Dataflow/DataflowSimulationScene.h"
+
+#define LOCTEXT_NAMESPACE "ChaosClothAssetDataflowSimulationVisualization"
+
+namespace UE::Chaos::ClothAsset
+{
+	const FName FClothDataflowSimulationVisualization::Name = FName("ClothDataflowSimulationVisualization");
+
+	FName FClothDataflowSimulationVisualization::GetName() const
+	{
+		return Name;
+	}
+
+	void FClothDataflowSimulationVisualization::ExtendSimulationVisualizationMenu(const TSharedPtr<FDataflowSimulationViewportClient>& ViewportClient, FMenuBuilder& MenuBuilder)
+	{
+		if (ViewportClient)
+		{
+			if (TSharedPtr<FDataflowEditorToolkit> Toolkit = ViewportClient->GetDataflowEditorToolkit().Pin())
+			{
+				if (const TSharedPtr<FDataflowSimulationScene>& SimulationScene = Toolkit->GetSimulationScene())
+				{
+					if (GetClothComponent(SimulationScene.Get()))
+					{
+						ClothEditorSimulationVisualization.ExtendViewportShowMenu(MenuBuilder, ViewportClient.ToSharedRef());
+					}
+				}
+			}
+		}
+	}
+
+	const UChaosClothComponent* FClothDataflowSimulationVisualization::GetClothComponent(const FDataflowSimulationScene* SimulationScene) const
+	{
+		if (SimulationScene)
+		{
+			if (const TObjectPtr<AActor> PreviewActor = SimulationScene->GetPreviewActor())
+			{
+				return PreviewActor->GetComponentByClass<UChaosClothComponent>();
+			}
+		}
+		return nullptr;
+	}
+
+	void FClothDataflowSimulationVisualization::Draw(const FDataflowSimulationScene* SimulationScene, FPrimitiveDrawInterface* PDI)
+	{
+		ClothEditorSimulationVisualization.DebugDrawSimulation(GetClothComponent(SimulationScene), PDI);
+	}
+
+	void FClothDataflowSimulationVisualization::DrawCanvas(const FDataflowSimulationScene* SimulationScene, FCanvas* Canvas, const FSceneView* SceneView)
+	{
+		ClothEditorSimulationVisualization.DebugDrawSimulationTexts(GetClothComponent(SimulationScene), Canvas, SceneView);
+	}
+
+	FText FClothDataflowSimulationVisualization::GetDisplayString(const FDataflowSimulationScene* SimulationScene) const
+	{
+		return ClothEditorSimulationVisualization.GetDisplayString(GetClothComponent(SimulationScene));
+	}
+
+}  // End namespace UE::Chaos::ClothAsset
+
+#undef LOCTEXT_NAMESPACE 
+ 
\ No newline at end of file
diff --git a/Engine/Plugins/ChaosClothAssetEditor/Source/ChaosClothAssetEditor/Private/ChaosClothAsset/ClothDataflowSimulationVisualization.h b/Engine/Plugins/ChaosClothAssetEditor/Source/ChaosClothAssetEditor/Private/ChaosClothAsset/ClothDataflowSimulationVisualization.h
new file mode 100644
--- /dev/null
+++ b/Engine/Plugins/ChaosClothAssetEditor/Source/ChaosClothAssetEditor/Private/ChaosClothAsset/ClothDataflowSimulationVisualization.h
@@ -0,0 +1,30 @@
+// Copyright Epic Games, Inc. All Rights Reserved.
+
+#pragma once
+
+#include "ChaosClothAsset/ClothEditorSimulationVisualization.h"
+#include "Dataflow/DataflowSimulationVisualization.h"
+
+namespace UE::Chaos::ClothAsset
+{
+	class FClothDataflowSimulationVisualization : public Dataflow::IDataflowSimulationVisualization
+	{
+	public:
+		static const FName Name;
+		virtual ~FClothDataflowSimulationVisualization() = default;
+
+	private:
+
+		// IDataflowSimulationVisualization
+		virtual FName GetName() const override;
+		virtual void ExtendSimulationVisualizationMenu(const TSharedPtr<FDataflowSimulationViewportClient>& ViewportClient, FMenuBuilder& MenuBuilder) override;
+		virtual void Draw(const FDataflowSimulationScene* SimulationScene, FPrimitiveDrawInterface* PDI) override;
+		virtual void DrawCanvas(const FDataflowSimulationScene* SimulationScene, FCanvas* Canvas, const FSceneView* SceneView) override;
+		virtual FText GetDisplayString(const FDataflowSimulationScene* SimulationScene) const override;
+
+		const UChaosClothComponent* GetClothComponent(const FDataflowSimulationScene* SimulationScene) const;
+
+		FClothEditorSimulationVisualization ClothEditorSimulationVisualization;
+	};
+
+}  // End namespace UE::Chaos::ClothAsset
diff --git a/Engine/Plugins/ChaosClothAssetEditor/Source/ChaosClothAssetEditor/Private/ChaosClothAsset/ClothEditorSimulationVisualization.cpp b/Engine/Plugins/ChaosClothAssetEditor/Source/ChaosClothAssetEditor/Private/ChaosClothAsset/ClothEditorSimulationVisualization.cpp
--- a/Engine/Plugins/ChaosClothAssetEditor/Source/ChaosClothAssetEditor/Private/ChaosClothAsset/ClothEditorSimulationVisualization.cpp
+++ b/Engine/Plugins/ChaosClothAssetEditor/Source/ChaosClothAssetEditor/Private/ChaosClothAsset/ClothEditorSimulationVisualization.cpp
@@ -6,6 +6,9 @@
 #include "ChaosClothAsset/ClothSimulationProxy.h"
 #include "ChaosClothAsset/ClothComponent.h"
 #include "ChaosCloth/ChaosClothVisualization.h"
+#include "Dataflow/DataflowEditorToolkit.h"
+#include "Dataflow/DataflowSimulationScene.h"
+#include "Dataflow/DataflowSimulationViewportClient.h"
 #include "Framework/MultiBox/MultiBoxBuilder.h"
 #include "Rendering/SkeletalMeshRenderData.h"
 #include "Delegates/Delegate.h"
@@ -95,7 +98,7 @@ DECLARE_DELEGATE_FourParams(FClothVisualizationDebugDrawTexts, const FClothEdito
 
 DECLARE_DELEGATE_RetVal_TwoParams(FText, FLocalDebugDisplayString, const FClothEditorSimulationVisualization&, const FClothSimulationProxy& );
 
-DECLARE_DELEGATE_ThreeParams(FAdditionalMenus, FClothEditorSimulationVisualization&, FMenuBuilder&, TSharedRef<FChaosClothAssetEditor3DViewportClient>);
+DECLARE_DELEGATE_TwoParams(FAdditionalMenus, FClothEditorSimulationVisualization&, FMenuBuilder&);
 
 struct FVisualizationOption
 {
@@ -167,7 +170,7 @@ const FVisualizationOption FVisualizationOption::OptionData[] =
 		}),
 		LOCTEXT("ChaosVisName_SimulationStatistics", "Simulation Statistics"),
 		LOCTEXT("ChaosVisName_SimulationStatistics_Tooltip", "Displays simulation statistics"),
-		/*bDisablesSimulation =*/false, /*bHidesClothSections=*/true, /*bDefaultFlagValue=*/true),
+		/*bDisablesSimulation =*/false, /*bHidesClothSections=*/false, /*bDefaultFlagValue=*/true),
 	FVisualizationOption(
 		FClothVisualizationDebugDraw::CreateLambda([](const FClothEditorSimulationVisualization&, const ::Chaos::FClothVisualizationNoGC& Visualization, FPrimitiveDrawInterface* PDI)
 		{
@@ -216,7 +219,7 @@ const FVisualizationOption FVisualizationOption::OptionData[] =
 		LOCTEXT("ChaosVisName_PointNormals", "Physical Mesh Normals"), 
 		LOCTEXT("ChaosVisName_PointNormals_ToolTip", "Draws the current point normals for the simulation mesh"),
 		/*bDisablesSimulation =*/false, /*bHidesClothSections=*/false, /*bDefaultFlagValue=*/false, 
-		FAdditionalMenus::CreateLambda([](FClothEditorSimulationVisualization& EditorVisualization, FMenuBuilder& MenuBuilder, TSharedRef<FChaosClothAssetEditor3DViewportClient> ViewportClient)
+		FAdditionalMenus::CreateLambda([](FClothEditorSimulationVisualization& EditorVisualization, FMenuBuilder& MenuBuilder)
 		{
 			EditorVisualization.ExtendViewportShowMenuPointNormalsLength(MenuBuilder);
 		})),
@@ -235,7 +238,7 @@ const FVisualizationOption FVisualizationOption::OptionData[] =
 		LOCTEXT("ChaosVisName_AnimNormals", "Animated Mesh Normals"), 
 		LOCTEXT("ChaosVisName_AnimNormals_ToolTip", "Draws the current point normals for the animated mesh"),
 		/*bDisablesSimulation =*/false, /*bHidesClothSections=*/false, /*bDefaultFlagValue=*/false,
-		FAdditionalMenus::CreateLambda([](FClothEditorSimulationVisualization& EditorVisualization, FMenuBuilder& MenuBuilder, TSharedRef<FChaosClothAssetEditor3DViewportClient> ViewportClient)
+		FAdditionalMenus::CreateLambda([](FClothEditorSimulationVisualization& EditorVisualization, FMenuBuilder& MenuBuilder)
 		{
 			EditorVisualization.ExtendViewportShowMenuAnimatedNormalsLength(MenuBuilder);
 		})),
@@ -311,7 +314,7 @@ const FVisualizationOption FVisualizationOption::OptionData[] =
 		LOCTEXT("ChaosVisName_WindAndPressureForces", "Wind Aerodynamic And Pressure Forces"), 
 		LOCTEXT("ChaosVisName_WindAndPressure_Tooltip", "Draws the Wind drag and lift and pressure forces"),
 		/*bDisablesSimulation =*/false, /*bHidesClothSections=*/false, /*bDefaultFlagValue=*/false,
-		FAdditionalMenus::CreateLambda([](FClothEditorSimulationVisualization& EditorVisualization, FMenuBuilder& MenuBuilder, TSharedRef<FChaosClothAssetEditor3DViewportClient> ViewportClient)
+		FAdditionalMenus::CreateLambda([](FClothEditorSimulationVisualization& EditorVisualization, FMenuBuilder& MenuBuilder)
 		{
 			EditorVisualization.ExtendViewportShowMenuAerodynamicsLengthScale(MenuBuilder);
 		})),
@@ -395,9 +398,9 @@ const FVisualizationOption FVisualizationOption::OptionData[] =
 		LOCTEXT("ChaosVisName_DrawWeightMap", "Weight Map"), 
 		LOCTEXT("ChaosVisName_DrawWeightMap_ToolTip", "Draw the weight map for the simulation mesh. You can control the name of the map to be visualized by setting the p.ChaosClothVisualization.WeightMapName console variable."), 
 		/*bDisablesSimulation =*/false, /*bHidesClothSections=*/true, /*bDefaultFlagValue=*/false,
-		FAdditionalMenus::CreateLambda([](FClothEditorSimulationVisualization& EditorVisualization, FMenuBuilder& MenuBuilder, TSharedRef<FChaosClothAssetEditor3DViewportClient> ViewportClient)
+		FAdditionalMenus::CreateLambda([](FClothEditorSimulationVisualization& EditorVisualization, FMenuBuilder& MenuBuilder)
 		{
-			EditorVisualization.ExtendViewportShowMenuWeightMapSelector(MenuBuilder, ViewportClient);
+			EditorVisualization.ExtendViewportShowMenuWeightMapSelector(MenuBuilder);
 		})),
 	FVisualizationOption(
 		FClothVisualizationDebugDraw::CreateLambda([](const FClothEditorSimulationVisualization&, const ::Chaos::FClothVisualizationNoGC& Visualization, FPrimitiveDrawInterface* PDI)
@@ -423,18 +426,19 @@ FClothEditorSimulationVisualization::FClothEditorSimulationVisualization()
 	}
 }
 
-void FClothEditorSimulationVisualization::ExtendViewportShowMenu(FMenuBuilder& MenuBuilder, TSharedRef<FChaosClothAssetEditor3DViewportClient> ViewportClient)
+
+void FClothEditorSimulationVisualization::ExtendViewportShowMenu(FMenuBuilder& MenuBuilder, const TFunction<UChaosClothComponent*()>& GetClothComponentFunc)
 {
-	MenuBuilder.BeginSection(TEXT("ChaosSimulation_Visualizations"), LOCTEXT("VisualizationSection", "Visualizations"));
+	MenuBuilder.BeginSection(TEXT("ChaosSimulation_Visualizations"), LOCTEXT("ClothVisualizationSection", "Chaos Cloth Visualization"));
 	{
 		for (uint32 OptionIndex = 0; OptionIndex < Private::FVisualizationOption::Count; ++OptionIndex)
 		{
 			// Handler for visualization entry being clicked
-			const FExecuteAction ExecuteAction = FExecuteAction::CreateLambda([this, OptionIndex, ViewportClient]()
+			const FExecuteAction ExecuteAction = FExecuteAction::CreateLambda([this, OptionIndex, GetClothComponentFunc]()
 			{
 				Flags[OptionIndex] = !Flags[OptionIndex];
 
-				if (UChaosClothComponent* const ClothComponent = ViewportClient->GetPreviewClothComponent())
+				if (UChaosClothComponent* const ClothComponent = GetClothComponentFunc())
 				{
 					// If we need to toggle the disabled or visibility states, handle it
 					// Disable simulation
@@ -462,13 +466,55 @@ void FClothEditorSimulationVisualization::ExtendViewportShowMenu(FMenuBuilder& M
 
 			// Add menu entry
 			MenuBuilder.AddMenuEntry(Private::FVisualizationOption::OptionData[OptionIndex].DisplayName, Private::FVisualizationOption::OptionData[OptionIndex].ToolTip, FSlateIcon(), Action, NAME_None, EUserInterfaceActionType::ToggleButton);
-			Private::FVisualizationOption::OptionData[OptionIndex].AdditionalMenus.ExecuteIfBound(*this, MenuBuilder, ViewportClient);
+			Private::FVisualizationOption::OptionData[OptionIndex].AdditionalMenus.ExecuteIfBound(*this, MenuBuilder);
 		}
 	}
 	MenuBuilder.EndSection();
 }
 
-void FClothEditorSimulationVisualization::ExtendViewportShowMenuWeightMapSelector(FMenuBuilder& MenuBuilder, TSharedRef<FChaosClothAssetEditor3DViewportClient> ViewportClient)
+
+void FClothEditorSimulationVisualization::ExtendViewportShowMenu(FMenuBuilder& MenuBuilder, const TSharedRef<FChaosClothAssetEditor3DViewportClient>& ViewportClient)
+{
+	TWeakPtr<FChaosClothAssetEditor3DViewportClient> ViewportClientWeakPtr(ViewportClient);
+
+	ExtendViewportShowMenu(MenuBuilder, [ViewportClientWeakPtr]() -> UChaosClothComponent*
+	{
+		if (TSharedPtr<FChaosClothAssetEditor3DViewportClient> ViewportClient = ViewportClientWeakPtr.Pin())
+		{
+			return ViewportClient->GetPreviewClothComponent();
+		}
+
+		return nullptr;
+	});
+}
+
+
+void FClothEditorSimulationVisualization::ExtendViewportShowMenu(FMenuBuilder& MenuBuilder, const TSharedRef<FDataflowSimulationViewportClient>& ViewportClient)
+{
+	TWeakPtr<FDataflowSimulationViewportClient> ViewportClientWeakPtr(ViewportClient);
+
+	ExtendViewportShowMenu(MenuBuilder, [ViewportClientWeakPtr]() -> UChaosClothComponent*
+	{
+		if (TSharedPtr<FDataflowSimulationViewportClient> ViewportClient = ViewportClientWeakPtr.Pin())
+		{
+			if (TSharedPtr<FDataflowEditorToolkit> Toolkit = ViewportClient->GetDataflowEditorToolkit().Pin())
+			{
+				if (const TSharedPtr<FDataflowSimulationScene>& SimulationScene = Toolkit->GetSimulationScene())
+				{
+					if (const TObjectPtr<AActor> PreviewActor = SimulationScene->GetPreviewActor())
+					{
+						return PreviewActor->GetComponentByClass<UChaosClothComponent>();
+					}
+				}
+			}
+		}
+
+		return nullptr;
+	});
+}
+
+
+void FClothEditorSimulationVisualization::ExtendViewportShowMenuWeightMapSelector(FMenuBuilder& MenuBuilder)
 {
 	WeightMapSelector =
 		SNew(STextComboBox)
diff --git a/Engine/Plugins/ChaosClothAssetEditor/Source/ChaosClothAssetEditor/Private/ChaosClothAsset/ClothEditorSimulationVisualization.h b/Engine/Plugins/ChaosClothAssetEditor/Source/ChaosClothAssetEditor/Private/ChaosClothAsset/ClothEditorSimulationVisualization.h
--- a/Engine/Plugins/ChaosClothAssetEditor/Source/ChaosClothAssetEditor/Private/ChaosClothAsset/ClothEditorSimulationVisualization.h
+++ b/Engine/Plugins/ChaosClothAssetEditor/Source/ChaosClothAssetEditor/Private/ChaosClothAsset/ClothEditorSimulationVisualization.h
@@ -10,6 +10,7 @@ class FPrimitiveDrawInterface;
 class FCanvas;
 class FSceneView;
 class STextComboBox;
+class FDataflowSimulationViewportClient;
 
 namespace ESelectInfo
 {
@@ -25,8 +26,10 @@ class FClothEditorSimulationVisualization
 {
 public:
 	FClothEditorSimulationVisualization();
+	   
+	void ExtendViewportShowMenu(FMenuBuilder& MenuBuilder, const TSharedRef<FChaosClothAssetEditor3DViewportClient>& ViewportClient);
+	void ExtendViewportShowMenu(FMenuBuilder& MenuBuilder, const TSharedRef<FDataflowSimulationViewportClient>& ViewportClient);
 
-	void ExtendViewportShowMenu(FMenuBuilder& MenuBuilder, TSharedRef<FChaosClothAssetEditor3DViewportClient> ViewportClient);
 	void DebugDrawSimulation(const UChaosClothComponent* ClothComponent, FPrimitiveDrawInterface* PDI);
 	void DebugDrawSimulationTexts(const UChaosClothComponent* ClothComponent, FCanvas* Canvas, const FSceneView* SceneView);
 	FText GetDisplayString(const UChaosClothComponent* ClothComponent) const;
@@ -34,7 +37,8 @@ public:
 
 	// WeightMaps 
 	const FString* GetCurrentlySelectedWeightMap() const { return CurrentlySelectedWeightMap.Get(); }
-	void ExtendViewportShowMenuWeightMapSelector(FMenuBuilder& MenuBuilder, TSharedRef<FChaosClothAssetEditor3DViewportClient> ViewportClient);
+	void ExtendViewportShowMenuWeightMapSelector(FMenuBuilder& MenuBuilder);
+
 	// Normals
 	void ExtendViewportShowMenuPointNormalsLength(FMenuBuilder& MenuBuilder)
 	{
@@ -51,7 +55,11 @@ public:
 	float GetPointNormalLength() const { return PointNormalLength; }
 	float GetAnimatedNormalLength() const { return AnimatedNormalLength; }
 	float GetAerodynamicsLengthScale() const { return AerodynamicsLengthScale; }
+
 private:
+
+	void ExtendViewportShowMenu(FMenuBuilder& MenuBuilder, const TFunction<UChaosClothComponent*()>& GetClothComponentFunc);
+
 	/** Return whether or not - given the current enabled options - the simulation should be disabled. */
 	bool ShouldDisableSimulation() const;
 	/** Show/hide all cloth sections for the specified mesh compoment. */
diff --git a/Engine/Plugins/Experimental/Dataflow/Source/DataflowEditor/Private/Dataflow/DataflowConstructionViewportClient.cpp b/Engine/Plugins/Experimental/Dataflow/Source/DataflowEditor/Private/Dataflow/DataflowConstructionViewportClient.cpp
--- a/Engine/Plugins/Experimental/Dataflow/Source/DataflowEditor/Private/Dataflow/DataflowConstructionViewportClient.cpp
+++ b/Engine/Plugins/Experimental/Dataflow/Source/DataflowEditor/Private/Dataflow/DataflowConstructionViewportClient.cpp
@@ -9,6 +9,7 @@
 #include "Dataflow/DataflowEngineSceneHitProxies.h"
 #include "Dataflow/DataflowGraphEditor.h"
 #include "Dataflow/DataflowConstructionScene.h"
+#include "Dataflow/DataflowConstructionVisualization.h"
 #include "Dataflow/DataflowEditorPreviewSceneBase.h"
 #include "Dataflow/DataflowRenderingViewMode.h"
 #include "EditorModeManager.h"
@@ -246,4 +247,26 @@ void FDataflowConstructionViewportClient::AddReferencedObjects(FReferenceCollect
 	Collector.AddReferencedObject(BehaviorSet);
 }
 
+void FDataflowConstructionViewportClient::Draw(const FSceneView* View, FPrimitiveDrawInterface* PDI)
+{
+	FEditorViewportClient::Draw(View, PDI);
+
+	using namespace UE::Dataflow;
+	for (const TPair<FName, TUniquePtr<IDataflowConstructionVisualization>>& Visualization : FDataflowConstructionVisualizationRegistry::GetInstance().GetVisualizations())
+	{
+		Visualization.Value->Draw(static_cast<FDataflowConstructionScene*>(PreviewScene), PDI);
+	}
+}
+
+void FDataflowConstructionViewportClient::DrawCanvas(FViewport& InViewport, FSceneView& View, FCanvas& Canvas)
+{
+	FEditorViewportClient::DrawCanvas(InViewport, View, Canvas);
+
+	using namespace UE::Dataflow;
+	for (const TPair<FName, TUniquePtr<IDataflowConstructionVisualization>>& Visualization : FDataflowConstructionVisualizationRegistry::GetInstance().GetVisualizations())
+	{
+		Visualization.Value->DrawCanvas(static_cast<FDataflowConstructionScene*>(PreviewScene), &Canvas, &View);
+	}
+}
+
 #undef LOCTEXT_NAMESPACE 
diff --git a/Engine/Plugins/Experimental/Dataflow/Source/DataflowEditor/Private/Dataflow/DataflowConstructionViewportToolbar.cpp b/Engine/Plugins/Experimental/Dataflow/Source/DataflowEditor/Private/Dataflow/DataflowConstructionViewportToolbar.cpp
--- a/Engine/Plugins/Experimental/Dataflow/Source/DataflowEditor/Private/Dataflow/DataflowConstructionViewportToolbar.cpp
+++ b/Engine/Plugins/Experimental/Dataflow/Source/DataflowEditor/Private/Dataflow/DataflowConstructionViewportToolbar.cpp
@@ -3,6 +3,8 @@
 #include "Dataflow/DataflowConstructionViewportToolbar.h"
 
 #include "Dataflow/DataflowConstructionViewport.h"
+#include "Dataflow/DataflowConstructionViewportClient.h"
+#include "Dataflow/DataflowConstructionVisualization.h"
 #include "Dataflow/DataflowEditorCommands.h"
 #include "Dataflow/DataflowEditorMode.h"
 #include "Dataflow/DataflowRenderingViewMode.h"
@@ -71,12 +73,24 @@ void SDataflowConstructionViewportSelectionToolBar::Construct(const FArguments&
 
 	// Display (Lit, Unlit, Wireframe, etc.)
 	MainBoxPtr->AddSlot()
+		.AutoWidth()
 		.Padding(ToolbarSlotPadding)
-		.HAlign(HAlign_Left)
 		[
 			MakeDisplayToolBar(InArgs._Extenders)
 		];
 
+	// Show menu
+	MainBoxPtr->AddSlot()
+		.AutoWidth()
+		.Padding(ToolbarSlotPadding)
+		[
+			SNew(SEditorViewportToolbarMenu)
+				.Label(LOCTEXT("ShowMenuTitle", "Show"))
+				.Cursor(EMouseCursor::Default)
+				.ParentToolBar(SharedThis(this))
+				.OnGetMenuContent(this, &SDataflowConstructionViewportSelectionToolBar::GenerateShowMenu)
+		];
+
 	// View mode (Sim2D/Sim3D/Render)
 	MainBoxPtr->AddSlot()
 		.Padding(ToolbarSlotPadding)
@@ -88,7 +102,25 @@ void SDataflowConstructionViewportSelectionToolBar::Construct(const FArguments&
 	// See SCommonEditorViewportToolbarBase::Construct for more possible menus to add
 }
 
+TSharedRef<SWidget> SDataflowConstructionViewportSelectionToolBar::GenerateShowMenu() const
+{
+	FMenuBuilder MenuBuilder(false, CommandList);
+
+	using namespace UE::Dataflow;
 
+	if (const TSharedPtr<SDataflowConstructionViewport> PinnedEditorViewport = EditorViewport.Pin())
+	{
+		if (const TSharedPtr<FEditorViewportClient> ViewportClient = PinnedEditorViewport->GetViewportClient())
+		{
+			for (const TPair<FName, TUniquePtr<IDataflowConstructionVisualization>>& Visualization : FDataflowConstructionVisualizationRegistry::GetInstance().GetVisualizations())
+			{
+				Visualization.Value->ExtendViewportShowMenu(StaticCastSharedPtr<FDataflowConstructionViewportClient>(ViewportClient), MenuBuilder);
+			}
+		}
+	}
+
+	return MenuBuilder.MakeWidget();
+}
 
 TSharedRef<SWidget> SDataflowConstructionViewportSelectionToolBar::MakeToolBar(const TSharedPtr<FExtender> InExtenders)
 {
diff --git a/Engine/Plugins/Experimental/Dataflow/Source/DataflowEditor/Private/Dataflow/DataflowConstructionVisualization.cpp b/Engine/Plugins/Experimental/Dataflow/Source/DataflowEditor/Private/Dataflow/DataflowConstructionVisualization.cpp
new file mode 100644
--- /dev/null
+++ b/Engine/Plugins/Experimental/Dataflow/Source/DataflowEditor/Private/Dataflow/DataflowConstructionVisualization.cpp
@@ -0,0 +1,63 @@
+// Copyright Epic Games, Inc. All Rights Reserved.
+
+#include "Dataflow/DataflowConstructionVisualization.h"
+#include "Misc/LazySingleton.h"
+#include "ChaosLog.h"
+
+#define LOCTEXT_NAMESPACE "DataflowConstructionVisualization"
+
+namespace UE::Dataflow
+{
+	FDataflowConstructionVisualizationRegistry& FDataflowConstructionVisualizationRegistry::GetInstance()
+	{
+		return TLazySingleton<FDataflowConstructionVisualizationRegistry>::Get();
+	}
+
+	void FDataflowConstructionVisualizationRegistry::TearDown()
+	{
+		return TLazySingleton<FDataflowConstructionVisualizationRegistry>::TearDown();
+	}
+
+	void FDataflowConstructionVisualizationRegistry::RegisterVisualization(TUniquePtr<IDataflowConstructionVisualization>&& Visualization)
+	{
+		const FName NewVisualizationName = Visualization->GetName();
+		if (VisualizationMap.Contains(NewVisualizationName))
+		{
+			UE_LOG(LogChaos, Warning, TEXT("Dataflow construction viewport visualization registration conflicts with existing visualization: %s"), *NewVisualizationName.ToString());
+		}
+		else
+		{
+			VisualizationMap.Add(NewVisualizationName, MoveTemp(Visualization));
+		}
+	}
+
+	void FDataflowConstructionVisualizationRegistry::DeregisterVisualization(const FName& VisualizationName)
+	{
+		if (!VisualizationMap.Contains(VisualizationName))
+		{
+			UE_LOG(LogChaos, Warning, TEXT("Dataflow visualization deregistration -- visualization not registered : %s"), *VisualizationName.ToString());
+		}
+		else
+		{
+			VisualizationMap.Remove(VisualizationName);
+		}
+	}
+
+	const TMap<FName, TUniquePtr<IDataflowConstructionVisualization>>& FDataflowConstructionVisualizationRegistry::GetVisualizations() const
+	{
+		return VisualizationMap;
+	}
+
+	const IDataflowConstructionVisualization* FDataflowConstructionVisualizationRegistry::GetVisualization(const FName& Name) const
+	{
+		if (VisualizationMap.Contains(Name))
+		{
+			return VisualizationMap[Name].Get();
+		}
+		return nullptr;
+	}
+
+}
+
+#undef LOCTEXT_NAMESPACE 
+
diff --git a/Engine/Plugins/Experimental/Dataflow/Source/DataflowEditor/Private/Dataflow/DataflowEditorToolkit.cpp b/Engine/Plugins/Experimental/Dataflow/Source/DataflowEditor/Private/Dataflow/DataflowEditorToolkit.cpp
--- a/Engine/Plugins/Experimental/Dataflow/Source/DataflowEditor/Private/Dataflow/DataflowEditorToolkit.cpp
+++ b/Engine/Plugins/Experimental/Dataflow/Source/DataflowEditor/Private/Dataflow/DataflowEditorToolkit.cpp
@@ -29,6 +29,7 @@
 #include "Dataflow/DataflowEditorPreviewSceneBase.h"
 #include "Dataflow/DataflowRenderingFactory.h"
 #include "Dataflow/DataflowSimulationScene.h"
+#include "Dataflow/DataflowSimulationVisualization.h"
 #include "Dataflow/DataflowSchema.h"
 #include "Dataflow/DataflowSkeletonView.h"
 #include "Dataflow/DataflowSimulationViewportClient.h"
@@ -77,6 +78,7 @@ const FName FDataflowEditorToolkit::CollectionSpreadSheetTabId_2(TEXT("DataflowE
 const FName FDataflowEditorToolkit::CollectionSpreadSheetTabId_3(TEXT("DataflowEditor_CollectionSpreadSheet_3"));
 const FName FDataflowEditorToolkit::CollectionSpreadSheetTabId_4(TEXT("DataflowEditor_CollectionSpreadSheet_4"));
 const FName FDataflowEditorToolkit::SimulationViewportTabId(TEXT("DataflowEditor_SimulationViewport"));
+const FName FDataflowEditorToolkit::SimulationVisualizationTabId(TEXT("DataflowEditor_SimulationVisualizationTab"));
 
 FDataflowEditorToolkit::FDataflowEditorToolkit(UAssetEditor* InOwningAssetEditor)
 	: FBaseCharacterFXEditorToolkit(InOwningAssetEditor, FName("DataflowEditor"))
@@ -157,6 +159,7 @@ FDataflowEditorToolkit::FDataflowEditorToolkit(UAssetEditor* InOwningAssetEditor
 					->SetSizeCoefficient(0.65f)	// Relative height of (Asset Details, Preview Scene Details) vs (Dataflow Node Details)
 					->AddTab(DetailsTabID, ETabState::OpenedTab)
 					->AddTab(PreviewSceneTabId, ETabState::OpenedTab)
+					->AddTab(SimulationVisualizationTabId, ETabState::OpenedTab)
 					->SetExtensionId("DetailsArea")
 					->SetHideTabWell(true)
 					->SetForegroundTab(DetailsTabID)
@@ -281,6 +284,11 @@ FDataflowEditorToolkit::FDataflowEditorToolkit(UAssetEditor* InOwningAssetEditor
 
 FDataflowEditorToolkit::~FDataflowEditorToolkit()
 {
+	if (SimulationScene && SimulationScene->GetPreviewSceneDescription())
+	{
+		SimulationScene->GetPreviewSceneDescription()->DataflowSimulationSceneDescriptionChanged.Remove(OnSimulationSceneChangedDelegateHandle);
+	}
+
 	if (GraphEditor)
 	{
 		GraphEditor->OnSelectionChangedMulticast.Remove(OnSelectionChangedMulticastDelegateHandle);
@@ -939,6 +947,8 @@ void FDataflowEditorToolkit::OnNodeSelectionChanged(const TSet<UObject*>& InNewS
 
 			EditorContent->SetSelectedNode(nullptr);
 
+			EditorContent->SetSelectedCollection(nullptr, /*bCollectionIsInput=*/ false);
+
 			if( UDataflowEditorMode* const DataflowMode = Cast<UDataflowEditorMode>(EditorModeManager->GetActiveScriptableMode(UDataflowEditorMode::EM_DataflowEditorModeId)) )
 			{
 				// Close any running tool. OnNodeSingleClicked() will start a new tool if a new node was clicked.
@@ -1485,6 +1495,37 @@ TSharedRef<SDockTab> FDataflowEditorToolkit::SpawnTab_CollectionSpreadSheet(cons
 	return DockableTab;
 }
 
+TSharedPtr<SWidget> FDataflowEditorToolkit::CreateSimulationVisualizationWidget()
+{
+	FMenuBuilder MenuBuilder(false, nullptr);
+
+	using namespace UE::Dataflow;
+	for (const TPair<FName, TUniquePtr<IDataflowSimulationVisualization>>& Visualization : FDataflowSimulationVisualizationRegistry::GetInstance().GetVisualizations())
+	{
+		Visualization.Value->ExtendSimulationVisualizationMenu(SimulationViewportClient, MenuBuilder);
+	}
+	return MenuBuilder.MakeWidget();
+}
+
+TSharedRef<SDockTab> FDataflowEditorToolkit::SpawnTab_SimulationVisualization(const FSpawnTabArgs& Args)
+{
+	TSharedRef<SDockTab> SimulationVisualizationTab = SNew(SDockTab)
+		.Label(LOCTEXT("SimulationVisualizationTitle", "Simulation Visualization"));
+
+	SimulationVisualizationWidget = CreateSimulationVisualizationWidget();
+	SimulationVisualizationTab->SetContent(SimulationVisualizationWidget.ToSharedRef());
+
+	// Re-create the visualization panel when the simulation scene changes
+	OnSimulationSceneChangedDelegateHandle = SimulationScene->GetPreviewSceneDescription()->DataflowSimulationSceneDescriptionChanged.AddLambda([this, SimulationVisualizationTab]()
+	{
+		SimulationVisualizationWidget = CreateSimulationVisualizationWidget();
+		SimulationVisualizationTab->SetContent(SimulationVisualizationWidget.ToSharedRef());
+	});
+
+	return SimulationVisualizationTab;
+}
+
+
 void FDataflowEditorToolkit::RegisterTabSpawners(const TSharedRef<FTabManager>& InTabManager)
 {
 	FAssetEditorToolkit::RegisterTabSpawners(InTabManager);
@@ -1570,6 +1611,12 @@ void FDataflowEditorToolkit::RegisterTabSpawners(const TSharedRef<FTabManager>&
 		.SetDisplayName(LOCTEXT("DataflowCollectionSpreadSheetTab4", "Collection SpreadSheet 4"))
 		.SetGroup(CollectionSpreadSheetWorkspaceMenuCategoryRef)
 		.SetIcon(FSlateIcon(FAppStyle::GetAppStyleSetName(), "LevelEditor.Tabs.Outliner"));
+
+	InTabManager->RegisterTabSpawner(SimulationVisualizationTabId, FOnSpawnTab::CreateSP(this, &FDataflowEditorToolkit::SpawnTab_SimulationVisualization))
+		.SetDisplayName(LOCTEXT("SimulationVisualizationTabDisplayName", "Simulation Visualization"))
+		.SetGroup(AssetEditorTabsCategory.ToSharedRef());
+
+	
 }
 
 void FDataflowEditorToolkit::UnregisterTabSpawners(const TSharedRef<class FTabManager>& InTabManager)
diff --git a/Engine/Plugins/Experimental/Dataflow/Source/DataflowEditor/Private/Dataflow/DataflowSimulationViewport.cpp b/Engine/Plugins/Experimental/Dataflow/Source/DataflowEditor/Private/Dataflow/DataflowSimulationViewport.cpp
--- a/Engine/Plugins/Experimental/Dataflow/Source/DataflowEditor/Private/Dataflow/DataflowSimulationViewport.cpp
+++ b/Engine/Plugins/Experimental/Dataflow/Source/DataflowEditor/Private/Dataflow/DataflowSimulationViewport.cpp
@@ -12,6 +12,8 @@
 #include "Dataflow/DataflowSimulationScene.h"
 #include "Dataflow/DataflowSimulationPanel.h"
 #include "Dataflow/DataflowSimulationViewportToolbar.h"
+#include "Dataflow/DataflowSimulationVisualization.h"
+#include "Widgets/Text/SRichTextBlock.h"
 
 #define LOCTEXT_NAMESPACE "SDataflowSimulationViewport"
 
@@ -65,8 +67,52 @@ void SDataflowSimulationViewport::Construct(const FArguments& InArgs, const FAss
 				]
 			]
 		];
+
+		ViewportOverlay->AddSlot()
+		[
+			SNew(SHorizontalBox)
+			+ SHorizontalBox::Slot()
+			.HAlign(HAlign_Left)
+			.VAlign(VAlign_Top)
+			.FillWidth(1)
+			.Padding(10.0f, 40.0f)
+			[
+				// Display text 
+				SNew(SRichTextBlock)
+					.DecoratorStyleSet(&FAppStyle::Get())
+					.Text(this, &SDataflowSimulationViewport::GetDisplayString)
+					.TextStyle(&FAppStyle::Get().GetWidgetStyle<FTextBlockStyle>("AnimViewport.MessageText"))
+			]
+		];
+
 	}
 }
+
+FText SDataflowSimulationViewport::GetDisplayString() const
+{
+	using namespace UE::Dataflow;
+	
+	auto ConcatenateLine = [](const FText& InText, const FText& InNewLine)->FText
+	{
+		if (InText.IsEmpty())
+		{
+			return InNewLine;
+		}
+		return FText::Format(LOCTEXT("ViewportTextNewlineFormatter", "{0}\n{1}"), InText, InNewLine);
+	};
+
+	FText DisplayText;
+
+	const TMap<FName, TUniquePtr<IDataflowSimulationVisualization>>& Visualizations = FDataflowSimulationVisualizationRegistry::GetInstance().GetVisualizations();
+	for (const TPair<FName, TUniquePtr<IDataflowSimulationVisualization>>& Visualization : Visualizations)
+	{
+		FText Text = Visualization.Value->GetDisplayString(GetSimulationScene().Get());
+		DisplayText = ConcatenateLine(DisplayText, Text);
+	}
+	
+	return DisplayText;
+}
+
 TSharedPtr<SWidget> SDataflowSimulationViewport::MakeViewportToolbar()
 {
 	return SNew(SDataflowSimulationViewportToolBar, SharedThis(this)).CommandList(CommandList);
diff --git a/Engine/Plugins/Experimental/Dataflow/Source/DataflowEditor/Private/Dataflow/DataflowSimulationViewportClient.cpp b/Engine/Plugins/Experimental/Dataflow/Source/DataflowEditor/Private/Dataflow/DataflowSimulationViewportClient.cpp
--- a/Engine/Plugins/Experimental/Dataflow/Source/DataflowEditor/Private/Dataflow/DataflowSimulationViewportClient.cpp
+++ b/Engine/Plugins/Experimental/Dataflow/Source/DataflowEditor/Private/Dataflow/DataflowSimulationViewportClient.cpp
@@ -9,6 +9,8 @@
 #include "Dataflow/DataflowEngineSceneHitProxies.h"
 #include "Dataflow/DataflowGraphEditor.h"
 #include "Dataflow/DataflowEditorPreviewSceneBase.h"
+#include "Dataflow/DataflowSimulationScene.h"
+#include "Dataflow/DataflowSimulationVisualization.h"
 #include "EditorModeManager.h"
 #include "EdModeInteractiveToolsContext.h"
 #include "GraphEditor.h"
@@ -96,3 +98,32 @@ void FDataflowSimulationViewportClient::AddReferencedObjects(FReferenceCollector
 	Collector.AddReferencedObject(BehaviorSet);
 }
 
+void FDataflowSimulationViewportClient::Draw(const FSceneView* View, FPrimitiveDrawInterface* PDI)
+{
+	FEditorViewportClient::Draw(View, PDI);
+
+	using namespace UE::Dataflow;
+	const TMap<FName, TUniquePtr<IDataflowSimulationVisualization>>& Visualizations = FDataflowSimulationVisualizationRegistry::GetInstance().GetVisualizations();
+	for (const TPair<FName, TUniquePtr<IDataflowSimulationVisualization>>& Visualization : Visualizations)
+	{
+		if (PreviewScene)
+		{
+			Visualization.Value->Draw(static_cast<FDataflowSimulationScene*>(PreviewScene), PDI);
+		}
+	}
+}
+
+void FDataflowSimulationViewportClient::DrawCanvas(FViewport& InViewport, FSceneView& View, FCanvas& Canvas)
+{
+	FEditorViewportClient::DrawCanvas(InViewport, View, Canvas);
+
+	using namespace UE::Dataflow;
+	const TMap<FName, TUniquePtr<IDataflowSimulationVisualization>>& Visualizations = FDataflowSimulationVisualizationRegistry::GetInstance().GetVisualizations();
+	for (const TPair<FName, TUniquePtr<IDataflowSimulationVisualization>>& Visualization : Visualizations)
+	{
+		if (PreviewScene)
+		{
+			Visualization.Value->DrawCanvas(static_cast<FDataflowSimulationScene*>(PreviewScene), &Canvas, &View);
+		}
+	}
+}
diff --git a/Engine/Plugins/Experimental/Dataflow/Source/DataflowEditor/Private/Dataflow/DataflowSimulationVisualization.cpp b/Engine/Plugins/Experimental/Dataflow/Source/DataflowEditor/Private/Dataflow/DataflowSimulationVisualization.cpp
new file mode 100644
--- /dev/null
+++ b/Engine/Plugins/Experimental/Dataflow/Source/DataflowEditor/Private/Dataflow/DataflowSimulationVisualization.cpp
@@ -0,0 +1,63 @@
+// Copyright Epic Games, Inc. All Rights Reserved.
+
+#include "Dataflow/DataflowSimulationVisualization.h"
+#include "Misc/LazySingleton.h"
+#include "ChaosLog.h"
+
+#define LOCTEXT_NAMESPACE "DataflowSimulationVisualization"
+
+namespace UE::Dataflow
+{
+	FDataflowSimulationVisualizationRegistry& FDataflowSimulationVisualizationRegistry::GetInstance()
+	{
+		return TLazySingleton<FDataflowSimulationVisualizationRegistry>::Get();
+	}
+
+	void FDataflowSimulationVisualizationRegistry::TearDown()
+	{
+		return TLazySingleton<FDataflowSimulationVisualizationRegistry>::TearDown();
+	}
+
+	void FDataflowSimulationVisualizationRegistry::RegisterVisualization(TUniquePtr<IDataflowSimulationVisualization>&& Visualization)
+	{
+		const FName NewVisualizationName = Visualization->GetName();
+		if (VisualizationMap.Contains(NewVisualizationName))
+		{
+			UE_LOG(LogChaos, Warning, TEXT("Dataflow simulation visualization registration conflicts with existing visualization: %s"), *NewVisualizationName.ToString());
+		}
+		else
+		{
+			VisualizationMap.Add(NewVisualizationName, MoveTemp(Visualization));
+		}
+	}
+
+	void FDataflowSimulationVisualizationRegistry::DeregisterVisualization(const FName& VisualizationName)
+	{
+		if (!VisualizationMap.Contains(VisualizationName))
+		{
+			UE_LOG(LogChaos, Warning, TEXT("Dataflow visualization deregistration -- visualization not registered : %s"), *VisualizationName.ToString());
+		}
+		else
+		{
+			VisualizationMap.Remove(VisualizationName);
+		}
+	}
+
+	const TMap<FName, TUniquePtr<IDataflowSimulationVisualization>>& FDataflowSimulationVisualizationRegistry::GetVisualizations() const
+	{
+		return VisualizationMap;
+	}
+
+	const IDataflowSimulationVisualization* FDataflowSimulationVisualizationRegistry::GetVisualization(const FName& Name) const
+	{
+		if (VisualizationMap.Contains(Name))
+		{
+			return VisualizationMap[Name].Get();
+		}
+		return nullptr;
+	}
+
+}
+
+#undef LOCTEXT_NAMESPACE 
+
diff --git a/Engine/Plugins/Experimental/Dataflow/Source/DataflowEditor/Public/Dataflow/DataflowConstructionViewportClient.h b/Engine/Plugins/Experimental/Dataflow/Source/DataflowEditor/Public/Dataflow/DataflowConstructionViewportClient.h
--- a/Engine/Plugins/Experimental/Dataflow/Source/DataflowEditor/Public/Dataflow/DataflowConstructionViewportClient.h
+++ b/Engine/Plugins/Experimental/Dataflow/Source/DataflowEditor/Public/Dataflow/DataflowConstructionViewportClient.h
@@ -60,6 +60,8 @@ private:
 	virtual void Tick(float DeltaSeconds) override;
 	virtual bool InputKey(const FInputKeyEventArgs& EventArgs) override;
 	virtual void ProcessClick(FSceneView& View, HHitProxy* HitProxy, FKey Key, EInputEvent Event, uint32 HitX, uint32 HitY) override;
+	virtual void Draw(const FSceneView* View, FPrimitiveDrawInterface* PDI) override;
+	virtual void DrawCanvas(FViewport& InViewport, FSceneView& View, FCanvas& Canvas) override;
 
 	/** Toolkit used to create the viewport client */
 	TWeakPtr<FDataflowEditorToolkit> DataflowEditorToolkitPtr = nullptr;
diff --git a/Engine/Plugins/Experimental/Dataflow/Source/DataflowEditor/Public/Dataflow/DataflowConstructionViewportToolbar.h b/Engine/Plugins/Experimental/Dataflow/Source/DataflowEditor/Public/Dataflow/DataflowConstructionViewportToolbar.h
--- a/Engine/Plugins/Experimental/Dataflow/Source/DataflowEditor/Public/Dataflow/DataflowConstructionViewportToolbar.h
+++ b/Engine/Plugins/Experimental/Dataflow/Source/DataflowEditor/Public/Dataflow/DataflowConstructionViewportToolbar.h
@@ -23,9 +23,10 @@ public:
 
 private:
 
-	TSharedRef<SWidget> MakeToolBar(const TSharedPtr<FExtender> InExtenders);
+	// SCommonEditorViewportToolbarBase
+	virtual TSharedRef<SWidget> GenerateShowMenu() const override;
 
-	TSharedPtr<SEditorViewportToolbarMenu> ViewModeDropDown;
+	TSharedRef<SWidget> MakeToolBar(const TSharedPtr<FExtender> InExtenders);
 
 	// View mode
 	FText GetViewModeMenuLabel() const;
@@ -35,6 +36,8 @@ private:
 	// Display (Lit, Unlit, Wireframe, etc.)
 	TSharedRef<SWidget> MakeDisplayToolBar(const TSharedPtr<FExtender> InExtenders);
 
+	TSharedPtr<SEditorViewportToolbarMenu> ViewModeDropDown;
+
 	/** Reference to the parent viewport */
 	TWeakPtr<SDataflowConstructionViewport> EditorViewport;
 
diff --git a/Engine/Plugins/Experimental/Dataflow/Source/DataflowEditor/Public/Dataflow/DataflowConstructionVisualization.h b/Engine/Plugins/Experimental/Dataflow/Source/DataflowEditor/Public/Dataflow/DataflowConstructionVisualization.h
new file mode 100644
--- /dev/null
+++ b/Engine/Plugins/Experimental/Dataflow/Source/DataflowEditor/Public/Dataflow/DataflowConstructionVisualization.h
@@ -0,0 +1,51 @@
+// Copyright Epic Games, Inc. All Rights Reserved.
+
+#pragma once
+
+#include "UObject/NameTypes.h"
+#include "Templates/SharedPointer.h"
+
+class FPrimitiveDrawInterface;
+class FSceneView;
+class FCanvas;
+class FMenuBuilder;
+class FDataflowConstructionScene;
+class FDataflowConstructionViewportClient;
+
+namespace UE::Dataflow
+{
+
+	class IDataflowConstructionVisualization
+	{
+	public:
+		virtual ~IDataflowConstructionVisualization() = default;
+
+		DATAFLOWEDITOR_API virtual FName GetName() const = 0;
+		DATAFLOWEDITOR_API virtual void ExtendViewportShowMenu(const TSharedPtr<FDataflowConstructionViewportClient>& ViewportClient, FMenuBuilder& MenuBuilder) {};
+		DATAFLOWEDITOR_API virtual void Draw(const FDataflowConstructionScene* ConstructionScene, FPrimitiveDrawInterface* PDI) {};
+		DATAFLOWEDITOR_API virtual void DrawCanvas(const FDataflowConstructionScene* ConstructionScene, FCanvas* Canvas, const FSceneView* SceneView) {};
+	};
+
+	class FDataflowConstructionVisualizationRegistry
+	{
+	public:
+
+		// FLazySingleton
+		static DATAFLOWEDITOR_API FDataflowConstructionVisualizationRegistry& GetInstance();
+		static DATAFLOWEDITOR_API void TearDown();
+
+		DATAFLOWEDITOR_API void RegisterVisualization(TUniquePtr<IDataflowConstructionVisualization>&& Visualization);
+		DATAFLOWEDITOR_API void DeregisterVisualization(const FName& VisualizationName);
+
+		DATAFLOWEDITOR_API const TMap<FName, TUniquePtr<IDataflowConstructionVisualization>>& GetVisualizations() const;
+		DATAFLOWEDITOR_API const IDataflowConstructionVisualization* GetVisualization(const FName& VisualizationOption) const;
+
+	private:
+
+		DATAFLOWEDITOR_API static TUniquePtr<FDataflowConstructionVisualizationRegistry> Instance;
+
+		TMap<FName, TUniquePtr<IDataflowConstructionVisualization>> VisualizationMap;
+	};
+
+}
+
diff --git a/Engine/Plugins/Experimental/Dataflow/Source/DataflowEditor/Public/Dataflow/DataflowEditorToolkit.h b/Engine/Plugins/Experimental/Dataflow/Source/DataflowEditor/Public/Dataflow/DataflowEditorToolkit.h
--- a/Engine/Plugins/Experimental/Dataflow/Source/DataflowEditor/Public/Dataflow/DataflowEditorToolkit.h
+++ b/Engine/Plugins/Experimental/Dataflow/Source/DataflowEditor/Public/Dataflow/DataflowEditorToolkit.h
@@ -107,6 +107,7 @@ private:
 	TSharedRef<SDockTab> SpawnTab_AssetDetails(const FSpawnTabArgs& Args);
 	TSharedRef<SDockTab> SpawnTab_SimulationViewport(const FSpawnTabArgs& Args);
 	TSharedRef<SDockTab> SpawnTab_PreviewScene(const FSpawnTabArgs& Args);
+	TSharedRef<SDockTab> SpawnTab_SimulationVisualization(const FSpawnTabArgs& Args);
 	
 	// FTickableEditorObject interface
 	virtual void Tick(float DeltaTime) override;
@@ -146,6 +147,7 @@ private:
 	static const FName CollectionSpreadSheetTabId_4;
 	static const FName SimulationViewportTabId;
 	static const FName PreviewSceneTabId;
+	static const FName SimulationVisualizationTabId;
 
 	// List of all the widgets shared ptr that will be built in the editor
 	TSharedPtr<SDataflowConstructionViewport> DataflowConstructionViewport;
@@ -163,6 +165,7 @@ private:
 	TSharedPtr<FDataflowCollectionSpreadSheet> DataflowCollectionSpreadSheet_3;
 	TSharedPtr<FDataflowCollectionSpreadSheet> DataflowCollectionSpreadSheet_4;
 	TSharedPtr<SWidget> AdvancedPreviewSettingsWidget;
+	TSharedPtr<SWidget> SimulationVisualizationWidget;
 
 	/** Customize preview scene with editor/terminal contents */
 	TSharedRef<class IDetailCustomization> CustomizePreviewSceneDescription() const;
@@ -170,6 +173,7 @@ private:
 	// Utility factory functions to build the widgets
 	TSharedRef<SDataflowGraphEditor> CreateGraphEditorWidget(UDataflow* ObjectToEdit, TSharedPtr<IStructureDetailsView> PropertiesEditor);
     TSharedPtr<IDetailsView> CreateAssetDetailsEditorWidget(const TArray<UObject*>& ObjectsToEdit);
+	TSharedPtr<SWidget> CreateSimulationVisualizationWidget();
     TSharedPtr<IStructureDetailsView> CreateNodeDetailsEditorWidget(UObject* ObjectToEdit);
 
 	/** Create the simulation viewport client */
@@ -187,6 +191,7 @@ private:
     FDelegateHandle OnFinishedChangingPropertiesDelegateHandle;
 	FDelegateHandle OnFinishedChangingAssetPropertiesDelegateHandle;
 	FDelegateHandle OnConstructionSelectionChangedDelegateHandle;
+	FDelegateHandle OnSimulationSceneChangedDelegateHandle;
 
 	// The currently selected set of dataflow nodes. 
 	UPROPERTY()
diff --git a/Engine/Plugins/Experimental/Dataflow/Source/DataflowEditor/Public/Dataflow/DataflowSimulationViewport.h b/Engine/Plugins/Experimental/Dataflow/Source/DataflowEditor/Public/Dataflow/DataflowSimulationViewport.h
--- a/Engine/Plugins/Experimental/Dataflow/Source/DataflowEditor/Public/Dataflow/DataflowSimulationViewport.h
+++ b/Engine/Plugins/Experimental/Dataflow/Source/DataflowEditor/Public/Dataflow/DataflowSimulationViewport.h
@@ -42,6 +42,8 @@ private:
 	float GetViewMinInput() const;
 	float GetViewMaxInput() const;
 
+	FText GetDisplayString() const;
+
 	UDataflowEditorMode* GetEdMode() const;
 };
 
diff --git a/Engine/Plugins/Experimental/Dataflow/Source/DataflowEditor/Public/Dataflow/DataflowSimulationViewportClient.h b/Engine/Plugins/Experimental/Dataflow/Source/DataflowEditor/Public/Dataflow/DataflowSimulationViewportClient.h
--- a/Engine/Plugins/Experimental/Dataflow/Source/DataflowEditor/Public/Dataflow/DataflowSimulationViewportClient.h
+++ b/Engine/Plugins/Experimental/Dataflow/Source/DataflowEditor/Public/Dataflow/DataflowSimulationViewportClient.h
@@ -44,6 +44,8 @@ private:
 	// FEditorViewportClient interface
 	virtual void Tick(float DeltaSeconds) override;
 	virtual void ProcessClick(FSceneView& View, HHitProxy* HitProxy, FKey Key, EInputEvent Event, uint32 HitX, uint32 HitY) override;
+	virtual void Draw(const FSceneView* View, FPrimitiveDrawInterface* PDI) override;
+	virtual void DrawCanvas(FViewport& InViewport, FSceneView& View, FCanvas& Canvas) override;
 
 	/** Toolkit used to create the viewport client */
 	TWeakPtr<FDataflowEditorToolkit> DataflowEditorToolkitPtr = nullptr;
diff --git a/Engine/Plugins/Experimental/Dataflow/Source/DataflowEditor/Public/Dataflow/DataflowSimulationVisualization.h b/Engine/Plugins/Experimental/Dataflow/Source/DataflowEditor/Public/Dataflow/DataflowSimulationVisualization.h
new file mode 100644
--- /dev/null
+++ b/Engine/Plugins/Experimental/Dataflow/Source/DataflowEditor/Public/Dataflow/DataflowSimulationVisualization.h
@@ -0,0 +1,53 @@
+// Copyright Epic Games, Inc. All Rights Reserved.
+
+#pragma once
+
+#include "UObject/NameTypes.h"
+#include "Templates/SharedPointer.h"
+#include "Internationalization/Text.h"
+
+class FPrimitiveDrawInterface;
+class FSceneView;
+class FCanvas;
+class FMenuBuilder;
+class FDataflowSimulationScene;
+class FDataflowSimulationViewportClient;
+
+namespace UE::Dataflow
+{
+
+	class IDataflowSimulationVisualization
+	{
+	public:
+		virtual ~IDataflowSimulationVisualization() = default;
+
+		DATAFLOWEDITOR_API virtual FName GetName() const = 0;
+		DATAFLOWEDITOR_API virtual void ExtendSimulationVisualizationMenu(const TSharedPtr<FDataflowSimulationViewportClient>& ViewportClient, FMenuBuilder& MenuBuilder) {};
+		DATAFLOWEDITOR_API virtual void Draw(const FDataflowSimulationScene* SimulationScene, FPrimitiveDrawInterface* PDI) {};
+		DATAFLOWEDITOR_API virtual void DrawCanvas(const FDataflowSimulationScene* SimulationScene, FCanvas* Canvas, const FSceneView* SceneView) {};
+		DATAFLOWEDITOR_API virtual FText GetDisplayString(const FDataflowSimulationScene* SimulationScene) const { return FText(); };
+	};
+
+	class FDataflowSimulationVisualizationRegistry
+	{
+	public:
+
+		// FLazySingleton
+		static DATAFLOWEDITOR_API FDataflowSimulationVisualizationRegistry& GetInstance();
+		static DATAFLOWEDITOR_API void TearDown();
+
+		DATAFLOWEDITOR_API void RegisterVisualization(TUniquePtr<IDataflowSimulationVisualization>&& Visualization);
+		DATAFLOWEDITOR_API void DeregisterVisualization(const FName& VisualizationName);
+
+		DATAFLOWEDITOR_API const TMap<FName, TUniquePtr<IDataflowSimulationVisualization>>& GetVisualizations() const;
+		DATAFLOWEDITOR_API const IDataflowSimulationVisualization* GetVisualization(const FName& VisualizationOption) const;
+
+	private:
+
+		DATAFLOWEDITOR_API static TUniquePtr<FDataflowSimulationVisualizationRegistry> Instance;
+
+		TMap<FName, TUniquePtr<IDataflowSimulationVisualization>> VisualizationMap;
+	};
+
+}
+
