Chaos Cloth Asset: update OuterLift and OuterDrag to import the Fabrics "SolverAirDamping" when requested. This matches the behavior of Lift and Drag.

diff --git a/Engine/Plugins/ChaosClothAssetEditor/Source/ChaosClothAssetDataflowNodes/Private/ChaosClothAsset/SimulationAerodynamicsConfigNode.cpp b/Engine/Plugins/ChaosClothAssetEditor/Source/ChaosClothAssetDataflowNodes/Private/ChaosClothAsset/SimulationAerodynamicsConfigNode.cpp
--- a/Engine/Plugins/ChaosClothAssetEditor/Source/ChaosClothAssetDataflowNodes/Private/ChaosClothAsset/SimulationAerodynamicsConfigNode.cpp
+++ b/Engine/Plugins/ChaosClothAssetEditor/Source/ChaosClothAssetDataflowNodes/Private/ChaosClothAsset/SimulationAerodynamicsConfigNode.cpp
@@ -38,7 +38,11 @@ void FChaosClothAssetSimulationAerodynamicsConfigNode::AddProperties(FPropertyHe
 
 	if (bEnableOuterDrag)
 	{
-		PropertyHelper.SetPropertyWeighted(this, &OuterDrag);
+		PropertyHelper.SetSolverPropertyWeighted(FName(TEXT("OuterDrag")), OuterDrag, [](
+			const UE::Chaos::ClothAsset::FCollectionClothFacade& ClothFacade)-> float
+			{
+				return ClothFacade.GetSolverAirDamping();
+			}, {});
 	}
 
 	PropertyHelper.SetSolverPropertyWeighted(FName(TEXT("Lift")), Lift, [](
@@ -49,6 +53,10 @@ void FChaosClothAssetSimulationAerodynamicsConfigNode::AddProperties(FPropertyHe
 
 	if (bEnableOuterLift)
 	{
-		PropertyHelper.SetPropertyWeighted(this, &OuterLift);
+		PropertyHelper.SetSolverPropertyWeighted(FName(TEXT("OuterLift")), OuterLift, [](
+			const UE::Chaos::ClothAsset::FCollectionClothFacade& ClothFacade)-> float
+			{
+				return ClothFacade.GetSolverAirDamping();
+			}, {});
 	}
 }
