Dataflow - Added access to the Dataflow Editor context from detail customizations.
Cloth Asset - Removed node specific code from the selection group customization.

#tests Created a new selection node and opened the group property list box to verify that the group list would display without delay.
#rnx

diff --git a/Engine/Plugins/ChaosClothAssetEditor/Source/ChaosClothAssetDataflowNodes/Private/ChaosClothAsset/AttributeNode.cpp b/Engine/Plugins/ChaosClothAssetEditor/Source/ChaosClothAssetDataflowNodes/Private/ChaosClothAsset/AttributeNode.cpp
--- a/Engine/Plugins/ChaosClothAssetEditor/Source/ChaosClothAssetDataflowNodes/Private/ChaosClothAsset/AttributeNode.cpp
+++ b/Engine/Plugins/ChaosClothAssetEditor/Source/ChaosClothAssetDataflowNodes/Private/ChaosClothAsset/AttributeNode.cpp
@@ -90,33 +90,6 @@ void FChaosClothAssetAttributeNode_v2::Evaluate(UE::Dataflow::FContext& Context,
 	}
 }
 
-void FChaosClothAssetAttributeNode_v2::OnSelected(UE::Dataflow::FContext& Context)
-{
-	using namespace UE::Chaos::ClothAsset;
-
-	// Re-evaluate the input collection
-	FManagedArrayCollection InCollection = GetValue<FManagedArrayCollection>(Context, &Collection);
-	const TSharedRef<FManagedArrayCollection> ClothCollection = MakeShared<FManagedArrayCollection>(MoveTemp(InCollection));
-	FCollectionClothFacade Cloth(ClothCollection);
-
-	// Update the list of used group for the UI customization
-	const TArray<FName> GroupNames = ClothCollection->GroupNames();
-	CachedCollectionGroupNames.Reset(GroupNames.Num());
-	for (const FName& GroupName : GroupNames)
-	{
-		if (Cloth.IsValidClothCollectionGroupName(GroupName))  // Restrict to the cloth facade groups
-		{
-			CachedCollectionGroupNames.Emplace(GroupName);
-		}
-	}
-}
-
-void FChaosClothAssetAttributeNode_v2::OnDeselected()
-{
-	// Clean up, to avoid another toolkit picking up the wrong context evaluation
-	CachedCollectionGroupNames.Reset();
-}
-
 FChaosClothAssetAttributeNode::FChaosClothAssetAttributeNode(const UE::Dataflow::FNodeParameters& InParam, FGuid InGuid)
 	: FDataflowNode(InParam, InGuid)
 {
@@ -192,32 +165,4 @@ void FChaosClothAssetAttributeNode::Evaluate(UE::Dataflow::FContext& Context, co
 	}
 }
 
-void FChaosClothAssetAttributeNode::OnSelected(UE::Dataflow::FContext& Context)
-{
-	using namespace UE::Chaos::ClothAsset;
-
-	// Re-evaluate the input collection
-	FManagedArrayCollection InCollection = GetValue<FManagedArrayCollection>(Context, &Collection);
-	const TSharedRef<FManagedArrayCollection> ClothCollection = MakeShared<FManagedArrayCollection>(MoveTemp(InCollection));
-	FCollectionClothFacade Cloth(ClothCollection);
-
-	// Update the list of used group for the UI customization
-	const TArray<FName> GroupNames = ClothCollection->GroupNames();
-	CachedCollectionGroupNames.Reset(GroupNames.Num());
-	for (const FName& GroupName : GroupNames)
-	{
-		if (Cloth.IsValidClothCollectionGroupName(GroupName))  // Restrict to the cloth facade groups
-		{
-			CachedCollectionGroupNames.Emplace(GroupName);
-		}
-	}
-}
-
-void FChaosClothAssetAttributeNode::OnDeselected()
-{
-	// Clean up, to avoid another toolkit picking up the wrong context evaluation
-	CachedCollectionGroupNames.Reset();
-}
-
-
 #undef LOCTEXT_NAMESPACE
diff --git a/Engine/Plugins/ChaosClothAssetEditor/Source/ChaosClothAssetDataflowNodes/Private/ChaosClothAsset/DeleteElementNode.cpp b/Engine/Plugins/ChaosClothAssetEditor/Source/ChaosClothAssetDataflowNodes/Private/ChaosClothAsset/DeleteElementNode.cpp
--- a/Engine/Plugins/ChaosClothAssetEditor/Source/ChaosClothAssetDataflowNodes/Private/ChaosClothAsset/DeleteElementNode.cpp
+++ b/Engine/Plugins/ChaosClothAssetEditor/Source/ChaosClothAssetDataflowNodes/Private/ChaosClothAsset/DeleteElementNode.cpp
@@ -372,21 +372,6 @@ void FChaosClothAssetDeleteElementNode::Evaluate(UE::Dataflow::FContext& Context
 	}
 }
 
-void FChaosClothAssetDeleteElementNode::OnSelected(UE::Dataflow::FContext& Context)
-{
-	// Re-evaluate the input collection
-	const FManagedArrayCollection& SelectionCollection = GetValue<FManagedArrayCollection>(Context, &Collection);
-
-	// Update the list of used group for the UI customization
-	CachedCollectionGroupNames = SelectionCollection.GroupNames();
-}
-
-void FChaosClothAssetDeleteElementNode::OnDeselected()
-{
-	// Clean up, to avoid another toolkit picking up the wrong context evaluation
-	CachedCollectionGroupNames.Reset();
-}
-
 void FChaosClothAssetDeleteElementNode::Serialize(FArchive& Ar)
 {
 	using namespace UE::Chaos::ClothAsset;
diff --git a/Engine/Plugins/ChaosClothAssetEditor/Source/ChaosClothAssetDataflowNodes/Private/ChaosClothAsset/ImportFilePathCustomization.cpp b/Engine/Plugins/ChaosClothAssetEditor/Source/ChaosClothAssetDataflowNodes/Private/ChaosClothAsset/ImportFilePathCustomization.cpp
--- a/Engine/Plugins/ChaosClothAssetEditor/Source/ChaosClothAssetDataflowNodes/Private/ChaosClothAsset/ImportFilePathCustomization.cpp
+++ b/Engine/Plugins/ChaosClothAssetEditor/Source/ChaosClothAssetDataflowNodes/Private/ChaosClothAsset/ImportFilePathCustomization.cpp
@@ -2,6 +2,7 @@
 
 #include "ChaosClothAsset/ImportFilePathCustomization.h"
 #include "ChaosClothAsset/ImportFilePath.h"
+#include "Dataflow/DataflowGraphEditor.h"
 #include "Misc/MessageDialog.h"
 #include "Misc/PackageName.h"
 #include "Widgets/Images/SImage.h"
@@ -45,6 +46,9 @@ namespace UE::Chaos::ClothAsset
 
 	void FImportFilePathCustomization::CustomizeHeader(TSharedRef<IPropertyHandle> StructPropertyHandle, FDetailWidgetRow& HeaderRow, IPropertyTypeCustomizationUtils& CustomizationUtils)
 	{
+		// Keep a weak pointer to the graph editor asking for this customization
+		DataflowGraphEditor = SDataflowGraphEditor::GetSelectedGraphEditor();
+
 		StructProperty = StructPropertyHandle;
 		PathStringProperty = StructPropertyHandle->GetChildHandle(TEXT("FilePath"));
 
@@ -165,7 +169,13 @@ namespace UE::Chaos::ClothAsset
 				PathStringProperty->SetValue(FinalPath, EPropertyValueSetFlags::InteractiveChange);  // Do as an interactive change so that the delegate is called with the correct value before the node invalidates
 				if (const FChaosClothAssetImportFilePath* const ImportFilePath = Private::GetImportFilePath(StructProperty))
 				{
-					ImportFilePath->Execute();
+					// Retrieve context if any
+					const TSharedPtr<const SDataflowGraphEditor> DataflowGraphEditorPtr = DataflowGraphEditor.Pin();
+					const TSharedPtr<UE::Dataflow::FContext> Context = DataflowGraphEditorPtr ? DataflowGraphEditorPtr->GetDataflowContext() : TSharedPtr<UE::Dataflow::FContext>();
+
+					// Execute function
+					UE::Dataflow::FContextThreaded EmptyContext;
+					ImportFilePath->Execute(Context.IsValid() ? *Context : EmptyContext);
 				}
 				PathStringProperty->SetValue(FinalPath);  // This will set the final value and call invalidate on the node
 			}
diff --git a/Engine/Plugins/ChaosClothAssetEditor/Source/ChaosClothAssetDataflowNodes/Private/ChaosClothAsset/ImportFilePathCustomization.h b/Engine/Plugins/ChaosClothAssetEditor/Source/ChaosClothAssetDataflowNodes/Private/ChaosClothAsset/ImportFilePathCustomization.h
--- a/Engine/Plugins/ChaosClothAssetEditor/Source/ChaosClothAssetDataflowNodes/Private/ChaosClothAsset/ImportFilePathCustomization.h
+++ b/Engine/Plugins/ChaosClothAssetEditor/Source/ChaosClothAssetDataflowNodes/Private/ChaosClothAsset/ImportFilePathCustomization.h
@@ -6,6 +6,7 @@
 #include "IPropertyTypeCustomization.h"
 
 class SButton;
+class SDataflowGraphEditor;
 
 namespace UE::Chaos::ClothAsset
 {
@@ -27,6 +28,7 @@ namespace UE::Chaos::ClothAsset
 		FString HandleFilePathPickerFilePath() const;
 		void HandleFilePathPickerPathPicked(const FString& PickedPath);
 
+		TWeakPtr<const SDataflowGraphEditor> DataflowGraphEditor;
 		TSharedPtr<IPropertyHandle> StructProperty;
 		TSharedPtr<IPropertyHandle> PathStringProperty;
 		bool bLongPackageName;
diff --git a/Engine/Plugins/ChaosClothAssetEditor/Source/ChaosClothAssetDataflowNodes/Private/ChaosClothAsset/SelectionGroupCustomization.cpp b/Engine/Plugins/ChaosClothAssetEditor/Source/ChaosClothAssetDataflowNodes/Private/ChaosClothAsset/SelectionGroupCustomization.cpp
--- a/Engine/Plugins/ChaosClothAssetEditor/Source/ChaosClothAssetDataflowNodes/Private/ChaosClothAsset/SelectionGroupCustomization.cpp
+++ b/Engine/Plugins/ChaosClothAssetEditor/Source/ChaosClothAssetDataflowNodes/Private/ChaosClothAsset/SelectionGroupCustomization.cpp
@@ -6,6 +6,7 @@
 #include "ChaosClothAsset/CollectionClothFacade.h"
 #include "ChaosClothAsset/DeleteElementNode.h"
 #include "ChaosClothAsset/SelectionNode.h"
+#include "Dataflow/DataflowGraphEditor.h"
 #include "Widgets/Input/SEditableTextBox.h"
 #include "Widgets/Input/SComboButton.h"
 #include "Widgets/Views/SListView.h"
@@ -18,6 +19,52 @@
 
 namespace UE::Chaos::ClothAsset
 {
+	namespace Private
+	{
+		// Return the FManagedArrayCollection with the specified name from the property held by the top level struct owner of ChildPropertyHandle
+		FManagedArrayCollection GetPropertyCollection(
+			UE::Dataflow::FContext& Context,
+			const TSharedPtr<IPropertyHandle>& ChildPropertyHandle,
+			const FName CollectionPropertyName = FName(TEXT("Collection")))
+		{
+			static const FManagedArrayCollection EmptyCollection;
+
+			TSharedPtr<IPropertyHandle> OwnerHandle = ChildPropertyHandle;
+			while (TSharedPtr<IPropertyHandle> ParentHandle = OwnerHandle->GetParentHandle())
+			{
+				OwnerHandle = MoveTemp(ParentHandle);
+			}
+			if (const TSharedPtr<IPropertyHandleStruct> OwnerHandleStruct = OwnerHandle->AsStruct())
+			{
+				if (const TSharedPtr<FStructOnScope> StructOnScope = OwnerHandleStruct->GetStructData())
+				{
+					if (const UStruct* const Struct = StructOnScope->GetStruct())
+					{
+						if (Struct->IsChildOf<FDataflowNode>())
+						{
+							const FDataflowNode* const DataflowNode = reinterpret_cast<FDataflowNode*>(StructOnScope->GetStructMemory());
+
+							if (const FProperty* const Property = Struct->FindPropertyByName(CollectionPropertyName))
+							{
+								if (const FStructProperty* const StructProperty = CastField<FStructProperty>(Property))
+								{
+									if (StructProperty->GetCPPType(nullptr, CPPF_None) == TEXT("FManagedArrayCollection"))
+									{
+										if (const FDataflowInput* const DataflowInput = DataflowNode->FindInput(StructProperty->ContainerPtrToValuePtr<FManagedArrayCollection*>(DataflowNode)))
+										{
+											return DataflowInput->GetValue(Context, EmptyCollection);
+										}
+									}
+								}
+							}
+						}
+					}
+				}
+			}
+			return EmptyCollection;
+		}
+	}  // End namespace Private
+
 	TSharedRef<IPropertyTypeCustomization> FSelectionGroupCustomization::MakeInstance()
 	{
 		return MakeShareable(new FSelectionGroupCustomization);
@@ -25,6 +72,8 @@ namespace UE::Chaos::ClothAsset
 
 	void FSelectionGroupCustomization::CustomizeHeader(TSharedRef<IPropertyHandle> PropertyHandle, FDetailWidgetRow& HeaderRow, IPropertyTypeCustomizationUtils& CustomizationUtils)
 	{
+		DataflowGraphEditor = SDataflowGraphEditor::GetSelectedGraphEditor();
+
 		uint32 NumChildren;
 		const FPropertyAccess::Result Result = PropertyHandle->GetNumChildren(NumChildren);
 
@@ -107,36 +156,22 @@ namespace UE::Chaos::ClothAsset
 
 	TSharedRef<SWidget> FSelectionGroupCustomization::OnGetMenuContent()
 	{
+		// Retrieve context if any, or use an empty context
+		const TSharedPtr<const SDataflowGraphEditor> DataflowGraphEditorPtr = DataflowGraphEditor.Pin();
+		const TSharedPtr<UE::Dataflow::FContext> Context = DataflowGraphEditorPtr ? DataflowGraphEditorPtr->GetDataflowContext() : TSharedPtr<UE::Dataflow::FContext>();
+
 		// Find all group names in the parent selection node's collection
 		GroupNames.Reset();
-
-		TArray<FName> CollectionGroupNames;
-		if (const FChaosClothAssetSelectionNode_v2* const SelectionNode_v2 = GetOwnerStruct<FChaosClothAssetSelectionNode_v2>())
-		{
-			CollectionGroupNames = SelectionNode_v2->GetCachedCollectionGroupNames();
-		}
-		else if (const FChaosClothAssetDeleteElementNode* const DeleteNode = GetOwnerStruct<FChaosClothAssetDeleteElementNode>())
-		{
-			CollectionGroupNames = DeleteNode->GetCachedCollectionGroupNames();
-		}
-		else if (const FChaosClothAssetAttributeNode_v2* const AttributeNode_v2 = GetOwnerStruct<FChaosClothAssetAttributeNode_v2>())
-		{
-			CollectionGroupNames = AttributeNode_v2->GetCachedCollectionGroupNames();
-		}
-PRAGMA_DISABLE_DEPRECATION_WARNINGS
-		else if (const FChaosClothAssetSelectionNode* const SelectionNode = GetOwnerStruct<FChaosClothAssetSelectionNode>())
+		UE::Dataflow::FContextThreaded EmptyContext;
+		FManagedArrayCollection Collection = Private::GetPropertyCollection(Context.IsValid() ? *Context : EmptyContext, ChildPropertyHandle);
+		const TSharedRef<FManagedArrayCollection> ClothCollection = MakeShared<FManagedArrayCollection>(MoveTemp(Collection));
+		FCollectionClothFacade Cloth(ClothCollection);
+		for (const FName& GroupName : ClothCollection->GroupNames())
 		{
-			CollectionGroupNames = SelectionNode->GetCachedCollectionGroupNames();
-		}
-		else if (const FChaosClothAssetAttributeNode* const AttributeNode = GetOwnerStruct<FChaosClothAssetAttributeNode>())
-		{
-			CollectionGroupNames = AttributeNode->GetCachedCollectionGroupNames();
-		}
-PRAGMA_ENABLE_DEPRECATION_WARNINGS
-
-		for (const FName& CollectionGroupName : CollectionGroupNames)
-		{
-			GroupNames.Add(MakeShareable(new FText(FText::FromName(CollectionGroupName))));
+			if (Cloth.IsValidClothCollectionGroupName(GroupName))  // Restrict to the cloth facade groups
+			{
+				GroupNames.Add(MakeShareable(new FText(FText::FromName(GroupName))));
+			}
 		}
 
 		return SNew(SVerticalBox)
diff --git a/Engine/Plugins/ChaosClothAssetEditor/Source/ChaosClothAssetDataflowNodes/Private/ChaosClothAsset/SelectionGroupCustomization.h b/Engine/Plugins/ChaosClothAssetEditor/Source/ChaosClothAssetDataflowNodes/Private/ChaosClothAsset/SelectionGroupCustomization.h
--- a/Engine/Plugins/ChaosClothAssetEditor/Source/ChaosClothAssetDataflowNodes/Private/ChaosClothAsset/SelectionGroupCustomization.h
+++ b/Engine/Plugins/ChaosClothAssetEditor/Source/ChaosClothAssetDataflowNodes/Private/ChaosClothAsset/SelectionGroupCustomization.h
@@ -9,6 +9,7 @@ class ITableRow;
 class SComboButton;
 class STableViewBase;
 class SWidget;
+class SDataflowGraphEditor;
 
 namespace UE::Chaos::ClothAsset
 {
@@ -34,6 +35,7 @@ namespace UE::Chaos::ClothAsset
 		template<typename T>
 		T* GetOwnerStruct() const;
 
+		TWeakPtr<const SDataflowGraphEditor> DataflowGraphEditor;
 		TSharedPtr<IPropertyHandle> ChildPropertyHandle;
 		TSharedPtr<SComboButton> ComboButton;
 		TArray<TSharedPtr<FText>> GroupNames;
diff --git a/Engine/Plugins/ChaosClothAssetEditor/Source/ChaosClothAssetDataflowNodes/Private/ChaosClothAsset/SelectionNode.cpp b/Engine/Plugins/ChaosClothAssetEditor/Source/ChaosClothAssetDataflowNodes/Private/ChaosClothAsset/SelectionNode.cpp
--- a/Engine/Plugins/ChaosClothAssetEditor/Source/ChaosClothAssetDataflowNodes/Private/ChaosClothAsset/SelectionNode.cpp
+++ b/Engine/Plugins/ChaosClothAssetEditor/Source/ChaosClothAssetDataflowNodes/Private/ChaosClothAsset/SelectionNode.cpp
@@ -275,9 +275,9 @@ private:
 
 FChaosClothAssetSelectionNode_v2::FChaosClothAssetSelectionNode_v2(const UE::Dataflow::FNodeParameters& InParam, FGuid InGuid)
 	: FDataflowNode(InParam, InGuid)
-	, Import(FSimpleDelegate::CreateRaw(this, &FChaosClothAssetSelectionNode_v2::OnImport))
-	, ImportSecondary(FSimpleDelegate::CreateRaw(this, &FChaosClothAssetSelectionNode_v2::OnImportSecondary))
-	, Transfer(FSimpleDelegate::CreateRaw(this, &FChaosClothAssetSelectionNode_v2::OnTransfer))
+	, Import(FDataflowFunctionProperty::FDelegate::CreateRaw(this, &FChaosClothAssetSelectionNode_v2::OnImport))
+	, ImportSecondary(FDataflowFunctionProperty::FDelegate::CreateRaw(this, &FChaosClothAssetSelectionNode_v2::OnImportSecondary))
+	, Transfer(FDataflowFunctionProperty::FDelegate::CreateRaw(this, &FChaosClothAssetSelectionNode_v2::OnTransfer))
 {
 	RegisterInputConnection(&Collection);
 	RegisterInputConnection(&InputName.StringValue, GET_MEMBER_NAME_CHECKED(FChaosClothAssetConnectableIStringValue, StringValue))
@@ -290,13 +290,10 @@ FChaosClothAssetSelectionNode_v2::FChaosClothAssetSelectionNode_v2(const UE::Dat
 	RegisterOutputConnection(&OutputName.StringValue, (FString*)nullptr, GET_MEMBER_NAME_CHECKED(FChaosClothAssetConnectableOStringValue, StringValue));
 }
 
-void FChaosClothAssetSelectionNode_v2::OnImport()
+void FChaosClothAssetSelectionNode_v2::OnImport(UE::Dataflow::FContext& Context)
 {
 	using namespace UE::Chaos::ClothAsset;
 
-	// Create a temporary context (until we find a way to re-use the one from the tool calling this function)
-	UE::Dataflow::FContextThreaded Context;
-
 	FManagedArrayCollection InCollection = GetValue<FManagedArrayCollection>(Context, &Collection);
 	const TSharedRef<FManagedArrayCollection> ClothCollection = MakeShared<FManagedArrayCollection>(MoveTemp(InCollection));
 	const FCollectionClothSelectionConstFacade SelectionFacade(ClothCollection);
@@ -318,13 +315,10 @@ void FChaosClothAssetSelectionNode_v2::OnImport()
 	}
 }
 
-void FChaosClothAssetSelectionNode_v2::OnImportSecondary()
+void FChaosClothAssetSelectionNode_v2::OnImportSecondary(UE::Dataflow::FContext& Context)
 {
 	using namespace UE::Chaos::ClothAsset;
 
-	// Create a temporary context (until we find a way to re-use the one from the tool calling this function)
-	UE::Dataflow::FContextThreaded Context;
-
 	FManagedArrayCollection InCollection = GetValue<FManagedArrayCollection>(Context, &Collection);
 	const TSharedRef<FManagedArrayCollection> ClothCollection = MakeShared<FManagedArrayCollection>(MoveTemp(InCollection));
 	const FCollectionClothSelectionConstFacade SelectionFacade(ClothCollection);
@@ -346,13 +340,10 @@ void FChaosClothAssetSelectionNode_v2::OnImportSecondary()
 	}
 }
 
-void FChaosClothAssetSelectionNode_v2::OnTransfer()
+void FChaosClothAssetSelectionNode_v2::OnTransfer(UE::Dataflow::FContext& Context)
 {
 	using namespace UE::Chaos::ClothAsset;
 
-	// Create a temporary context (until we find a way to re-use the one from the tool calling this function)
-	UE::Dataflow::FContextThreaded Context;
-
 	// Transfer selection if the transfer collection input has changed and is valid
 	FManagedArrayCollection InCollection = GetValue<FManagedArrayCollection>(Context, &Collection);
 	const TSharedRef<FManagedArrayCollection> ClothCollection = MakeShared<FManagedArrayCollection>(MoveTemp(InCollection));
@@ -460,33 +451,6 @@ void FChaosClothAssetSelectionNode_v2::Evaluate(UE::Dataflow::FContext& Context,
 	}
 }
 
-void FChaosClothAssetSelectionNode_v2::OnSelected(UE::Dataflow::FContext& Context)
-{
-	using namespace UE::Chaos::ClothAsset;
-
-	// Re-evaluate the input collection
-	FManagedArrayCollection InCollection = GetValue<FManagedArrayCollection>(Context, &Collection);
-	const TSharedRef<FManagedArrayCollection> ClothCollection = MakeShared<FManagedArrayCollection>(MoveTemp(InCollection));
-	FCollectionClothFacade Cloth(ClothCollection);
-
-	// Update the list of used group for the UI customization
-	const TArray<FName> GroupNames = ClothCollection->GroupNames();
-	CachedCollectionGroupNames.Reset(GroupNames.Num());
-	for (const FName& GroupName : GroupNames)
-	{
-		if (Cloth.IsValidClothCollectionGroupName(GroupName))  // Restrict to the cloth facade groups
-		{
-			CachedCollectionGroupNames.Emplace(GroupName);
-		}
-	}
-}
-
-void FChaosClothAssetSelectionNode_v2::OnDeselected()
-{
-	// Clean up, to avoid another toolkit picking up the wrong context evaluation
-	CachedCollectionGroupNames.Reset();
-}
-
 FName FChaosClothAssetSelectionNode_v2::GetInputName(UE::Dataflow::FContext& Context) const
 {
 	FString InputNameString = GetValue<FString>(Context, &InputName.StringValue);
@@ -709,33 +673,6 @@ void FChaosClothAssetSelectionNode::Evaluate(UE::Dataflow::FContext& Context, co
 	}
 }
 
-void FChaosClothAssetSelectionNode::OnSelected(UE::Dataflow::FContext& Context)
-{
-	using namespace UE::Chaos::ClothAsset;
-
-	// Re-evaluate the input collection
-	FManagedArrayCollection InCollection = GetValue<FManagedArrayCollection>(Context, &Collection);
-	const TSharedRef<FManagedArrayCollection> ClothCollection = MakeShared<FManagedArrayCollection>(MoveTemp(InCollection));
-	FCollectionClothFacade Cloth(ClothCollection);
-
-	// Update the list of used group for the UI customization
-	const TArray<FName> GroupNames = ClothCollection->GroupNames();
-	CachedCollectionGroupNames.Reset(GroupNames.Num());
-	for (const FName& GroupName : GroupNames)
-	{
-		if (Cloth.IsValidClothCollectionGroupName(GroupName))  // Restrict to the cloth facade groups
-		{
-			CachedCollectionGroupNames.Emplace(GroupName);
-		}
-	}
-}
-
-void FChaosClothAssetSelectionNode::OnDeselected()
-{
-	// Clean up, to avoid another toolkit picking up the wrong context evaluation
-	CachedCollectionGroupNames.Reset();
-}
-
 void FChaosClothAssetSelectionNode::Serialize(FArchive& Ar)
 {
 	using namespace UE::Chaos::ClothAsset;
diff --git a/Engine/Plugins/ChaosClothAssetEditor/Source/ChaosClothAssetDataflowNodes/Private/ChaosClothAsset/USDImportNode_v2.cpp b/Engine/Plugins/ChaosClothAssetEditor/Source/ChaosClothAssetDataflowNodes/Private/ChaosClothAsset/USDImportNode_v2.cpp
--- a/Engine/Plugins/ChaosClothAssetEditor/Source/ChaosClothAssetDataflowNodes/Private/ChaosClothAsset/USDImportNode_v2.cpp
+++ b/Engine/Plugins/ChaosClothAssetEditor/Source/ChaosClothAssetDataflowNodes/Private/ChaosClothAsset/USDImportNode_v2.cpp
@@ -211,7 +211,7 @@ namespace UE::Chaos::ClothAsset::Private
 FChaosClothAssetUSDImportNode_v2::FChaosClothAssetUSDImportNode_v2(const UE::Dataflow::FNodeParameters& InParam, FGuid InGuid)
 	: FDataflowNode(InParam, InGuid)
 	, UsdFile(
-		FSimpleDelegate::CreateLambda([this, OwningObject = InParam.OwningObject]()
+		FDataflowFunctionProperty::FDelegate::CreateLambda([this, OwningObject = InParam.OwningObject](UE::Dataflow::FContext& /*Context*/)
 			{
 				const FString AssetPath = OwningObject ? OwningObject->GetPackage()->GetPathName() : FString();
 				FText ErrorText;
@@ -224,12 +224,12 @@ FChaosClothAssetUSDImportNode_v2::FChaosClothAssetUSDImportNode_v2(const UE::Dat
 				}
 			}))
 	, ReimportUsdFile(
-		FSimpleDelegate::CreateLambda([this]()
+		FDataflowFunctionProperty::FDelegate::CreateLambda([this](UE::Dataflow::FContext& Context)
 			{
-				UsdFile.Execute();
+				UsdFile.Execute(Context);
 			}))
 	, ReloadSimStaticMesh(
-		FSimpleDelegate::CreateLambda([this]()
+		FDataflowFunctionProperty::FDelegate::CreateLambda([this](UE::Dataflow::FContext& /*Context*/)
 			{
 				const TSharedRef<FManagedArrayCollection> ClothCollection = MakeShared<FManagedArrayCollection>(MoveTemp(Collection));
 				FText ErrorText;
@@ -242,7 +242,7 @@ FChaosClothAssetUSDImportNode_v2::FChaosClothAssetUSDImportNode_v2(const UE::Dat
 				Collection = MoveTemp(*ClothCollection);
 			}))
 	, ReloadRenderStaticMesh(
-		FSimpleDelegate::CreateLambda([this]()
+		FDataflowFunctionProperty::FDelegate::CreateLambda([this](UE::Dataflow::FContext& /*Context*/)
 			{
 				const TSharedRef<FManagedArrayCollection> ClothCollection = MakeShared<FManagedArrayCollection>(MoveTemp(Collection));
 				FText ErrorText;
diff --git a/Engine/Plugins/ChaosClothAssetEditor/Source/ChaosClothAssetDataflowNodes/Private/ChaosClothAsset/WeightMapNode.cpp b/Engine/Plugins/ChaosClothAssetEditor/Source/ChaosClothAssetDataflowNodes/Private/ChaosClothAsset/WeightMapNode.cpp
--- a/Engine/Plugins/ChaosClothAssetEditor/Source/ChaosClothAssetDataflowNodes/Private/ChaosClothAsset/WeightMapNode.cpp
+++ b/Engine/Plugins/ChaosClothAssetEditor/Source/ChaosClothAssetDataflowNodes/Private/ChaosClothAsset/WeightMapNode.cpp
@@ -36,7 +36,7 @@ namespace UE::Chaos::ClothAsset::Private
 
 FChaosClothAssetWeightMapNode::FChaosClothAssetWeightMapNode(const UE::Dataflow::FNodeParameters& InParam, FGuid InGuid)
 	: FDataflowNode(InParam, InGuid)
-	, Transfer(FSimpleDelegate::CreateRaw(this, &FChaosClothAssetWeightMapNode::OnTransfer))
+	, Transfer(FDataflowFunctionProperty::FDelegate::CreateRaw(this, &FChaosClothAssetWeightMapNode::OnTransfer))
 {
 	RegisterInputConnection(&Collection);
 	RegisterInputConnection(&InputName.StringValue, GET_MEMBER_NAME_CHECKED(FChaosClothAssetConnectableIStringValue, StringValue))
@@ -49,13 +49,10 @@ FChaosClothAssetWeightMapNode::FChaosClothAssetWeightMapNode(const UE::Dataflow:
 	RegisterOutputConnection(&OutputName.StringValue, (FString*)nullptr, GET_MEMBER_NAME_CHECKED(FChaosClothAssetConnectableOStringValue, StringValue));
 }
 
-void FChaosClothAssetWeightMapNode::OnTransfer()
+void FChaosClothAssetWeightMapNode::OnTransfer(UE::Dataflow::FContext& Context)
 {
 	using namespace UE::Chaos::ClothAsset;
 
-	// Create a temporary context (until we find a way to re-use the one from the tool calling this function)
-	UE::Dataflow::FContextThreaded Context;
-
 	// Transfer weight map if the transfer collection input has changed and is valid
 	FManagedArrayCollection InCollection = GetValue<FManagedArrayCollection>(Context, &Collection);
 	const TSharedRef<FManagedArrayCollection> ClothCollection = MakeShared<FManagedArrayCollection>(MoveTemp(InCollection));
diff --git a/Engine/Plugins/ChaosClothAssetEditor/Source/ChaosClothAssetDataflowNodes/Public/ChaosClothAsset/AttributeNode.h b/Engine/Plugins/ChaosClothAssetEditor/Source/ChaosClothAssetDataflowNodes/Public/ChaosClothAsset/AttributeNode.h
--- a/Engine/Plugins/ChaosClothAssetEditor/Source/ChaosClothAssetDataflowNodes/Public/ChaosClothAsset/AttributeNode.h
+++ b/Engine/Plugins/ChaosClothAssetEditor/Source/ChaosClothAssetDataflowNodes/Public/ChaosClothAsset/AttributeNode.h
@@ -30,7 +30,7 @@ enum class EChaosClothAssetNodeAttributeType : uint8
 
 /** Create a new attribute for the specified group. */
 USTRUCT(Meta = (DataflowCloth, Experimental))
-struct FChaosClothAssetAttributeNode_v2 : public FDataflowNode
+struct FChaosClothAssetAttributeNode_v2 final : public FDataflowNode
 {
 	GENERATED_USTRUCT_BODY()
 	DATAFLOW_NODE_DEFINE_INTERNAL(FChaosClothAssetAttributeNode_v2, "Attribute", "Cloth", "Cloth Attribute")
@@ -67,17 +67,10 @@ public:
 
 	FChaosClothAssetAttributeNode_v2(const UE::Dataflow::FNodeParameters& InParam, FGuid InGuid = FGuid::NewGuid());
 
-	/** Return a cached array of all the groups used by the input collection during at the time of the latest evaluation. */
-	const TArray<FName>& GetCachedCollectionGroupNames() const { return CachedCollectionGroupNames; }
-
 private:
 	//~ Begin FDataflowNode interface
 	virtual void Evaluate(UE::Dataflow::FContext& Context, const FDataflowOutput* Out) const override;
-	virtual void OnSelected(UE::Dataflow::FContext& Context) override;
-	virtual void OnDeselected() override;
 	//~ End FDataflowNode interface
-
-	TArray<FName> CachedCollectionGroupNames;
 };
 
 /** Create a new attribute for the specified group. */
@@ -119,13 +112,16 @@ public:
 	FChaosClothAssetAttributeNode(const UE::Dataflow::FNodeParameters& InParam, FGuid InGuid = FGuid::NewGuid());
 
 	/** Return a cached array of all the groups used by the input collection during at the time of the latest evaluation. */
+	UE_DEPRECATED(5.5, "This function is deprecated and will now return an empty array.")
 	const TArray<FName>& GetCachedCollectionGroupNames() const { return CachedCollectionGroupNames; }
 
 private:
 	//~ Begin FDataflowNode interface
 	virtual void Evaluate(UE::Dataflow::FContext& Context, const FDataflowOutput* Out) const override;
-	virtual void OnSelected(UE::Dataflow::FContext& Context) override;
-	virtual void OnDeselected() override;
+	UE_DEPRECATED(5.5, "This function is deprecated and will not be called on selection/deselection.")
+	virtual void OnSelected(UE::Dataflow::FContext& Context) {}
+	UE_DEPRECATED(5.5, "This function is deprecated and will not be called on selection/deselection.")
+	virtual void OnDeselected() {}
 	//~ End FDataflowNode interface
 
 	TArray<FName> CachedCollectionGroupNames;
diff --git a/Engine/Plugins/ChaosClothAssetEditor/Source/ChaosClothAssetDataflowNodes/Public/ChaosClothAsset/DeleteElementNode.h b/Engine/Plugins/ChaosClothAssetEditor/Source/ChaosClothAssetDataflowNodes/Public/ChaosClothAsset/DeleteElementNode.h
--- a/Engine/Plugins/ChaosClothAssetEditor/Source/ChaosClothAssetDataflowNodes/Public/ChaosClothAsset/DeleteElementNode.h
+++ b/Engine/Plugins/ChaosClothAssetEditor/Source/ChaosClothAssetDataflowNodes/Public/ChaosClothAsset/DeleteElementNode.h
@@ -67,12 +67,15 @@ public:
 	FChaosClothAssetDeleteElementNode(const UE::Dataflow::FNodeParameters& InParam, FGuid InGuid = FGuid::NewGuid());
 
 	/** Return a cached array of all the groups used by the input collection during at the time of the latest evaluation. */
+	UE_DEPRECATED(5.5, "This function is deprecated and will now return an empty array.")
 	const TArray<FName>& GetCachedCollectionGroupNames() const { return CachedCollectionGroupNames; }
 
 private:
 	virtual void Evaluate(UE::Dataflow::FContext& Context, const FDataflowOutput* Out) const override;
-	virtual void OnSelected(UE::Dataflow::FContext& Context) override;
-	virtual void OnDeselected() override;
+	UE_DEPRECATED(5.5, "This function is deprecated and will not be called on selection/deselection.")
+	virtual void OnSelected(UE::Dataflow::FContext& Context) {}
+	UE_DEPRECATED(5.5, "This function is deprecated and will not be called on selection/deselection.")
+	virtual void OnDeselected() {}
 	virtual void Serialize(FArchive& Ar);
 
 	TArray<FName> CachedCollectionGroupNames;
diff --git a/Engine/Plugins/ChaosClothAssetEditor/Source/ChaosClothAssetDataflowNodes/Public/ChaosClothAsset/ImportFilePath.h b/Engine/Plugins/ChaosClothAssetEditor/Source/ChaosClothAssetDataflowNodes/Public/ChaosClothAsset/ImportFilePath.h
--- a/Engine/Plugins/ChaosClothAssetEditor/Source/ChaosClothAssetDataflowNodes/Public/ChaosClothAsset/ImportFilePath.h
+++ b/Engine/Plugins/ChaosClothAssetEditor/Source/ChaosClothAssetDataflowNodes/Public/ChaosClothAsset/ImportFilePath.h
@@ -2,9 +2,14 @@
 
 #pragma once
 
-#include "Delegates/Delegate.h"
+#include "Delegates/DelegateCombinations.h"
 #include "ImportFilePath.generated.h"
 
+namespace UE::Dataflow
+{
+	class FContext;
+}
+
 PRAGMA_DISABLE_DEPRECATION_WARNINGS  // For bForceReimport
 USTRUCT()
 struct FChaosClothAssetImportFilePath
@@ -13,6 +18,8 @@ PRAGMA_ENABLE_DEPRECATION_WARNINGS
 	GENERATED_USTRUCT_BODY()
 
 public:
+	DECLARE_DELEGATE_OneParam(FDelegate, UE::Dataflow::FContext&);
+
 	UPROPERTY(EditAnywhere, Category = "Import File Path")
 	FString FilePath;
 
@@ -22,13 +29,10 @@ public:
 
 	FChaosClothAssetImportFilePath() = default;
 
-	explicit FChaosClothAssetImportFilePath(FSimpleDelegate&& InDelegate) : Delegate(MoveTemp(InDelegate)) {}
+	explicit FChaosClothAssetImportFilePath(FDelegate&& InDelegate) : Delegate(MoveTemp(InDelegate)) {}
 
-	void Execute() const
-	{
-		Delegate.ExecuteIfBound();
-	}
+	void Execute(UE::Dataflow::FContext& Context) const { Delegate.ExecuteIfBound(Context); }
 
 private:
-	FSimpleDelegate Delegate;
+	FDelegate Delegate;
 };
diff --git a/Engine/Plugins/ChaosClothAssetEditor/Source/ChaosClothAssetDataflowNodes/Public/ChaosClothAsset/SelectionNode.h b/Engine/Plugins/ChaosClothAssetEditor/Source/ChaosClothAssetDataflowNodes/Public/ChaosClothAsset/SelectionNode.h
--- a/Engine/Plugins/ChaosClothAssetEditor/Source/ChaosClothAssetDataflowNodes/Public/ChaosClothAsset/SelectionNode.h
+++ b/Engine/Plugins/ChaosClothAssetEditor/Source/ChaosClothAssetDataflowNodes/Public/ChaosClothAsset/SelectionNode.h
@@ -55,7 +55,7 @@ enum class EChaosClothAssetSelectionOverrideType : uint8
 };
 
 USTRUCT(Meta = (DataflowCloth))
-struct FChaosClothAssetSelectionNode_v2 : public FDataflowNode
+struct FChaosClothAssetSelectionNode_v2 final : public FDataflowNode
 {
 	GENERATED_USTRUCT_BODY()
 	DATAFLOW_NODE_DEFINE_INTERNAL(FChaosClothAssetSelectionNode_v2, "Selection", "Cloth", "Cloth Selection")
@@ -126,9 +126,6 @@ public:
 
 	FChaosClothAssetSelectionNode_v2(const UE::Dataflow::FNodeParameters& InParam, FGuid InGuid = FGuid::NewGuid());
 
-	/** Return a cached array of all the groups used by the input collection during at the time of the latest evaluation. */
-	const TArray<FName>& GetCachedCollectionGroupNames() const { return CachedCollectionGroupNames; }
-
 private:
 	friend class UClothMeshSelectionTool;
 
@@ -141,15 +138,11 @@ private:
 
 	//~ Begin FDataflowNode implementation
 	virtual void Evaluate(UE::Dataflow::FContext& Context, const FDataflowOutput* Out) const override;
-	virtual void OnSelected(UE::Dataflow::FContext& Context) override;
-	virtual void OnDeselected() override;
 	//~ End FDataflowNode implementation
 
-	void OnImport();
-	void OnImportSecondary();
-	void OnTransfer();
-
-	TArray<FName> CachedCollectionGroupNames;
+	void OnImport(UE::Dataflow::FContext& Context);
+	void OnImportSecondary(UE::Dataflow::FContext& Context);
+	void OnTransfer(UE::Dataflow::FContext& Context);
 };
 
 PRAGMA_DISABLE_DEPRECATION_WARNINGS  // For EChaosClothAssetSelectionType
@@ -236,6 +229,7 @@ public:
 	FChaosClothAssetSelectionNode(const UE::Dataflow::FNodeParameters& InParam, FGuid InGuid = FGuid::NewGuid());
 
 	/** Return a cached array of all the groups used by the input collection during at the time of the latest evaluation. */
+	UE_DEPRECATED(5.5, "This function is deprecated and will now return an empty array.")
 	const TArray<FName>& GetCachedCollectionGroupNames() const { return CachedCollectionGroupNames; }
 
 	FName CHAOSCLOTHASSETDATAFLOWNODES_API GetInputName(UE::Dataflow::FContext& Context) const;
@@ -250,8 +244,10 @@ public:
 private:
 	virtual void SetAssetValue(TObjectPtr<UObject> Asset, UE::Dataflow::FContext& Context) const override;
 	virtual void Evaluate(UE::Dataflow::FContext& Context, const FDataflowOutput* Out) const override;
-	virtual void OnSelected(UE::Dataflow::FContext& Context) override;
-	virtual void OnDeselected() override;
+	UE_DEPRECATED(5.5, "This function is deprecated and will not be called on selection/deselection.")
+	virtual void OnSelected(UE::Dataflow::FContext& /*Context*/) {}
+	UE_DEPRECATED(5.5, "This function is deprecated and will not be called on selection/deselection.")
+	virtual void OnDeselected() {}
 	virtual void Serialize(FArchive& Ar);
 
 	TArray<FName> CachedCollectionGroupNames;
diff --git a/Engine/Plugins/ChaosClothAssetEditor/Source/ChaosClothAssetDataflowNodes/Public/ChaosClothAsset/WeightMapNode.h b/Engine/Plugins/ChaosClothAssetEditor/Source/ChaosClothAssetDataflowNodes/Public/ChaosClothAsset/WeightMapNode.h
--- a/Engine/Plugins/ChaosClothAssetEditor/Source/ChaosClothAssetDataflowNodes/Public/ChaosClothAsset/WeightMapNode.h
+++ b/Engine/Plugins/ChaosClothAssetEditor/Source/ChaosClothAssetDataflowNodes/Public/ChaosClothAsset/WeightMapNode.h
@@ -88,7 +88,7 @@ private:
 	virtual FDataflowOutput* RedirectSerializedOutput(const FName& MissingOutputName) override;
 	//~ End FDataflowNode interface
 
-	void OnTransfer();
+	void OnTransfer(UE::Dataflow::FContext& Context);
 
 	UPROPERTY()
 	TArray<float> VertexWeights;
diff --git a/Engine/Plugins/ChaosClothAssetEditor/Source/ChaosClothAssetEditor/Private/ChaosClothAsset/ClothEditorToolkit.cpp b/Engine/Plugins/ChaosClothAssetEditor/Source/ChaosClothAssetEditor/Private/ChaosClothAsset/ClothEditorToolkit.cpp
--- a/Engine/Plugins/ChaosClothAssetEditor/Source/ChaosClothAssetEditor/Private/ChaosClothAsset/ClothEditorToolkit.cpp
+++ b/Engine/Plugins/ChaosClothAssetEditor/Source/ChaosClothAssetEditor/Private/ChaosClothAsset/ClothEditorToolkit.cpp
@@ -48,6 +48,42 @@
 
 #define LOCTEXT_NAMESPACE "ChaosClothAssetEditorToolkit"
 
+// Same as SDataflowGraphEditor, only adds a ClothAssetEditorToolkit pointer to correctly override the GetDataflowContext() method.
+class SClothAssetDataflowGraphEditor : public SDataflowGraphEditor
+{
+public:
+	SLATE_BEGIN_ARGS(SClothAssetDataflowGraphEditor) {}
+	SLATE_ARGUMENT_DEFAULT(UEdGraph*, GraphToEdit) = nullptr;
+	SLATE_ARGUMENT(FGraphEditorEvents, GraphEvents)
+	SLATE_ARGUMENT(TSharedPtr<IStructureDetailsView>, DetailsView)
+	SLATE_ARGUMENT(FDataflowEditorCommands::FGraphEvaluationCallback, EvaluateGraph)
+	SLATE_ARGUMENT_DEFAULT(UE::Chaos::ClothAsset::FChaosClothAssetEditorToolkit*, ClothAssetEditorToolkit) = nullptr;
+	SLATE_END_ARGS()
+
+	void Construct(const FArguments& InArgs, UObject* InAssetOwner)
+	{
+		check(InArgs._GraphToEdit);
+		ClothAssetEditorToolkit = InArgs._ClothAssetEditorToolkit;
+
+		SDataflowGraphEditor::FArguments Arguments;
+		Arguments._GraphToEdit = InArgs._GraphToEdit;
+		Arguments._GraphEvents = InArgs._GraphEvents;
+		Arguments._DetailsView = InArgs._DetailsView;
+		Arguments._EvaluateGraph = InArgs._EvaluateGraph;
+		SDataflowGraphEditor::Construct(Arguments, InAssetOwner);
+	}
+
+	//~ Begin FDataflowSNodeInterface interface
+	virtual TSharedPtr<UE::Dataflow::FContext> GetDataflowContext() const override
+	{
+		return ensure(ClothAssetEditorToolkit) ? ClothAssetEditorToolkit->GetDataflowContext() : TSharedPtr<UE::Dataflow::FContext>();
+	}
+	//~ End FDataflowSNodeInterface interface
+
+private:
+	UE::Chaos::ClothAsset::FChaosClothAssetEditorToolkit* ClothAssetEditorToolkit = nullptr;
+};
+
 namespace UE::Chaos::ClothAsset
 {
 const FName FChaosClothAssetEditorToolkit::ClothPreviewTabID(TEXT("ChaosClothAssetEditor_ClothPreviewTab"));
@@ -188,7 +224,6 @@ FChaosClothAssetEditorToolkit::~FChaosClothAssetEditorToolkit()
 	if (SelectedDataflowNode && OnNodeInvalidatedDelegateHandle.IsValid())
 	{
 		SelectedDataflowNode->GetOnNodeInvalidatedDelegate().Remove(OnNodeInvalidatedDelegateHandle);
-		SelectedDataflowNode->OnDeselected();
 	}
 	SelectedDataflowNodeGuid.Invalidate();
 
@@ -370,7 +405,6 @@ bool FChaosClothAssetEditorToolkit::OnRequestClose(EAssetEditorCloseReason InClo
 	if (SelectedDataflowNode && OnNodeInvalidatedDelegateHandle.IsValid())
 	{
 		SelectedDataflowNode->GetOnNodeInvalidatedDelegate().Remove(OnNodeInvalidatedDelegateHandle);
-		SelectedDataflowNode->OnDeselected();
 	}
 	SelectedDataflowNodeGuid.Invalidate();
 
@@ -905,12 +939,12 @@ TSharedRef<SDataflowGraphEditor> FChaosClothAssetEditorToolkit::CreateGraphEdito
 	InEvents.OnTextCommitted = FOnNodeTextCommitted::CreateSP(this, &FChaosClothAssetEditorToolkit::OnNodeTitleCommitted);
 	InEvents.OnNodeSingleClicked = SGraphEditor::FOnNodeSingleClicked::CreateSP(this, &FChaosClothAssetEditorToolkit::OnNodeSingleClicked);
 
-	TSharedRef<SDataflowGraphEditor> NewGraphEditor = SNew(SDataflowGraphEditor, Dataflow)
+	TSharedRef<SDataflowGraphEditor> NewGraphEditor = SNew(SClothAssetDataflowGraphEditor, Dataflow)
 		.GraphToEdit(Dataflow)
 		.GraphEvents(InEvents)
 		.DetailsView(NodeDetailsEditor)
 		.EvaluateGraph(EvalLambda)
-		.DataflowEditor(nullptr);
+		.ClothAssetEditorToolkit(this);
 
 	NewGraphEditor->OnSelectionChangedMulticast.AddSP(this, &FChaosClothAssetEditorToolkit::OnNodeSelectionChanged);
 	NewGraphEditor->OnNodeDeletedMulticast.AddSP(this, &FChaosClothAssetEditorToolkit::OnNodeDeleted);
@@ -1137,7 +1171,6 @@ void FChaosClothAssetEditorToolkit::OnNodeSelectionChanged(const TSet<UObject*>&
 				{
 					SelectedDataflowNode->GetOnNodeInvalidatedDelegate().Remove(OnNodeInvalidatedDelegateHandle);
 				}
-				SelectedDataflowNode->OnDeselected();
 				SelectedDataflowNodeGuid.Invalidate();
 			}
 		}
@@ -1160,7 +1193,6 @@ void FChaosClothAssetEditorToolkit::OnNodeSelectionChanged(const TSet<UObject*>&
 							{
 								SelectedDataflowNode->GetOnNodeInvalidatedDelegate().Remove(OnNodeInvalidatedDelegateHandle);
 							}
-							SelectedDataflowNode->OnDeselected();
 						}
 						SelectedDataflowNode = Node->GetDataflowNode();
 						SelectedDataflowNodeGuid = SelectedDataflowNode->GetGuid();
@@ -1169,10 +1201,6 @@ void FChaosClothAssetEditorToolkit::OnNodeSelectionChanged(const TSet<UObject*>&
 						{
 							InputCollection = GetInputClothCollectionIfPossible(SelectedDataflowNode, DataflowContext);
 							Collection = GetClothCollectionIfPossible(SelectedDataflowNode, DataflowContext);
-							if (DataflowContext.IsValid())  // TODO: The context shouldn't be nullptr on the first node creation
-							{
-								SelectedDataflowNode->OnSelected(*DataflowContext);
-							}
 
 							// Set a callback to re-evaluate the node if it is invalidated
 							OnNodeInvalidatedDelegateHandle = SelectedDataflowNode->GetOnNodeInvalidatedDelegate().AddLambda(
@@ -1192,11 +1220,6 @@ void FChaosClothAssetEditorToolkit::OnNodeSelectionChanged(const TSet<UObject*>&
 											constexpr bool bDeferDynamicMeshInit = false;
 											ClothMode->SetSelectedClothCollection(Collection, InputCollection, bDeferDynamicMeshInit);
 										}
-
-										if (DataflowContext.IsValid())  // TODO: The context shouldn't be nullptr on the first node creation
-										{
-											SelectedDataflowNode->OnSelected(*DataflowContext);
-										}
 									}
 									// The detail panel won't update correctly when a change affects a node's properties and needs a refresh at the next tick
 									TickCommands.AddLambda([this]()
diff --git a/Engine/Plugins/Experimental/Dataflow/Source/DataflowEditor/Private/Dataflow/DataflowEditorToolkit.cpp b/Engine/Plugins/Experimental/Dataflow/Source/DataflowEditor/Private/Dataflow/DataflowEditorToolkit.cpp
--- a/Engine/Plugins/Experimental/Dataflow/Source/DataflowEditor/Private/Dataflow/DataflowEditorToolkit.cpp
+++ b/Engine/Plugins/Experimental/Dataflow/Source/DataflowEditor/Private/Dataflow/DataflowEditorToolkit.cpp
@@ -967,8 +967,6 @@ void FDataflowEditorToolkit::OnNodeSelectionChanged(const TSet<UObject*>& InNewS
 				{
 					if (const TSharedPtr<FDataflowNode> DataflowNode = PrimarySelection->GetDataflowNode())
 					{
-						DataflowNode->OnSelected(*DataflowContext);
-
 						// Update selected Collection in the ContextObject
 						for (const FDataflowOutput* const Output : DataflowNode->GetOutputs())
 						{
diff --git a/Engine/Plugins/Experimental/Dataflow/Source/DataflowEditor/Private/Dataflow/DataflowFunctionPropertyCustomization.cpp b/Engine/Plugins/Experimental/Dataflow/Source/DataflowEditor/Private/Dataflow/DataflowFunctionPropertyCustomization.cpp
--- a/Engine/Plugins/Experimental/Dataflow/Source/DataflowEditor/Private/Dataflow/DataflowFunctionPropertyCustomization.cpp
+++ b/Engine/Plugins/Experimental/Dataflow/Source/DataflowEditor/Private/Dataflow/DataflowFunctionPropertyCustomization.cpp
@@ -2,6 +2,8 @@
 
 #include "Dataflow/DataflowFunctionPropertyCustomization.h"
 #include "Dataflow/DataflowFunctionProperty.h"
+#include "Dataflow/DataflowGraphEditor.h"
+#include "Dataflow/DataflowNodeParameters.h"
 #include "Widgets/Images/SImage.h"
 #include "Widgets/Input/SButton.h"
 #include "Widgets/Layout/SWrapBox.h"
@@ -44,6 +46,10 @@ namespace UE::Dataflow
 	void FFunctionPropertyCustomization::CustomizeHeader(TSharedRef<IPropertyHandle> StructPropertyHandle, FDetailWidgetRow& HeaderRow, IPropertyTypeCustomizationUtils& CustomizationUtils)
 	{
 		StructProperty = StructPropertyHandle;
+
+		// Keep a weak pointer to the graph editor that is creating this customization
+		DataflowGraphEditor = SDataflowGraphEditor::GetSelectedGraphEditor();
+
 		// Find all functions under the same category
 		const FName CategoryName = StructProperty->GetDefaultCategoryName();
 		const TSharedPtr<IPropertyHandle> OwnerProperty = StructProperty->GetParentHandle();
@@ -81,11 +87,19 @@ namespace UE::Dataflow
 						];
 					}
 
-					auto OnClicked = [ChildProperty, FunctionProperty]() -> FReply
+					auto OnClicked = [this, ChildProperty, FunctionProperty]() -> FReply
 						{
+							// Retrieve context if any
+							const TSharedPtr<const SDataflowGraphEditor> DataflowGraphEditorPtr = DataflowGraphEditor.Pin();
+							const TSharedPtr<UE::Dataflow::FContext> Context = DataflowGraphEditorPtr ? DataflowGraphEditorPtr->GetDataflowContext() : TSharedPtr<UE::Dataflow::FContext>();
+							
 							// Execute function
-							FunctionProperty->Execute();
-							ChildProperty->NotifyFinishedChangingProperties();  // Triggers node invalidation
+							UE::Dataflow::FContextThreaded EmptyContext;
+							FunctionProperty->Execute(Context.IsValid() ? *Context : EmptyContext);
+
+							// Triggers node invalidation
+							ChildProperty->NotifyFinishedChangingProperties();
+
 							return FReply::Handled();
 						};
 
diff --git a/Engine/Plugins/Experimental/Dataflow/Source/DataflowEditor/Private/Dataflow/DataflowFunctionPropertyCustomization.h b/Engine/Plugins/Experimental/Dataflow/Source/DataflowEditor/Private/Dataflow/DataflowFunctionPropertyCustomization.h
--- a/Engine/Plugins/Experimental/Dataflow/Source/DataflowEditor/Private/Dataflow/DataflowFunctionPropertyCustomization.h
+++ b/Engine/Plugins/Experimental/Dataflow/Source/DataflowEditor/Private/Dataflow/DataflowFunctionPropertyCustomization.h
@@ -4,6 +4,8 @@
 
 #include "IPropertyTypeCustomization.h"
 
+class SDataflowGraphEditor;
+
 namespace UE::Dataflow
 {
 	/**
@@ -20,5 +22,6 @@ namespace UE::Dataflow
 
 	private:
 		TSharedPtr<IPropertyHandle> StructProperty;
+		TWeakPtr<const SDataflowGraphEditor> DataflowGraphEditor;
 	};
 }
diff --git a/Engine/Plugins/Experimental/Dataflow/Source/DataflowEditor/Private/Dataflow/DataflowGraphEditor.cpp b/Engine/Plugins/Experimental/Dataflow/Source/DataflowEditor/Private/Dataflow/DataflowGraphEditor.cpp
--- a/Engine/Plugins/Experimental/Dataflow/Source/DataflowEditor/Private/Dataflow/DataflowGraphEditor.cpp
+++ b/Engine/Plugins/Experimental/Dataflow/Source/DataflowEditor/Private/Dataflow/DataflowGraphEditor.cpp
@@ -22,6 +22,7 @@
 #define LOCTEXT_NAMESPACE "DataflowGraphEditor"
 
 TSharedPtr<FDataflowGraphEditorNodeFactory> SDataflowGraphEditor::NodeFactory;
+TWeakPtr<SDataflowGraphEditor> SDataflowGraphEditor::SelectedGraphEditor;
 
 void SDataflowGraphEditor::Construct(const FArguments& InArgs, UObject* InAssetOwner)
 {
@@ -164,7 +165,7 @@ void SDataflowGraphEditor::Construct(const FArguments& InArgs, UObject* InAssetO
 	SetNodeFactory( MakeShared<FDataflowGraphNodeFactory>(this) );
 }
 
-const TSharedPtr<UE::Dataflow::FEngineContext> SDataflowGraphEditor::GetDataflowContext() const
+TSharedPtr<UE::Dataflow::FContext> SDataflowGraphEditor::GetDataflowContext() const
 {
 	if (DataflowEditor)
 	{
@@ -173,7 +174,7 @@ const TSharedPtr<UE::Dataflow::FEngineContext> SDataflowGraphEditor::GetDataflow
 			return DataflowEditor->GetEditorContent()->GetDataflowContext();
 		}
 	}
-	return TSharedPtr<UE::Dataflow::FEngineContext>(nullptr);
+	return TSharedPtr<UE::Dataflow::FContext>();
 }
 
 void SDataflowGraphEditor::EvaluateNode()
@@ -292,7 +293,11 @@ bool SDataflowGraphEditor::CanRenameNode() const
 
 void SDataflowGraphEditor::OnSelectedNodesChanged(const TSet<UObject*>& NewSelection)
 {
-	OnSelectionChangedMulticast.Broadcast(NewSelection);  // Broadcast the selection change before refreshing the DetailsView, the nodes' specific UI data have to be updated before the UI is being redrawn
+	// Set the currently selected graph editor before running any callback
+	ensureMsgf(!SelectedGraphEditor.IsValid(), TEXT("Two different editors cannot have their selection changed at once."));
+	SelectedGraphEditor = StaticCastSharedRef<SDataflowGraphEditor>(AsShared()).ToWeakPtr();
+
+	OnSelectionChangedMulticast.Broadcast(NewSelection);
 
 	if (DataflowAsset.Get() && DetailsView)
 	{
@@ -303,6 +308,9 @@ void SDataflowGraphEditor::OnSelectedNodesChanged(const TSet<UObject*>& NewSelec
 
 		FDataflowEditorCommands::OnSelectedNodesChanged(DetailsView, AssetOwner.Get(), DataflowAsset.Get(), AsObjectPointers(NewSelection) );
 	}
+
+	// Clear the current selected editor
+	SelectedGraphEditor.Reset();
 }
 
 FReply SDataflowGraphEditor::OnKeyDown(const FGeometry& MyGeometry, const FKeyEvent& InKeyEvent)
diff --git a/Engine/Plugins/Experimental/Dataflow/Source/DataflowEditor/Private/Dataflow/DataflowSNode.cpp b/Engine/Plugins/Experimental/Dataflow/Source/DataflowEditor/Private/Dataflow/DataflowSNode.cpp
--- a/Engine/Plugins/Experimental/Dataflow/Source/DataflowEditor/Private/Dataflow/DataflowSNode.cpp
+++ b/Engine/Plugins/Experimental/Dataflow/Source/DataflowEditor/Private/Dataflow/DataflowSNode.cpp
@@ -147,7 +147,7 @@ TSharedPtr<SGraphPin> SDataflowEdNode::CreatePinWidget(UEdGraphPin* Pin) const
 			{
 				if (FDataflowOutput* Output = DataflowNode->FindOutput(Pin->GetFName()))
 				{
-					if (const TSharedPtr<UE::Dataflow::FEngineContext> DataflowContext = DataflowInterface->GetDataflowContext())
+					if (const TSharedPtr<UE::Dataflow::FContext> DataflowContext = DataflowInterface->GetDataflowContext())
 					{
 						TSet<UE::Dataflow::FContextCacheKey> CacheKeys;
 						const int32 NumKeys = DataflowContext->GetKeys(CacheKeys);
diff --git a/Engine/Plugins/Experimental/Dataflow/Source/DataflowEditor/Private/Dataflow/ScalarVertexPropertyGroupCustomization.cpp b/Engine/Plugins/Experimental/Dataflow/Source/DataflowEditor/Private/Dataflow/ScalarVertexPropertyGroupCustomization.cpp
--- a/Engine/Plugins/Experimental/Dataflow/Source/DataflowEditor/Private/Dataflow/ScalarVertexPropertyGroupCustomization.cpp
+++ b/Engine/Plugins/Experimental/Dataflow/Source/DataflowEditor/Private/Dataflow/ScalarVertexPropertyGroupCustomization.cpp
@@ -127,7 +127,9 @@ namespace UE::Dataflow
 		TArray<FName> CollectionGroupNames;
 		if (const FDataflowCollectionAddScalarVertexPropertyNode* const ScalarVertexPropertyNode = GetOwnerStruct<FDataflowCollectionAddScalarVertexPropertyNode>())
 		{
-			CollectionGroupNames = ScalarVertexPropertyNode->GetCachedCollectionGroupNames();
+			// TODO: 
+			//   Fix this using the new context based paradigm (see SelectionGroupCustomization.cpp GetCollection() implementation).
+			//CollectionGroupNames = ScalarVertexPropertyNode->GetCachedCollectionGroupNames();
 			NodeType = ScalarVertexPropertyNode->GetType();
 
 			const TArray<FName> AvailableTargetGroups = DataflowAddScalarVertexPropertyCallbackRegistry::Get().GetTargetGroupNames();
diff --git a/Engine/Plugins/Experimental/Dataflow/Source/DataflowEditor/Public/Dataflow/DataflowGraphEditor.h b/Engine/Plugins/Experimental/Dataflow/Source/DataflowEditor/Public/Dataflow/DataflowGraphEditor.h
--- a/Engine/Plugins/Experimental/Dataflow/Source/DataflowEditor/Public/Dataflow/DataflowGraphEditor.h
+++ b/Engine/Plugins/Experimental/Dataflow/Source/DataflowEditor/Public/Dataflow/DataflowGraphEditor.h
@@ -45,7 +45,7 @@ public:
 
 	SLATE_ARGUMENT(TSharedPtr<FUICommandList>, AdditionalCommands)
 	SLATE_ATTRIBUTE(FGraphAppearanceInfo, Appearance)
-	SLATE_ARGUMENT(UEdGraph*, GraphToEdit)
+	SLATE_ARGUMENT_DEFAULT(UEdGraph*, GraphToEdit) = nullptr;
 	SLATE_ARGUMENT(FGraphEditorEvents, GraphEvents)
 	SLATE_ARGUMENT(TSharedPtr<IStructureDetailsView>, DetailsView)
 	SLATE_ARGUMENT(FDataflowEditorCommands::FGraphEvaluationCallback, EvaluateGraph)
@@ -146,7 +146,10 @@ public:
 	virtual FString GetReferencerName() const override { return TEXT("SDataflowGraphEditor"); }
 
 	/** FDataflowSNodeInterface */
-	virtual const TSharedPtr<UE::Dataflow::FEngineContext> GetDataflowContext() const override;
+	virtual TSharedPtr<UE::Dataflow::FContext> GetDataflowContext() const override;
+
+	/** Return the currently selected editor. Only valid for the duration of the OnSelectedNodesChanged callback where the property editor is updated. */
+	static const TWeakPtr<SDataflowGraphEditor>& GetSelectedGraphEditor() { return SelectedGraphEditor; }
 
 private:
 	/** Add an additional option pin to all selected Dataflow nodes for those that overrides the AddPin function. */
@@ -179,6 +182,9 @@ private:
 	/** Factory to create the associated SGraphNode classes for Dataprep graph's UEdGraph classes */
 	static TSharedPtr<FDataflowGraphEditorNodeFactory> NodeFactory;
 
+	/** The current graph editor when the selection callback is invoked. */
+	static TWeakPtr<SDataflowGraphEditor> SelectedGraphEditor;
+
 	/** Editor for the content */
 	UDataflowEditor* DataflowEditor = nullptr;
 
diff --git a/Engine/Plugins/Experimental/Dataflow/Source/DataflowEditor/Public/Dataflow/DataflowSEditorInterface.h b/Engine/Plugins/Experimental/Dataflow/Source/DataflowEditor/Public/Dataflow/DataflowSEditorInterface.h
--- a/Engine/Plugins/Experimental/Dataflow/Source/DataflowEditor/Public/Dataflow/DataflowSEditorInterface.h
+++ b/Engine/Plugins/Experimental/Dataflow/Source/DataflowEditor/Public/Dataflow/DataflowSEditorInterface.h
@@ -2,24 +2,23 @@
 
 #pragma once
 
-#include "CoreMinimal.h"
-#include "Dataflow/DataflowObjectInterface.h"
+#include "Templates/SharedPointerFwd.h"
+
+namespace UE::Dataflow
+{
+	class FContext;
+}
 
 /**
 * FDataflowSEditorInterface
 * 
 */
-
-class UDataflowBaseContent;
-class UDataflowEditor;
-
 class FDataflowSEditorInterface
 {
 public:
-	FDataflowSEditorInterface() {}
-
 	/** Dataflow editor content accessors */
-	virtual const TSharedPtr<UE::Dataflow::FEngineContext> GetDataflowContext() const { return TSharedPtr<UE::Dataflow::FEngineContext>(nullptr); };
+	virtual TSharedPtr<UE::Dataflow::FContext> GetDataflowContext() const = 0;
 
+protected:
+	FDataflowSEditorInterface() = default;
 };
-
diff --git a/Engine/Plugins/Experimental/Dataflow/Source/DataflowNodes/Private/Dataflow/DataflowCollectionAddScalarVertexPropertyNode.cpp b/Engine/Plugins/Experimental/Dataflow/Source/DataflowNodes/Private/Dataflow/DataflowCollectionAddScalarVertexPropertyNode.cpp
--- a/Engine/Plugins/Experimental/Dataflow/Source/DataflowNodes/Private/Dataflow/DataflowCollectionAddScalarVertexPropertyNode.cpp
+++ b/Engine/Plugins/Experimental/Dataflow/Source/DataflowNodes/Private/Dataflow/DataflowCollectionAddScalarVertexPropertyNode.cpp
@@ -114,26 +114,4 @@ void FDataflowCollectionAddScalarVertexPropertyNode::Evaluate(UE::Dataflow::FCon
 	}
 }
 
-void FDataflowCollectionAddScalarVertexPropertyNode::OnSelected(UE::Dataflow::FContext& Context)
-{
-	// Re-evaluate the input collection
-	FManagedArrayCollection InCollection = GetValue<FManagedArrayCollection>(Context, &Collection);
-	const TSharedRef<FManagedArrayCollection> ManagedArrayCollection = MakeShared<FManagedArrayCollection>(MoveTemp(InCollection));
-
-	// Update the list of used group for the UI customization
-	const TArray<FName> GroupNames = ManagedArrayCollection->GroupNames();
-	CachedCollectionGroupNames.Reset(GroupNames.Num());
-	for (const FName& GroupName : GroupNames)
-	{
-		CachedCollectionGroupNames.Emplace(GroupName);
-	}
-}
-
-void FDataflowCollectionAddScalarVertexPropertyNode::OnDeselected()
-{
-	// Clean up, to avoid another toolkit picking up the wrong context evaluation
-	CachedCollectionGroupNames.Reset();
-}
-
-
 #undef LOCTEXT_NAMESPACE
diff --git a/Engine/Plugins/Experimental/Dataflow/Source/DataflowNodes/Public/Dataflow/DataflowCollectionAddScalarVertexPropertyNode.h b/Engine/Plugins/Experimental/Dataflow/Source/DataflowNodes/Public/Dataflow/DataflowCollectionAddScalarVertexPropertyNode.h
--- a/Engine/Plugins/Experimental/Dataflow/Source/DataflowNodes/Public/Dataflow/DataflowCollectionAddScalarVertexPropertyNode.h
+++ b/Engine/Plugins/Experimental/Dataflow/Source/DataflowNodes/Public/Dataflow/DataflowCollectionAddScalarVertexPropertyNode.h
@@ -82,13 +82,6 @@ public:
 
 	FDataflowCollectionAddScalarVertexPropertyNode(const UE::Dataflow::FNodeParameters& InParam, FGuid InGuid = FGuid::NewGuid());
 
-	/** Return a cached array of all the groups used by the input collection during at the time of the latest evaluation. */
-	const TArray<FName>& GetCachedCollectionGroupNames() const { return CachedCollectionGroupNames; }
-
 private:
 	virtual void Evaluate(UE::Dataflow::FContext& Context, const FDataflowOutput* Out) const override;
-	virtual void OnSelected(UE::Dataflow::FContext& Context) override;
-	virtual void OnDeselected() override;
-
-	TArray<FName> CachedCollectionGroupNames;
 };
diff --git a/Engine/Plugins/Experimental/Dataflow/Source/DataflowNodes/Public/Dataflow/DataflowFunctionProperty.h b/Engine/Plugins/Experimental/Dataflow/Source/DataflowNodes/Public/Dataflow/DataflowFunctionProperty.h
--- a/Engine/Plugins/Experimental/Dataflow/Source/DataflowNodes/Public/Dataflow/DataflowFunctionProperty.h
+++ b/Engine/Plugins/Experimental/Dataflow/Source/DataflowNodes/Public/Dataflow/DataflowFunctionProperty.h
@@ -2,9 +2,15 @@
 
 #pragma once
 
-#include "Delegates/Delegate.h"
+#include "Delegates/DelegateCombinations.h"
+#include "Templates/SharedPointerFwd.h"
 #include "DataflowFunctionProperty.generated.h"
 
+namespace UE::Dataflow
+{
+	class FContext;
+}
+
 /**
  * Function property for all Dataflow nodes.
  * The structure is also used in DataFlow::FFunctionPropertyCustomization to appear as text and/or image buttons.
@@ -41,13 +47,14 @@ struct FDataflowFunctionProperty
 	GENERATED_BODY()
 
 public:
+	DECLARE_DELEGATE_OneParam(FDelegate, UE::Dataflow::FContext&);
 
 	FDataflowFunctionProperty() = default;
 
-	explicit FDataflowFunctionProperty(FSimpleDelegate&& InDelegate) : Delegate(MoveTemp(InDelegate)) {}
+	explicit FDataflowFunctionProperty(FDelegate&& InDelegate) : Delegate(MoveTemp(InDelegate)) {}
 
-	void Execute() const { Delegate.ExecuteIfBound(); }
+	void Execute(UE::Dataflow::FContext& Context) const { Delegate.ExecuteIfBound(Context); }
 
 private:
-	FSimpleDelegate Delegate;
+	FDelegate Delegate;
 };
diff --git a/Engine/Source/Runtime/Experimental/Dataflow/Core/Public/Dataflow/DataflowNode.h b/Engine/Source/Runtime/Experimental/Dataflow/Core/Public/Dataflow/DataflowNode.h
--- a/Engine/Source/Runtime/Experimental/Dataflow/Core/Public/Dataflow/DataflowNode.h
+++ b/Engine/Source/Runtime/Experimental/Dataflow/Core/Public/Dataflow/DataflowNode.h
@@ -210,11 +210,6 @@ PRAGMA_ENABLE_DEPRECATION_WARNINGS
 	/** Override this method to provide custom reconnections when a node outputs has been deprecated and removed. */
 	virtual FDataflowOutput* RedirectSerializedOutput(const FName& MissingOutputName) { return nullptr; }
 
-	/** Called by editor toolkits when the node is selected, or already selected and invalidated. */
-	virtual void OnSelected(UE::Dataflow::FContext& Context) {}
-	/** Called by editor toolkits when the node is deselected. */
-	virtual void OnDeselected() {}
-
 	//
 	//  Struct Support
 	//
