Chaos Cloth Asset - Added Dataflow namespace alias to allow time for deprecation.

#rnx

diff --git a/Engine/Plugins/ChaosClothAsset/Source/ChaosClothAssetEngine/Public/ChaosClothAsset/ClothAsset.h b/Engine/Plugins/ChaosClothAsset/Source/ChaosClothAssetEngine/Public/ChaosClothAsset/ClothAsset.h
--- a/Engine/Plugins/ChaosClothAsset/Source/ChaosClothAssetEngine/Public/ChaosClothAsset/ClothAsset.h
+++ b/Engine/Plugins/ChaosClothAsset/Source/ChaosClothAssetEngine/Public/ChaosClothAsset/ClothAsset.h
@@ -9,6 +9,12 @@
 #include "Dataflow/DataflowContent.h"
 #include "ClothAsset.generated.h"
 
+#if UE_ENABLE_INCLUDE_ORDER_DEPRECATED_IN_5_5
+namespace Dataflow = UE::Dataflow;
+#else
+namespace UE_DEPRECATED(5.5, "Use UE::Dafaflow instead.") Dataflow {}
+#endif
+
 class FSkeletalMeshRenderData;
 class FSkeletalMeshModel;
 class FSkinnedAssetCompilationContext;
diff --git a/Engine/Plugins/ChaosClothAsset/Source/ChaosClothAssetEngine/Public/ChaosClothAsset/ClothComponent.h b/Engine/Plugins/ChaosClothAsset/Source/ChaosClothAssetEngine/Public/ChaosClothAsset/ClothComponent.h
--- a/Engine/Plugins/ChaosClothAsset/Source/ChaosClothAssetEngine/Public/ChaosClothAsset/ClothComponent.h
+++ b/Engine/Plugins/ChaosClothAsset/Source/ChaosClothAssetEngine/Public/ChaosClothAsset/ClothComponent.h
@@ -6,6 +6,12 @@
 #include "Dataflow/Interfaces/DataflowPhysicsSolver.h"
 #include "ClothComponent.generated.h"
 
+#if UE_ENABLE_INCLUDE_ORDER_DEPRECATED_IN_5_5
+namespace Dataflow = UE::Dataflow;
+#else
+namespace UE_DEPRECATED(5.5, "Use UE::Dafaflow instead.") Dataflow {}
+#endif
+
 class UChaosClothAsset;
 class UChaosClothComponent;
 class UChaosClothAssetInteractor;
diff --git a/Engine/Plugins/ChaosClothAsset/Source/ChaosClothAssetEngine/Public/ChaosClothAsset/ClothSimulationProxy.h b/Engine/Plugins/ChaosClothAsset/Source/ChaosClothAssetEngine/Public/ChaosClothAsset/ClothSimulationProxy.h
--- a/Engine/Plugins/ChaosClothAsset/Source/ChaosClothAssetEngine/Public/ChaosClothAsset/ClothSimulationProxy.h
+++ b/Engine/Plugins/ChaosClothAsset/Source/ChaosClothAssetEngine/Public/ChaosClothAsset/ClothSimulationProxy.h
@@ -9,6 +9,12 @@
 #include "ClothingSystemRuntimeTypes.h"
 #include "Dataflow/Interfaces/DataflowPhysicsSolver.h"
 
+#if UE_ENABLE_INCLUDE_ORDER_DEPRECATED_IN_5_5
+namespace Dataflow = UE::Dataflow;
+#else
+namespace UE_DEPRECATED(5.5, "Use UE::Dafaflow instead.") Dataflow {}
+#endif
+
 namespace Chaos
 {
 	class FClothingSimulationSolver;
diff --git a/Engine/Plugins/ChaosClothAssetEditor/Source/ChaosClothAssetDataflowNodes/Public/ChaosClothAsset/AddStitchNode.h b/Engine/Plugins/ChaosClothAssetEditor/Source/ChaosClothAssetDataflowNodes/Public/ChaosClothAsset/AddStitchNode.h
--- a/Engine/Plugins/ChaosClothAssetEditor/Source/ChaosClothAssetDataflowNodes/Public/ChaosClothAsset/AddStitchNode.h
+++ b/Engine/Plugins/ChaosClothAssetEditor/Source/ChaosClothAssetDataflowNodes/Public/ChaosClothAsset/AddStitchNode.h
@@ -7,6 +7,12 @@
 #include "ChaosClothAsset/ConnectableValue.h"
 #include "AddStitchNode.generated.h"
 
+#if UE_ENABLE_INCLUDE_ORDER_DEPRECATED_IN_5_5
+namespace Dataflow = UE::Dataflow;
+#else
+namespace UE_DEPRECATED(5.5, "Use UE::Dafaflow instead.") Dataflow {}
+#endif
+
 USTRUCT(Meta = (DataflowCloth))
 struct FChaosClothAssetAddStitchNode : public FDataflowNode
 {
diff --git a/Engine/Plugins/ChaosClothAssetEditor/Source/ChaosClothAssetDataflowNodes/Public/ChaosClothAsset/AddWeightMapNode.h b/Engine/Plugins/ChaosClothAssetEditor/Source/ChaosClothAssetDataflowNodes/Public/ChaosClothAsset/AddWeightMapNode.h
--- a/Engine/Plugins/ChaosClothAssetEditor/Source/ChaosClothAssetDataflowNodes/Public/ChaosClothAsset/AddWeightMapNode.h
+++ b/Engine/Plugins/ChaosClothAssetEditor/Source/ChaosClothAssetDataflowNodes/Public/ChaosClothAsset/AddWeightMapNode.h
@@ -7,6 +7,12 @@
 #include "ChaosClothAsset/ConnectableValue.h"
 #include "AddWeightMapNode.generated.h"
 
+#if UE_ENABLE_INCLUDE_ORDER_DEPRECATED_IN_5_5
+namespace Dataflow = UE::Dataflow;
+#else
+namespace UE_DEPRECATED(5.5, "Use UE::Dafaflow instead.") Dataflow {}
+#endif
+
 UENUM()
 enum class EChaosClothAssetWeightMapTransferType : uint8
 {
diff --git a/Engine/Plugins/ChaosClothAssetEditor/Source/ChaosClothAssetDataflowNodes/Public/ChaosClothAsset/AttributeNode.h b/Engine/Plugins/ChaosClothAssetEditor/Source/ChaosClothAssetDataflowNodes/Public/ChaosClothAsset/AttributeNode.h
--- a/Engine/Plugins/ChaosClothAssetEditor/Source/ChaosClothAssetDataflowNodes/Public/ChaosClothAsset/AttributeNode.h
+++ b/Engine/Plugins/ChaosClothAssetEditor/Source/ChaosClothAssetDataflowNodes/Public/ChaosClothAsset/AttributeNode.h
@@ -7,6 +7,12 @@
 #include "ChaosClothAsset/ConnectableValue.h"
 #include "AttributeNode.generated.h"
 
+#if UE_ENABLE_INCLUDE_ORDER_DEPRECATED_IN_5_5
+namespace Dataflow = UE::Dataflow;
+#else
+namespace UE_DEPRECATED(5.5, "Use UE::Dafaflow instead.") Dataflow {}
+#endif
+
 /**
  * The managed array collection group used for the attribute creation.
  * This separate structure is required to allow for customization of the UI.
diff --git a/Engine/Plugins/ChaosClothAssetEditor/Source/ChaosClothAssetDataflowNodes/Public/ChaosClothAsset/BindToRootBoneNode.h b/Engine/Plugins/ChaosClothAssetEditor/Source/ChaosClothAssetDataflowNodes/Public/ChaosClothAsset/BindToRootBoneNode.h
--- a/Engine/Plugins/ChaosClothAssetEditor/Source/ChaosClothAssetDataflowNodes/Public/ChaosClothAsset/BindToRootBoneNode.h
+++ b/Engine/Plugins/ChaosClothAssetEditor/Source/ChaosClothAssetDataflowNodes/Public/ChaosClothAsset/BindToRootBoneNode.h
@@ -6,6 +6,12 @@
 #include "GeometryCollection/ManagedArrayCollection.h"
 #include "BindToRootBoneNode.generated.h"
 
+#if UE_ENABLE_INCLUDE_ORDER_DEPRECATED_IN_5_5
+namespace Dataflow = UE::Dataflow;
+#else
+namespace UE_DEPRECATED(5.5, "Use UE::Dafaflow instead.") Dataflow {}
+#endif
+
 /** Bind an entire mesh to the single root bone of the current skeleton set on the cloth collection. */
 USTRUCT(Meta = (DataflowCloth))
 struct FChaosClothAssetBindToRootBoneNode : public FDataflowNode
diff --git a/Engine/Plugins/ChaosClothAssetEditor/Source/ChaosClothAssetDataflowNodes/Public/ChaosClothAsset/CopySimulationToRenderMeshNode.h b/Engine/Plugins/ChaosClothAssetEditor/Source/ChaosClothAssetDataflowNodes/Public/ChaosClothAsset/CopySimulationToRenderMeshNode.h
--- a/Engine/Plugins/ChaosClothAssetEditor/Source/ChaosClothAssetDataflowNodes/Public/ChaosClothAsset/CopySimulationToRenderMeshNode.h
+++ b/Engine/Plugins/ChaosClothAssetEditor/Source/ChaosClothAssetDataflowNodes/Public/ChaosClothAsset/CopySimulationToRenderMeshNode.h
@@ -6,6 +6,12 @@
 #include "GeometryCollection/ManagedArrayCollection.h"
 #include "CopySimulationToRenderMeshNode.generated.h"
 
+#if UE_ENABLE_INCLUDE_ORDER_DEPRECATED_IN_5_5
+namespace Dataflow = UE::Dataflow;
+#else
+namespace UE_DEPRECATED(5.5, "Use UE::Dafaflow instead.") Dataflow {}
+#endif
+
 class UMaterialInterface;
 
 /** Copy the simulation mesh to the render mesh to be able to render the simulation mesh, or when not using a different mesh for rendering. */
diff --git a/Engine/Plugins/ChaosClothAssetEditor/Source/ChaosClothAssetDataflowNodes/Public/ChaosClothAsset/DatasmithImportNode.h b/Engine/Plugins/ChaosClothAssetEditor/Source/ChaosClothAssetDataflowNodes/Public/ChaosClothAsset/DatasmithImportNode.h
--- a/Engine/Plugins/ChaosClothAssetEditor/Source/ChaosClothAssetDataflowNodes/Public/ChaosClothAsset/DatasmithImportNode.h
+++ b/Engine/Plugins/ChaosClothAssetEditor/Source/ChaosClothAssetDataflowNodes/Public/ChaosClothAsset/DatasmithImportNode.h
@@ -7,6 +7,12 @@
 
 #include "DatasmithImportNode.generated.h"
 
+#if UE_ENABLE_INCLUDE_ORDER_DEPRECATED_IN_5_5
+namespace Dataflow = UE::Dataflow;
+#else
+namespace UE_DEPRECATED(5.5, "Use UE::Dafaflow instead.") Dataflow {}
+#endif
+
 /** Deprecated. */
 UCLASS()
 class UChaosClothAssetDatasmithClothAssetFactory final : public UObject
diff --git a/Engine/Plugins/ChaosClothAssetEditor/Source/ChaosClothAssetDataflowNodes/Public/ChaosClothAsset/DeleteElementNode.h b/Engine/Plugins/ChaosClothAssetEditor/Source/ChaosClothAssetDataflowNodes/Public/ChaosClothAsset/DeleteElementNode.h
--- a/Engine/Plugins/ChaosClothAssetEditor/Source/ChaosClothAssetDataflowNodes/Public/ChaosClothAsset/DeleteElementNode.h
+++ b/Engine/Plugins/ChaosClothAssetEditor/Source/ChaosClothAssetDataflowNodes/Public/ChaosClothAsset/DeleteElementNode.h
@@ -8,6 +8,12 @@
 #include "ChaosClothAsset/ConnectableValue.h"
 #include "DeleteElementNode.generated.h"
 
+#if UE_ENABLE_INCLUDE_ORDER_DEPRECATED_IN_5_5
+namespace Dataflow = UE::Dataflow;
+#else
+namespace UE_DEPRECATED(5.5, "Use UE::Dafaflow instead.") Dataflow {}
+#endif
+
 UENUM()
 enum class UE_DEPRECATED(5.4, "Use FChaosClothAssetNodeSelectionGroup instead") EChaosClothAssetElementType : uint8
 {
diff --git a/Engine/Plugins/ChaosClothAssetEditor/Source/ChaosClothAssetDataflowNodes/Public/ChaosClothAsset/ImportNode.h b/Engine/Plugins/ChaosClothAssetEditor/Source/ChaosClothAssetDataflowNodes/Public/ChaosClothAsset/ImportNode.h
--- a/Engine/Plugins/ChaosClothAssetEditor/Source/ChaosClothAssetDataflowNodes/Public/ChaosClothAsset/ImportNode.h
+++ b/Engine/Plugins/ChaosClothAssetEditor/Source/ChaosClothAssetDataflowNodes/Public/ChaosClothAsset/ImportNode.h
@@ -6,6 +6,12 @@
 #include "GeometryCollection/ManagedArrayCollection.h"
 #include "ImportNode.generated.h"
 
+#if UE_ENABLE_INCLUDE_ORDER_DEPRECATED_IN_5_5
+namespace Dataflow = UE::Dataflow;
+#else
+namespace UE_DEPRECATED(5.5, "Use UE::Dafaflow instead.") Dataflow {}
+#endif
+
 class UChaosClothAsset;
 
 /** Refresh structure for push button customization. */
diff --git a/Engine/Plugins/ChaosClothAssetEditor/Source/ChaosClothAssetDataflowNodes/Public/ChaosClothAsset/MergeClothCollectionsNode.h b/Engine/Plugins/ChaosClothAssetEditor/Source/ChaosClothAssetDataflowNodes/Public/ChaosClothAsset/MergeClothCollectionsNode.h
--- a/Engine/Plugins/ChaosClothAssetEditor/Source/ChaosClothAssetDataflowNodes/Public/ChaosClothAsset/MergeClothCollectionsNode.h
+++ b/Engine/Plugins/ChaosClothAssetEditor/Source/ChaosClothAssetDataflowNodes/Public/ChaosClothAsset/MergeClothCollectionsNode.h
@@ -6,6 +6,12 @@
 #include "GeometryCollection/ManagedArrayCollection.h"
 #include "MergeClothCollectionsNode.generated.h"
 
+#if UE_ENABLE_INCLUDE_ORDER_DEPRECATED_IN_5_5
+namespace Dataflow = UE::Dataflow;
+#else
+namespace UE_DEPRECATED(5.5, "Use UE::Dafaflow instead.") Dataflow {}
+#endif
+
 /** Merge multiple cloth collections into a single cloth collection of multiple patterns. */
 USTRUCT(Meta = (DataflowCloth))
 struct FChaosClothAssetMergeClothCollectionsNode_v2 : public FDataflowNode
diff --git a/Engine/Plugins/ChaosClothAssetEditor/Source/ChaosClothAssetDataflowNodes/Public/ChaosClothAsset/ProxyDeformerNode.h b/Engine/Plugins/ChaosClothAssetEditor/Source/ChaosClothAssetDataflowNodes/Public/ChaosClothAsset/ProxyDeformerNode.h
--- a/Engine/Plugins/ChaosClothAssetEditor/Source/ChaosClothAssetDataflowNodes/Public/ChaosClothAsset/ProxyDeformerNode.h
+++ b/Engine/Plugins/ChaosClothAssetEditor/Source/ChaosClothAssetDataflowNodes/Public/ChaosClothAsset/ProxyDeformerNode.h
@@ -8,6 +8,12 @@
 #include "GeometryCollection/ManagedArrayCollection.h"
 #include "ProxyDeformerNode.generated.h"
 
+#if UE_ENABLE_INCLUDE_ORDER_DEPRECATED_IN_5_5
+namespace Dataflow = UE::Dataflow;
+#else
+namespace UE_DEPRECATED(5.5, "Use UE::Dafaflow instead.") Dataflow {}
+#endif
+
 /** Selection Filter Set*/
 USTRUCT()
 struct FChaosClothAssetSelectionFilterSet
diff --git a/Engine/Plugins/ChaosClothAssetEditor/Source/ChaosClothAssetDataflowNodes/Public/ChaosClothAsset/RemeshNode.h b/Engine/Plugins/ChaosClothAssetEditor/Source/ChaosClothAssetDataflowNodes/Public/ChaosClothAsset/RemeshNode.h
--- a/Engine/Plugins/ChaosClothAssetEditor/Source/ChaosClothAssetDataflowNodes/Public/ChaosClothAsset/RemeshNode.h
+++ b/Engine/Plugins/ChaosClothAssetEditor/Source/ChaosClothAssetDataflowNodes/Public/ChaosClothAsset/RemeshNode.h
@@ -6,6 +6,11 @@
 #include "GeometryCollection/ManagedArrayCollection.h"
 #include "RemeshNode.generated.h"
 
+#if UE_ENABLE_INCLUDE_ORDER_DEPRECATED_IN_5_5
+namespace Dataflow = UE::Dataflow;
+#else
+namespace UE_DEPRECATED(5.5, "Use UE::Dafaflow instead.") Dataflow {}
+#endif
 
 UENUM(BlueprintType)
 enum class EChaosClothAssetRemeshMethod : uint8
diff --git a/Engine/Plugins/ChaosClothAssetEditor/Source/ChaosClothAssetDataflowNodes/Public/ChaosClothAsset/ReverseNormalsNode.h b/Engine/Plugins/ChaosClothAssetEditor/Source/ChaosClothAssetDataflowNodes/Public/ChaosClothAsset/ReverseNormalsNode.h
--- a/Engine/Plugins/ChaosClothAssetEditor/Source/ChaosClothAssetDataflowNodes/Public/ChaosClothAsset/ReverseNormalsNode.h
+++ b/Engine/Plugins/ChaosClothAssetEditor/Source/ChaosClothAssetDataflowNodes/Public/ChaosClothAsset/ReverseNormalsNode.h
@@ -6,6 +6,12 @@
 #include "GeometryCollection/ManagedArrayCollection.h"
 #include "ReverseNormalsNode.generated.h"
 
+#if UE_ENABLE_INCLUDE_ORDER_DEPRECATED_IN_5_5
+namespace Dataflow = UE::Dataflow;
+#else
+namespace UE_DEPRECATED(5.5, "Use UE::Dafaflow instead.") Dataflow {}
+#endif
+
 /** Reverse the geometry's normals or/and winding order of the simulation or/and render meshes stored in the cloth collection. */
 USTRUCT(Meta = (DataflowCloth))
 struct FChaosClothAssetReverseNormalsNode : public FDataflowNode
diff --git a/Engine/Plugins/ChaosClothAssetEditor/Source/ChaosClothAssetDataflowNodes/Public/ChaosClothAsset/SelectionNode.h b/Engine/Plugins/ChaosClothAssetEditor/Source/ChaosClothAssetDataflowNodes/Public/ChaosClothAsset/SelectionNode.h
--- a/Engine/Plugins/ChaosClothAssetEditor/Source/ChaosClothAssetDataflowNodes/Public/ChaosClothAsset/SelectionNode.h
+++ b/Engine/Plugins/ChaosClothAssetEditor/Source/ChaosClothAssetDataflowNodes/Public/ChaosClothAsset/SelectionNode.h
@@ -8,6 +8,12 @@
 #include "ChaosClothAsset/AddWeightMapNode.h"
 #include "SelectionNode.generated.h"
 
+#if UE_ENABLE_INCLUDE_ORDER_DEPRECATED_IN_5_5
+namespace Dataflow = UE::Dataflow;
+#else
+namespace UE_DEPRECATED(5.5, "Use UE::Dafaflow instead.") Dataflow {}
+#endif
+
 /** What type of element is selected in the Selection */
 UENUM()
 enum class UE_DEPRECATED(5.4, "Use FChaosClothAssetNodeSelectionGroup instead") EChaosClothAssetSelectionType : uint8
diff --git a/Engine/Plugins/ChaosClothAssetEditor/Source/ChaosClothAssetDataflowNodes/Public/ChaosClothAsset/SelectionToIntMapNode.h b/Engine/Plugins/ChaosClothAssetEditor/Source/ChaosClothAssetDataflowNodes/Public/ChaosClothAsset/SelectionToIntMapNode.h
--- a/Engine/Plugins/ChaosClothAssetEditor/Source/ChaosClothAssetDataflowNodes/Public/ChaosClothAsset/SelectionToIntMapNode.h
+++ b/Engine/Plugins/ChaosClothAssetEditor/Source/ChaosClothAssetDataflowNodes/Public/ChaosClothAsset/SelectionToIntMapNode.h
@@ -7,6 +7,12 @@
 #include "ChaosClothAsset/ConnectableValue.h"
 #include "SelectionToIntMapNode.generated.h"
 
+#if UE_ENABLE_INCLUDE_ORDER_DEPRECATED_IN_5_5
+namespace Dataflow = UE::Dataflow;
+#else
+namespace UE_DEPRECATED(5.5, "Use UE::Dafaflow instead.") Dataflow {}
+#endif
+
 /** Convert an integer index selection to an integer map. Map type will match the selection type.*/
 USTRUCT(Meta = (DataflowCloth))
 struct FChaosClothAssetSelectionToIntMapNode : public FDataflowNode
diff --git a/Engine/Plugins/ChaosClothAssetEditor/Source/ChaosClothAssetDataflowNodes/Public/ChaosClothAsset/SelectionToWeightMapNode.h b/Engine/Plugins/ChaosClothAssetEditor/Source/ChaosClothAssetDataflowNodes/Public/ChaosClothAsset/SelectionToWeightMapNode.h
--- a/Engine/Plugins/ChaosClothAssetEditor/Source/ChaosClothAssetDataflowNodes/Public/ChaosClothAsset/SelectionToWeightMapNode.h
+++ b/Engine/Plugins/ChaosClothAssetEditor/Source/ChaosClothAssetDataflowNodes/Public/ChaosClothAsset/SelectionToWeightMapNode.h
@@ -6,6 +6,12 @@
 #include "GeometryCollection/ManagedArrayCollection.h"
 #include "SelectionToWeightMapNode.generated.h"
 
+#if UE_ENABLE_INCLUDE_ORDER_DEPRECATED_IN_5_5
+namespace Dataflow = UE::Dataflow;
+#else
+namespace UE_DEPRECATED(5.5, "Use UE::Dafaflow instead.") Dataflow {}
+#endif
+
 /** Convert an integer index selection to a vertex weight map where different map values can be set for selected and unselected vertices. */
 USTRUCT(Meta = (DataflowCloth))
 struct FChaosClothAssetSelectionToWeightMapNode : public FDataflowNode
diff --git a/Engine/Plugins/ChaosClothAssetEditor/Source/ChaosClothAssetDataflowNodes/Public/ChaosClothAsset/SetPhysicsAssetNode.h b/Engine/Plugins/ChaosClothAssetEditor/Source/ChaosClothAssetDataflowNodes/Public/ChaosClothAsset/SetPhysicsAssetNode.h
--- a/Engine/Plugins/ChaosClothAssetEditor/Source/ChaosClothAssetDataflowNodes/Public/ChaosClothAsset/SetPhysicsAssetNode.h
+++ b/Engine/Plugins/ChaosClothAssetEditor/Source/ChaosClothAssetDataflowNodes/Public/ChaosClothAsset/SetPhysicsAssetNode.h
@@ -6,6 +6,12 @@
 #include "GeometryCollection/ManagedArrayCollection.h"
 #include "SetPhysicsAssetNode.generated.h"
 
+#if UE_ENABLE_INCLUDE_ORDER_DEPRECATED_IN_5_5
+namespace Dataflow = UE::Dataflow;
+#else
+namespace UE_DEPRECATED(5.5, "Use UE::Dafaflow instead.") Dataflow {}
+#endif
+
 class UPhysicsAsset;
 
 /** Replace the current physics assets to collide the simulation mesh against. */
diff --git a/Engine/Plugins/ChaosClothAssetEditor/Source/ChaosClothAssetDataflowNodes/Public/ChaosClothAsset/SimulationBaseConfigNode.h b/Engine/Plugins/ChaosClothAssetEditor/Source/ChaosClothAssetDataflowNodes/Public/ChaosClothAsset/SimulationBaseConfigNode.h
--- a/Engine/Plugins/ChaosClothAssetEditor/Source/ChaosClothAssetDataflowNodes/Public/ChaosClothAsset/SimulationBaseConfigNode.h
+++ b/Engine/Plugins/ChaosClothAssetEditor/Source/ChaosClothAssetDataflowNodes/Public/ChaosClothAsset/SimulationBaseConfigNode.h
@@ -12,6 +12,12 @@
 #include "GeometryCollection/ManagedArrayCollection.h"
 #include "SimulationBaseConfigNode.generated.h"
 
+#if UE_ENABLE_INCLUDE_ORDER_DEPRECATED_IN_5_5
+namespace Dataflow = UE::Dataflow;
+#else
+namespace UE_DEPRECATED(5.5, "Use UE::Dafaflow instead.") Dataflow {}
+#endif
+
 /**
  * Base abstract class for all cloth asset config nodes.
  * Inherited class must call RegisterCollectionConnections() in constructor to use this base class Collection.
diff --git a/Engine/Plugins/ChaosClothAssetEditor/Source/ChaosClothAssetDataflowNodes/Public/ChaosClothAsset/SimulationDefaultConfigNode.h b/Engine/Plugins/ChaosClothAssetEditor/Source/ChaosClothAssetDataflowNodes/Public/ChaosClothAsset/SimulationDefaultConfigNode.h
--- a/Engine/Plugins/ChaosClothAssetEditor/Source/ChaosClothAssetDataflowNodes/Public/ChaosClothAsset/SimulationDefaultConfigNode.h
+++ b/Engine/Plugins/ChaosClothAssetEditor/Source/ChaosClothAssetDataflowNodes/Public/ChaosClothAsset/SimulationDefaultConfigNode.h
@@ -6,6 +6,12 @@
 #include "GeometryCollection/ManagedArrayCollection.h"
 #include "SimulationDefaultConfigNode.generated.h"
 
+#if UE_ENABLE_INCLUDE_ORDER_DEPRECATED_IN_5_5
+namespace Dataflow = UE::Dataflow;
+#else
+namespace UE_DEPRECATED(5.5, "Use UE::Dafaflow instead.") Dataflow {}
+#endif
+
 class UChaosClothConfig;
 class UChaosClothSharedSimConfig;
 
diff --git a/Engine/Plugins/ChaosClothAssetEditor/Source/ChaosClothAssetDataflowNodes/Public/ChaosClothAsset/SkeletalMeshImportNode.h b/Engine/Plugins/ChaosClothAssetEditor/Source/ChaosClothAssetDataflowNodes/Public/ChaosClothAsset/SkeletalMeshImportNode.h
--- a/Engine/Plugins/ChaosClothAssetEditor/Source/ChaosClothAssetDataflowNodes/Public/ChaosClothAsset/SkeletalMeshImportNode.h
+++ b/Engine/Plugins/ChaosClothAssetEditor/Source/ChaosClothAssetDataflowNodes/Public/ChaosClothAsset/SkeletalMeshImportNode.h
@@ -6,6 +6,12 @@
 #include "GeometryCollection/ManagedArrayCollection.h"
 #include "SkeletalMeshImportNode.generated.h"
 
+#if UE_ENABLE_INCLUDE_ORDER_DEPRECATED_IN_5_5
+namespace Dataflow = UE::Dataflow;
+#else
+namespace UE_DEPRECATED(5.5, "Use UE::Dafaflow instead.") Dataflow {}
+#endif
+
 class USkeletalMesh;
 /** Import a skeletal mesh asset into the cloth collection simulation and/or render mesh containers. */
 USTRUCT(Meta = (DataflowCloth))
diff --git a/Engine/Plugins/ChaosClothAssetEditor/Source/ChaosClothAssetDataflowNodes/Public/ChaosClothAsset/SkinningBlendNode.h b/Engine/Plugins/ChaosClothAssetEditor/Source/ChaosClothAssetDataflowNodes/Public/ChaosClothAsset/SkinningBlendNode.h
--- a/Engine/Plugins/ChaosClothAssetEditor/Source/ChaosClothAssetDataflowNodes/Public/ChaosClothAsset/SkinningBlendNode.h
+++ b/Engine/Plugins/ChaosClothAssetEditor/Source/ChaosClothAssetDataflowNodes/Public/ChaosClothAsset/SkinningBlendNode.h
@@ -7,6 +7,12 @@
 #include "GeometryCollection/ManagedArrayCollection.h"
 #include "SkinningBlendNode.generated.h"
 
+#if UE_ENABLE_INCLUDE_ORDER_DEPRECATED_IN_5_5
+namespace Dataflow = UE::Dataflow;
+#else
+namespace UE_DEPRECATED(5.5, "Use UE::Dafaflow instead.") Dataflow {}
+#endif
+
 /**
  * Initialize the RenderDeformerSkinningBlend weight map from the ProxyDeformer mapping data.
  * 
diff --git a/Engine/Plugins/ChaosClothAssetEditor/Source/ChaosClothAssetDataflowNodes/Public/ChaosClothAsset/StaticMeshImportNode.h b/Engine/Plugins/ChaosClothAssetEditor/Source/ChaosClothAssetDataflowNodes/Public/ChaosClothAsset/StaticMeshImportNode.h
--- a/Engine/Plugins/ChaosClothAssetEditor/Source/ChaosClothAssetDataflowNodes/Public/ChaosClothAsset/StaticMeshImportNode.h
+++ b/Engine/Plugins/ChaosClothAssetEditor/Source/ChaosClothAssetDataflowNodes/Public/ChaosClothAsset/StaticMeshImportNode.h
@@ -6,6 +6,12 @@
 #include "GeometryCollection/ManagedArrayCollection.h"
 #include "StaticMeshImportNode.generated.h"
 
+#if UE_ENABLE_INCLUDE_ORDER_DEPRECATED_IN_5_5
+namespace Dataflow = UE::Dataflow;
+#else
+namespace UE_DEPRECATED(5.5, "Use UE::Dafaflow instead.") Dataflow {}
+#endif
+
 class UStaticMesh;
 
 /** Import a static mesh asset into the cloth collection simulation and/or render mesh containers.*/
diff --git a/Engine/Plugins/ChaosClothAssetEditor/Source/ChaosClothAssetDataflowNodes/Public/ChaosClothAsset/TerminalNode.h b/Engine/Plugins/ChaosClothAssetEditor/Source/ChaosClothAssetDataflowNodes/Public/ChaosClothAsset/TerminalNode.h
--- a/Engine/Plugins/ChaosClothAssetEditor/Source/ChaosClothAssetDataflowNodes/Public/ChaosClothAsset/TerminalNode.h
+++ b/Engine/Plugins/ChaosClothAssetEditor/Source/ChaosClothAssetDataflowNodes/Public/ChaosClothAsset/TerminalNode.h
@@ -7,6 +7,12 @@
 #include "ChaosClothAsset/ClothLodTransitionDataCache.h"
 #include "TerminalNode.generated.h"
 
+#if UE_ENABLE_INCLUDE_ORDER_DEPRECATED_IN_5_5
+namespace Dataflow = UE::Dataflow;
+#else
+namespace UE_DEPRECATED(5.5, "Use UE::Dafaflow instead.") Dataflow {}
+#endif
+
 /** Refresh structure for push buton customization. */
 USTRUCT()
 struct FChaosClothAssetTerminalNodeRefreshAsset
diff --git a/Engine/Plugins/ChaosClothAssetEditor/Source/ChaosClothAssetDataflowNodes/Public/ChaosClothAsset/TransferSkinWeightsNode.h b/Engine/Plugins/ChaosClothAssetEditor/Source/ChaosClothAssetDataflowNodes/Public/ChaosClothAsset/TransferSkinWeightsNode.h
--- a/Engine/Plugins/ChaosClothAssetEditor/Source/ChaosClothAssetDataflowNodes/Public/ChaosClothAsset/TransferSkinWeightsNode.h
+++ b/Engine/Plugins/ChaosClothAssetEditor/Source/ChaosClothAssetDataflowNodes/Public/ChaosClothAsset/TransferSkinWeightsNode.h
@@ -7,6 +7,12 @@
 #include "ChaosClothAsset/WeightedValue.h"
 #include "TransferSkinWeightsNode.generated.h"
 
+#if UE_ENABLE_INCLUDE_ORDER_DEPRECATED_IN_5_5
+namespace Dataflow = UE::Dataflow;
+#else
+namespace UE_DEPRECATED(5.5, "Use UE::Dafaflow instead.") Dataflow {}
+#endif
+
 class USkeletalMesh;
 
 UENUM(BlueprintType)
diff --git a/Engine/Plugins/ChaosClothAssetEditor/Source/ChaosClothAssetDataflowNodes/Public/ChaosClothAsset/TransformPositionsNode.h b/Engine/Plugins/ChaosClothAssetEditor/Source/ChaosClothAssetDataflowNodes/Public/ChaosClothAsset/TransformPositionsNode.h
--- a/Engine/Plugins/ChaosClothAssetEditor/Source/ChaosClothAssetDataflowNodes/Public/ChaosClothAsset/TransformPositionsNode.h
+++ b/Engine/Plugins/ChaosClothAssetEditor/Source/ChaosClothAssetDataflowNodes/Public/ChaosClothAsset/TransformPositionsNode.h
@@ -6,6 +6,12 @@
 #include "GeometryCollection/ManagedArrayCollection.h"
 #include "TransformPositionsNode.generated.h"
 
+#if UE_ENABLE_INCLUDE_ORDER_DEPRECATED_IN_5_5
+namespace Dataflow = UE::Dataflow;
+#else
+namespace UE_DEPRECATED(5.5, "Use UE::Dafaflow instead.") Dataflow {}
+#endif
+
 USTRUCT(Meta = (DataflowCloth))
 struct FChaosClothAssetTransformPositionsNode : public FDataflowNode
 {
diff --git a/Engine/Plugins/ChaosClothAssetEditor/Source/ChaosClothAssetDataflowNodes/Public/ChaosClothAsset/TransformUVsNode.h b/Engine/Plugins/ChaosClothAssetEditor/Source/ChaosClothAssetDataflowNodes/Public/ChaosClothAsset/TransformUVsNode.h
--- a/Engine/Plugins/ChaosClothAssetEditor/Source/ChaosClothAssetDataflowNodes/Public/ChaosClothAsset/TransformUVsNode.h
+++ b/Engine/Plugins/ChaosClothAssetEditor/Source/ChaosClothAssetDataflowNodes/Public/ChaosClothAsset/TransformUVsNode.h
@@ -6,6 +6,12 @@
 #include "GeometryCollection/ManagedArrayCollection.h"
 #include "TransformUVsNode.generated.h"
 
+#if UE_ENABLE_INCLUDE_ORDER_DEPRECATED_IN_5_5
+namespace Dataflow = UE::Dataflow;
+#else
+namespace UE_DEPRECATED(5.5, "Use UE::Dafaflow instead.") Dataflow {}
+#endif
+
 USTRUCT(Meta = (DataflowCloth))
 struct FChaosClothAssetTransformUVsNode : public FDataflowNode
 {
diff --git a/Engine/Plugins/ChaosClothAssetEditor/Source/ChaosClothAssetDataflowNodes/Public/ChaosClothAsset/USDImportNode.h b/Engine/Plugins/ChaosClothAssetEditor/Source/ChaosClothAssetDataflowNodes/Public/ChaosClothAsset/USDImportNode.h
--- a/Engine/Plugins/ChaosClothAssetEditor/Source/ChaosClothAssetDataflowNodes/Public/ChaosClothAsset/USDImportNode.h
+++ b/Engine/Plugins/ChaosClothAssetEditor/Source/ChaosClothAssetDataflowNodes/Public/ChaosClothAsset/USDImportNode.h
@@ -7,6 +7,12 @@
 #include "Misc/SecureHash.h"
 #include "USDImportNode.generated.h"
 
+#if UE_ENABLE_INCLUDE_ORDER_DEPRECATED_IN_5_5
+namespace Dataflow = UE::Dataflow;
+#else
+namespace UE_DEPRECATED(5.5, "Use UE::Dafaflow instead.") Dataflow {}
+#endif
+
 /** Import a USD file from a third party garment construction software. */
 USTRUCT(Meta = (DataflowCloth, Deprecated = "5.5"))
 struct UE_DEPRECATED(5.5, "Use the newer version of this node instead.") FChaosClothAssetUSDImportNode : public FDataflowTerminalNode
diff --git a/Engine/Plugins/ChaosClothAssetEditor/Source/ChaosClothAssetDataflowNodes/Public/ChaosClothAsset/USDImportNode_v2.h b/Engine/Plugins/ChaosClothAssetEditor/Source/ChaosClothAssetDataflowNodes/Public/ChaosClothAsset/USDImportNode_v2.h
--- a/Engine/Plugins/ChaosClothAssetEditor/Source/ChaosClothAssetDataflowNodes/Public/ChaosClothAsset/USDImportNode_v2.h
+++ b/Engine/Plugins/ChaosClothAssetEditor/Source/ChaosClothAssetDataflowNodes/Public/ChaosClothAsset/USDImportNode_v2.h
@@ -11,6 +11,12 @@
 #endif
 #include "USDImportNode_v2.generated.h"
 
+#if UE_ENABLE_INCLUDE_ORDER_DEPRECATED_IN_5_5
+namespace Dataflow = UE::Dataflow;
+#else
+namespace UE_DEPRECATED(5.5, "Use UE::Dafaflow instead.") Dataflow {}
+#endif
+
 class UStaticMesh;
 
 /** Import a USD file from a third party garment construction software. */
diff --git a/Engine/Plugins/ChaosClothAssetEditor/Source/ChaosClothAssetDataflowNodes/Public/ChaosClothAsset/WeightMapNode.h b/Engine/Plugins/ChaosClothAssetEditor/Source/ChaosClothAssetDataflowNodes/Public/ChaosClothAsset/WeightMapNode.h
--- a/Engine/Plugins/ChaosClothAssetEditor/Source/ChaosClothAssetDataflowNodes/Public/ChaosClothAsset/WeightMapNode.h
+++ b/Engine/Plugins/ChaosClothAssetEditor/Source/ChaosClothAssetDataflowNodes/Public/ChaosClothAsset/WeightMapNode.h
@@ -8,6 +8,12 @@
 #include "ChaosClothAsset/ConnectableValue.h"
 #include "WeightMapNode.generated.h"
 
+#if UE_ENABLE_INCLUDE_ORDER_DEPRECATED_IN_5_5
+namespace Dataflow = UE::Dataflow;
+#else
+namespace UE_DEPRECATED(5.5, "Use UE::Dafaflow instead.") Dataflow {}
+#endif
+
 /** Which mesh to update with the node's weight map */
 UENUM()
 enum class EChaosClothAssetWeightMapMeshTarget : uint8
diff --git a/Engine/Plugins/ChaosClothAssetEditor/Source/ChaosClothAssetDataflowNodes/Public/ChaosClothAsset/WeightMapToSelectionNode.h b/Engine/Plugins/ChaosClothAssetEditor/Source/ChaosClothAssetDataflowNodes/Public/ChaosClothAsset/WeightMapToSelectionNode.h
--- a/Engine/Plugins/ChaosClothAssetEditor/Source/ChaosClothAssetDataflowNodes/Public/ChaosClothAsset/WeightMapToSelectionNode.h
+++ b/Engine/Plugins/ChaosClothAssetEditor/Source/ChaosClothAssetDataflowNodes/Public/ChaosClothAsset/WeightMapToSelectionNode.h
@@ -6,6 +6,12 @@
 #include "GeometryCollection/ManagedArrayCollection.h"
 #include "WeightMapToSelectionNode.generated.h"
 
+#if UE_ENABLE_INCLUDE_ORDER_DEPRECATED_IN_5_5
+namespace Dataflow = UE::Dataflow;
+#else
+namespace UE_DEPRECATED(5.5, "Use UE::Dafaflow instead.") Dataflow {}
+#endif
+
 /** What type of element to convert to */
 UENUM()
 enum class EChaosClothAssetWeightMapConvertableSelectionType : uint8
diff --git a/Engine/Plugins/ChaosClothAssetEditor/Source/ChaosClothAssetEditorTools/Public/ChaosClothAsset/ClothEditorContextObject.h b/Engine/Plugins/ChaosClothAssetEditor/Source/ChaosClothAssetEditorTools/Public/ChaosClothAsset/ClothEditorContextObject.h
--- a/Engine/Plugins/ChaosClothAssetEditor/Source/ChaosClothAssetEditorTools/Public/ChaosClothAsset/ClothEditorContextObject.h
+++ b/Engine/Plugins/ChaosClothAssetEditor/Source/ChaosClothAssetEditorTools/Public/ChaosClothAsset/ClothEditorContextObject.h
@@ -8,6 +8,12 @@
 #include "Dataflow/DataflowObject.h"
 #include "ClothEditorContextObject.generated.h"
 
+#if UE_ENABLE_INCLUDE_ORDER_DEPRECATED_IN_5_5
+namespace Dataflow = UE::Dataflow;
+#else
+namespace UE_DEPRECATED(5.5, "Use UE::Dafaflow instead.") Dataflow {}
+#endif
+
 class UDataflow;
 class UEdGraphNode;
 
