Chaos Cloth Asset: remove degenerate triangles when importing Skeletal/Static meshes.

Only updating the _v2 versions of the importers (new in UE 5.5) since the resulting topology can change.

diff --git a/Engine/Plugins/ChaosClothAssetEditor/Source/ChaosClothAssetDataflowNodes/Private/ChaosClothAsset/SkeletalMeshImportNode.cpp b/Engine/Plugins/ChaosClothAssetEditor/Source/ChaosClothAssetDataflowNodes/Private/ChaosClothAsset/SkeletalMeshImportNode.cpp
--- a/Engine/Plugins/ChaosClothAssetEditor/Source/ChaosClothAssetDataflowNodes/Private/ChaosClothAsset/SkeletalMeshImportNode.cpp
+++ b/Engine/Plugins/ChaosClothAssetEditor/Source/ChaosClothAssetDataflowNodes/Private/ChaosClothAsset/SkeletalMeshImportNode.cpp
@@ -83,6 +83,8 @@ void FChaosClothAssetSkeletalMeshImportNode_v2::Evaluate(UE::Dataflow::FContext&
 				}
 			}
 
+			FClothGeometryTools::CleanupAndCompactMesh(ClothCollection);
+
 			if (const UPhysicsAsset* PhysicsAsset = bSetPhysicsAsset ? InSkeletalMesh->GetPhysicsAsset() : nullptr)
 			{
 				ClothFacade.SetPhysicsAssetPathName(PhysicsAsset->GetPathName());
diff --git a/Engine/Plugins/ChaosClothAssetEditor/Source/ChaosClothAssetDataflowNodes/Private/ChaosClothAsset/StaticMeshImportNode.cpp b/Engine/Plugins/ChaosClothAssetEditor/Source/ChaosClothAssetDataflowNodes/Private/ChaosClothAsset/StaticMeshImportNode.cpp
--- a/Engine/Plugins/ChaosClothAssetEditor/Source/ChaosClothAssetDataflowNodes/Private/ChaosClothAsset/StaticMeshImportNode.cpp
+++ b/Engine/Plugins/ChaosClothAssetEditor/Source/ChaosClothAssetDataflowNodes/Private/ChaosClothAsset/StaticMeshImportNode.cpp
@@ -147,8 +147,11 @@ void FChaosClothAssetStaticMeshImportNode_v2::Evaluate(UE::Dataflow::FContext& C
 						}
 					}
 				}
+				
 				// Bind to root bone by default
+				FClothGeometryTools::CleanupAndCompactMesh(ClothCollection);
 				FClothGeometryTools::BindMeshToRootBone(ClothCollection, bImportSimMesh, bImportRenderMesh);
+
 			}
 			else
 			{
diff --git a/Engine/Plugins/ChaosClothAssetEditor/Source/ChaosClothAssetDataflowNodes/Public/ChaosClothAsset/SkeletalMeshImportNode.h b/Engine/Plugins/ChaosClothAssetEditor/Source/ChaosClothAssetDataflowNodes/Public/ChaosClothAsset/SkeletalMeshImportNode.h
--- a/Engine/Plugins/ChaosClothAssetEditor/Source/ChaosClothAssetDataflowNodes/Public/ChaosClothAsset/SkeletalMeshImportNode.h
+++ b/Engine/Plugins/ChaosClothAssetEditor/Source/ChaosClothAssetDataflowNodes/Public/ChaosClothAsset/SkeletalMeshImportNode.h
@@ -71,7 +71,7 @@ private:
 };
 
 /** Import a skeletal mesh asset into the cloth collection simulation and/or render mesh containers. 
- * This version re-calculates the Sim Mesh Normals, and they are flipped. */
+ * This version re-calculates the Sim Mesh Normals, and they are flipped. It also does not remove topologically degenerate triangles. */
 USTRUCT(Meta = (DataflowCloth, Deprecated = "5.5"))
 struct UE_DEPRECATED(5.5, "Use the newer version of this node instead.") FChaosClothAssetSkeletalMeshImportNode : public FDataflowNode
 {
diff --git a/Engine/Plugins/ChaosClothAssetEditor/Source/ChaosClothAssetDataflowNodes/Public/ChaosClothAsset/StaticMeshImportNode.h b/Engine/Plugins/ChaosClothAssetEditor/Source/ChaosClothAssetDataflowNodes/Public/ChaosClothAsset/StaticMeshImportNode.h
--- a/Engine/Plugins/ChaosClothAssetEditor/Source/ChaosClothAssetDataflowNodes/Public/ChaosClothAsset/StaticMeshImportNode.h
+++ b/Engine/Plugins/ChaosClothAssetEditor/Source/ChaosClothAssetDataflowNodes/Public/ChaosClothAsset/StaticMeshImportNode.h
@@ -69,7 +69,7 @@ private:
 
 
 /** Import a static mesh asset into the cloth collection simulation and/or render mesh containers. 
-* This version re-calculates the Sim Mesh Normals, and they are flipped. */
+* This version re-calculates the Sim Mesh Normals, and they are flipped. It also does not remove topologically degenerate triangles.*/
 USTRUCT(Meta = (DataflowCloth, Deprecated = "5.5"))
 struct UE_DEPRECATED(5.5, "Use the newer version of this node instead.") FChaosClothAssetStaticMeshImportNode : public FDataflowNode
 {
