Dataflow - Fixed the AddScalarVertexProperty node.
- Made a generic implementation for the PropertyGroup Customization.
- Added OnVerifyTextChanged to validate the group name before committing the text.
- Updated ClothAsset::FSelectionGroupCustomization and Dataflow::FScalarVertexPropertyGroupCustomization.

#rnx

diff --git a/Engine/Plugins/ChaosClothAsset/Source/ChaosClothAsset/Private/ChaosClothAsset/ClothCollection.cpp b/Engine/Plugins/ChaosClothAsset/Source/ChaosClothAsset/Private/ChaosClothAsset/ClothCollection.cpp
--- a/Engine/Plugins/ChaosClothAsset/Source/ChaosClothAsset/Private/ChaosClothAsset/ClothCollection.cpp
+++ b/Engine/Plugins/ChaosClothAsset/Source/ChaosClothAsset/Private/ChaosClothAsset/ClothCollection.cpp
@@ -823,6 +823,14 @@ namespace UE::Chaos::ClothAsset
 	template CHAOSCLOTHASSET_API TManagedArray<float>* FClothCollection::GetUserDefinedAttribute<float>(const FName& Name, const FName& GroupName);
 	template CHAOSCLOTHASSET_API TManagedArray<FVector3f>* FClothCollection::GetUserDefinedAttribute<FVector3f>(const FName& Name, const FName& GroupName);
 
+	TArray<FName> FClothCollection::GetValidClothCollectionGroupName()
+	{
+		using namespace UE::Chaos::ClothAsset::Private;
+		TArray<FName> ValidClothCollectionGroupName;
+		FixedAttributeNamesMap.GetKeys(ValidClothCollectionGroupName);
+		return ValidClothCollectionGroupName;
+	}
+
 	bool FClothCollection::IsValidClothCollectionGroupName(const FName& GroupName)
 	{
 		using namespace UE::Chaos::ClothAsset::Private;
diff --git a/Engine/Plugins/ChaosClothAsset/Source/ChaosClothAsset/Private/ChaosClothAsset/ClothCollection.h b/Engine/Plugins/ChaosClothAsset/Source/ChaosClothAsset/Private/ChaosClothAsset/ClothCollection.h
--- a/Engine/Plugins/ChaosClothAsset/Source/ChaosClothAsset/Private/ChaosClothAsset/ClothCollection.h
+++ b/Engine/Plugins/ChaosClothAsset/Source/ChaosClothAsset/Private/ChaosClothAsset/ClothCollection.h
@@ -151,6 +151,7 @@ namespace UE::Chaos::ClothAsset
 		template<typename T UE_REQUIRES(TIsUserAttributeType<T>::Value)>
 		TManagedArray<T>* GetUserDefinedAttribute(const FName& Name, const FName& GroupName);
 
+		static TArray<FName> GetValidClothCollectionGroupName();
 		static bool IsValidClothCollectionGroupName(const FName& GroupName);
 		bool IsValidUserDefinedAttributeName(const FName& Name, const FName& GroupName);
 
diff --git a/Engine/Plugins/ChaosClothAsset/Source/ChaosClothAsset/Private/ChaosClothAsset/CollectionClothFacade.cpp b/Engine/Plugins/ChaosClothAsset/Source/ChaosClothAsset/Private/ChaosClothAsset/CollectionClothFacade.cpp
--- a/Engine/Plugins/ChaosClothAsset/Source/ChaosClothAsset/Private/ChaosClothAsset/CollectionClothFacade.cpp
+++ b/Engine/Plugins/ChaosClothAsset/Source/ChaosClothAsset/Private/ChaosClothAsset/CollectionClothFacade.cpp
@@ -490,6 +490,11 @@ namespace UE::Chaos::ClothAsset
 		return FClothCollection::IsValidClothCollectionGroupName(GroupName);
 	}
 
+	TArray<FName> FCollectionClothConstFacade::GetValidClothCollectionGroupName()
+	{
+		return FClothCollection::GetValidClothCollectionGroupName();
+	}
+
 	void FCollectionClothConstFacade::BuildSimulationMesh(TArray<FVector3f>& Positions, TArray<FVector3f>& Normals, TArray<uint32>& Indices, TArray<FVector2f>& PatternsPositions, TArray<uint32>& PatternsIndices,
 		TArray<uint32>& PatternToWeldedIndices, TArray<TArray<int32>>* OptionalWeldedToPatternIndices) const
 	{
diff --git a/Engine/Plugins/ChaosClothAsset/Source/ChaosClothAsset/Public/ChaosClothAsset/CollectionClothFacade.h b/Engine/Plugins/ChaosClothAsset/Source/ChaosClothAsset/Public/ChaosClothAsset/CollectionClothFacade.h
--- a/Engine/Plugins/ChaosClothAsset/Source/ChaosClothAsset/Public/ChaosClothAsset/CollectionClothFacade.h
+++ b/Engine/Plugins/ChaosClothAsset/Source/ChaosClothAsset/Public/ChaosClothAsset/CollectionClothFacade.h
@@ -177,6 +177,7 @@ namespace UE::Chaos::ClothAsset
 		template<typename T UE_REQUIRES(TIsUserAttributeType<T>::Value)>
 		TConstArrayView<T> GetUserDefinedAttribute(const FName& Name, const FName& GroupName) const;
 		static bool IsValidClothCollectionGroupName(const FName& GroupName);
+		static TArray<FName> GetValidClothCollectionGroupName();
 
 		void BuildSimulationMesh(TArray<FVector3f>& Positions, TArray<FVector3f>& Normals, TArray<uint32>& Indices, TArray<FVector2f>& PatternsPositions, TArray<uint32>& PatternsIndices, 
 			TArray<uint32>& PatternToWeldedIndices, TArray<TArray<int32>>* OptionalWeldedToPatternIndices = nullptr) const;
diff --git a/Engine/Plugins/ChaosClothAssetEditor/Source/ChaosClothAssetDataflowNodes/Private/ChaosClothAsset/SelectionGroupCustomization.cpp b/Engine/Plugins/ChaosClothAssetEditor/Source/ChaosClothAssetDataflowNodes/Private/ChaosClothAsset/SelectionGroupCustomization.cpp
--- a/Engine/Plugins/ChaosClothAssetEditor/Source/ChaosClothAssetDataflowNodes/Private/ChaosClothAsset/SelectionGroupCustomization.cpp
+++ b/Engine/Plugins/ChaosClothAssetEditor/Source/ChaosClothAssetDataflowNodes/Private/ChaosClothAsset/SelectionGroupCustomization.cpp
@@ -1,213 +1,17 @@
 // Copyright Epic Games, Inc. All Rights Reserved.
 
 #include "ChaosClothAsset/SelectionGroupCustomization.h"
-#include "ChaosClothAsset/AttributeNode.h"
-#include "ChaosClothAsset/ClothDataflowTools.h"  // For MakeCollectionName
 #include "ChaosClothAsset/CollectionClothFacade.h"
-#include "ChaosClothAsset/DeleteElementNode.h"
-#include "ChaosClothAsset/SelectionNode.h"
-#include "Dataflow/DataflowGraphEditor.h"
-#include "Widgets/Input/SEditableTextBox.h"
-#include "Widgets/Input/SComboButton.h"
-#include "Widgets/Views/SListView.h"
-#include "Widgets/Views/STableRow.h"
-#include "DetailLayoutBuilder.h"
-#include "DetailWidgetRow.h"
-#include "Editor.h"
-
-#define LOCTEXT_NAMESPACE "DataflowNodeSelectionGroupCustomization"
 
 namespace UE::Chaos::ClothAsset
 {
-	namespace Private
-	{
-		// Return the FManagedArrayCollection with the specified name from the property held by the top level struct owner of ChildPropertyHandle
-		FManagedArrayCollection GetPropertyCollection(
-			UE::Dataflow::FContext& Context,
-			const TSharedPtr<IPropertyHandle>& ChildPropertyHandle,
-			const FName CollectionPropertyName = FName(TEXT("Collection")))
-		{
-			static const FManagedArrayCollection EmptyCollection;
-
-			TSharedPtr<IPropertyHandle> OwnerHandle = ChildPropertyHandle;
-			while (TSharedPtr<IPropertyHandle> ParentHandle = OwnerHandle->GetParentHandle())
-			{
-				OwnerHandle = MoveTemp(ParentHandle);
-			}
-			if (const TSharedPtr<IPropertyHandleStruct> OwnerHandleStruct = OwnerHandle->AsStruct())
-			{
-				if (const TSharedPtr<FStructOnScope> StructOnScope = OwnerHandleStruct->GetStructData())
-				{
-					if (const UStruct* const Struct = StructOnScope->GetStruct())
-					{
-						if (Struct->IsChildOf<FDataflowNode>())
-						{
-							const FDataflowNode* const DataflowNode = reinterpret_cast<FDataflowNode*>(StructOnScope->GetStructMemory());
-
-							if (const FProperty* const Property = Struct->FindPropertyByName(CollectionPropertyName))
-							{
-								if (const FStructProperty* const StructProperty = CastField<FStructProperty>(Property))
-								{
-									if (StructProperty->GetCPPType(nullptr, CPPF_None) == TEXT("FManagedArrayCollection"))
-									{
-										if (const FDataflowInput* const DataflowInput = DataflowNode->FindInput(StructProperty->ContainerPtrToValuePtr<FManagedArrayCollection*>(DataflowNode)))
-										{
-											return DataflowInput->GetValue(Context, EmptyCollection);
-										}
-									}
-								}
-							}
-						}
-					}
-				}
-			}
-			return EmptyCollection;
-		}
-	}  // End namespace Private
-
 	TSharedRef<IPropertyTypeCustomization> FSelectionGroupCustomization::MakeInstance()
 	{
 		return MakeShareable(new FSelectionGroupCustomization);
 	}
 
-	void FSelectionGroupCustomization::CustomizeHeader(TSharedRef<IPropertyHandle> PropertyHandle, FDetailWidgetRow& HeaderRow, IPropertyTypeCustomizationUtils& CustomizationUtils)
+	TArray<FName> FSelectionGroupCustomization::GetTargetGroupNames(const FManagedArrayCollection& /*Collection*/) const
 	{
-		DataflowGraphEditor = SDataflowGraphEditor::GetSelectedGraphEditor();
-
-		uint32 NumChildren;
-		const FPropertyAccess::Result Result = PropertyHandle->GetNumChildren(NumChildren);
-
-		ChildPropertyHandle = (Result == FPropertyAccess::Success && NumChildren) ? PropertyHandle->GetChildHandle(0) : TSharedPtr<IPropertyHandle>();
-
-		GroupNames.Reset();
-
-		HeaderRow
-			.NameContent()
-			[
-				PropertyHandle->CreatePropertyNameWidget(PropertyHandle->GetPropertyDisplayName())
-			]
-			.ValueContent()
-			.MinDesiredWidth(250)
-			.MaxDesiredWidth(350.f)
-			[
-				SNew(SHorizontalBox)
-				+ SHorizontalBox::Slot()
-				.MaxWidth(145.f)
-				[
-					SAssignNew(ComboButton, SComboButton)
-					.ButtonStyle(FAppStyle::Get(), "NoBorder")
-					.ContentPadding(0)
-					.OnGetMenuContent(this, &FSelectionGroupCustomization::OnGetMenuContent)
-					.ButtonContent()
-					[
-						SNew(SEditableTextBox)
-						.Text_Raw(this, &FSelectionGroupCustomization::GetText)
-						.Font(FAppStyle::GetFontStyle(TEXT("PropertyWindow.NormalFont")))
-						.OnTextCommitted(this, &FSelectionGroupCustomization::OnTextCommitted)
-					]
-				]
-			];
-	}
-
-	FText FSelectionGroupCustomization::GetText() const
-	{
-		FText Text;
-		ChildPropertyHandle->GetValueAsFormattedText(Text);
-		return Text;
-	}
-
-	void FSelectionGroupCustomization::OnTextCommitted(const FText& NewText, ETextCommit::Type CommitInfo)
-	{
-		FText CurrentText;
-		ChildPropertyHandle->GetValueAsFormattedText(CurrentText);
-
-		if (!NewText.ToString().Equals(CurrentText.ToString(), ESearchCase::CaseSensitive))
-		{
-			FString String = NewText.ToString();
-			FClothDataflowTools::MakeCollectionName(String);
-			ChildPropertyHandle->SetValueFromFormattedString(String);
-		}
-	}
-
-	void FSelectionGroupCustomization::OnSelectionChanged(TSharedPtr<FText> ItemSelected, ESelectInfo::Type /*SelectInfo*/)
-	{
-		// Set the child property's value
-		if (ItemSelected)
-		{
-			FText CurrentText;
-			ChildPropertyHandle->GetValueAsFormattedText(CurrentText);
-
-			if (!ItemSelected->EqualTo(CurrentText))
-			{
-				ChildPropertyHandle->SetValueFromFormattedString(ItemSelected->ToString());
-			}
-
-			ComboButton->SetIsOpen(false);
-		}
-	}
-
-	TSharedRef<ITableRow> FSelectionGroupCustomization::MakeCategoryViewWidget(TSharedPtr<FText> Item, const TSharedRef< STableViewBase >& OwnerTable)
-	{
-		return SNew(STableRow<TSharedPtr<FString>>, OwnerTable)
-			[
-				SNew(STextBlock).Text(*Item.Get())
-			];
-	}
-
-	TSharedRef<SWidget> FSelectionGroupCustomization::OnGetMenuContent()
-	{
-		// Retrieve context if any, or use an empty context
-		const TSharedPtr<const SDataflowGraphEditor> DataflowGraphEditorPtr = DataflowGraphEditor.Pin();
-		const TSharedPtr<UE::Dataflow::FContext> Context = DataflowGraphEditorPtr ? DataflowGraphEditorPtr->GetDataflowContext() : TSharedPtr<UE::Dataflow::FContext>();
-
-		// Find all group names in the parent selection node's collection
-		GroupNames.Reset();
-		UE::Dataflow::FContextThreaded EmptyContext;
-		FManagedArrayCollection Collection = Private::GetPropertyCollection(Context.IsValid() ? *Context : EmptyContext, ChildPropertyHandle);
-		const TSharedRef<FManagedArrayCollection> ClothCollection = MakeShared<FManagedArrayCollection>(MoveTemp(Collection));
-		FCollectionClothFacade Cloth(ClothCollection);
-		for (const FName& GroupName : ClothCollection->GroupNames())
-		{
-			if (Cloth.IsValidClothCollectionGroupName(GroupName))  // Restrict to the cloth facade groups
-			{
-				GroupNames.Add(MakeShareable(new FText(FText::FromName(GroupName))));
-			}
-		}
-
-		return SNew(SVerticalBox)
-			+ SVerticalBox::Slot()
-			.AutoHeight()
-			.MaxHeight(400.0f)
-			[
-				SNew(SListView<TSharedPtr<FText>>)
-				.ListItemsSource(&GroupNames)
-				.SelectionMode(ESelectionMode::Single)
-				.OnGenerateRow(this, &FSelectionGroupCustomization::MakeCategoryViewWidget)
-				.OnSelectionChanged(this, &FSelectionGroupCustomization::OnSelectionChanged)
-			];
-	}
-
-	template<typename T>
-	T* FSelectionGroupCustomization::GetOwnerStruct() const
-	{
-		if (const TSharedPtr<IPropertyHandle> PropertyHandle = ChildPropertyHandle->GetParentHandle())
-		{
-			if (const TSharedPtr<IPropertyHandle> OwnerHandle = PropertyHandle->GetParentHandle())  // Assume that the group struct is only used at the node struct level
-			{
-				if (const TSharedPtr<IPropertyHandleStruct> OwnerHandleStruct = OwnerHandle->AsStruct())
-				{
-					if (TSharedPtr<FStructOnScope> StructOnScope = OwnerHandleStruct->GetStructData())
-					{
-						if (StructOnScope->GetStruct() == T::StaticStruct())
-						{
-							return reinterpret_cast<T*>(StructOnScope->GetStructMemory());
-						}
-					}
-				}
-			}
-		}
-		return nullptr;
+		return FCollectionClothFacade::GetValidClothCollectionGroupName();
 	}
 }  // End namespace UE::Chaos::ClothAsset
-
-#undef LOCTEXT_NAMESPACE
diff --git a/Engine/Plugins/ChaosClothAssetEditor/Source/ChaosClothAssetDataflowNodes/Private/ChaosClothAsset/SelectionGroupCustomization.h b/Engine/Plugins/ChaosClothAssetEditor/Source/ChaosClothAssetDataflowNodes/Private/ChaosClothAsset/SelectionGroupCustomization.h
--- a/Engine/Plugins/ChaosClothAssetEditor/Source/ChaosClothAssetDataflowNodes/Private/ChaosClothAsset/SelectionGroupCustomization.h
+++ b/Engine/Plugins/ChaosClothAssetEditor/Source/ChaosClothAssetDataflowNodes/Private/ChaosClothAsset/SelectionGroupCustomization.h
@@ -2,42 +2,21 @@
 
 #pragma once
 
-#include "IPropertyTypeCustomization.h"
-
-class IPropertyHandle;
-class ITableRow;
-class SComboButton;
-class STableViewBase;
-class SWidget;
-class SDataflowGraphEditor;
+#include "Dataflow/PropertyGroupCustomization.h"
 
 namespace UE::Chaos::ClothAsset
 {
 	/**
 	 * Customization for selection node group names input.
 	 */
-	class FSelectionGroupCustomization : public IPropertyTypeCustomization
+	class FSelectionGroupCustomization : public UE::Dataflow::FPropertyGroupCustomization
 	{
 	public:
-
 		static TSharedRef<IPropertyTypeCustomization> MakeInstance();
 
-		virtual void CustomizeHeader(TSharedRef<IPropertyHandle> InPropertyHandle, FDetailWidgetRow& HeaderRow, IPropertyTypeCustomizationUtils& CustomizationUtils) override;
-		virtual void CustomizeChildren(TSharedRef<IPropertyHandle> /*InPropertyHandle*/, IDetailChildrenBuilder& /*ChildBuilder*/, IPropertyTypeCustomizationUtils& /*CustomizationUtils*/) override {}
-
 	private:
-		FText GetText() const;
-		void OnTextCommitted(const FText& NewText, ETextCommit::Type CommitInfo);
-		void OnSelectionChanged(TSharedPtr<FText> ItemSelected, ESelectInfo::Type SelectInfo);
-		TSharedRef<ITableRow> MakeCategoryViewWidget(TSharedPtr<FText> Item, const TSharedRef<STableViewBase>& OwnerTable);
-		TSharedRef<SWidget> OnGetMenuContent();
-		
-		template<typename T>
-		T* GetOwnerStruct() const;
-
-		TWeakPtr<const SDataflowGraphEditor> DataflowGraphEditor;
-		TSharedPtr<IPropertyHandle> ChildPropertyHandle;
-		TSharedPtr<SComboButton> ComboButton;
-		TArray<TSharedPtr<FText>> GroupNames;
+		//~ Begin FPropertyGroupCustomization interface
+		virtual TArray<FName> GetTargetGroupNames(const FManagedArrayCollection& Collection) const override;
+		//~ End FPropertyGroupCustomization interface
 	};
 }
diff --git a/Engine/Plugins/Experimental/Dataflow/Source/DataflowEditor/Private/Dataflow/PropertyGroupCustomization.cpp b/Engine/Plugins/Experimental/Dataflow/Source/DataflowEditor/Private/Dataflow/PropertyGroupCustomization.cpp
new file mode 100644
--- /dev/null
+++ b/Engine/Plugins/Experimental/Dataflow/Source/DataflowEditor/Private/Dataflow/PropertyGroupCustomization.cpp
@@ -0,0 +1,228 @@
+// Copyright Epic Games, Inc. All Rights Reserved.
+
+#include "Dataflow/PropertyGroupCustomization.h"
+#include "Dataflow/DataflowGraphEditor.h"
+#include "Dataflow/DataflowNode.h"
+#include "Widgets/Input/SEditableTextBox.h"
+#include "Widgets/Input/SComboButton.h"
+#include "Widgets/Views/SListView.h"
+#include "Widgets/Views/STableRow.h"
+#include "DetailWidgetRow.h"
+
+#define LOCTEXT_NAMESPACE "PropertyGroupCustomization"
+
+namespace UE::Dataflow
+{
+	const FManagedArrayCollection& FPropertyGroupCustomization::GetPropertyCollection(
+		const TSharedPtr<UE::Dataflow::FContext>& Context,
+		const TSharedPtr<IPropertyHandle>& ChildPropertyHandle,
+		const FName CollectionPropertyName = FName(TEXT("Collection")))
+	{
+		static const FManagedArrayCollection EmptyCollection;
+
+		TSharedPtr<IPropertyHandle> OwnerHandle = ChildPropertyHandle;
+		while (TSharedPtr<IPropertyHandle> ParentHandle = OwnerHandle->GetParentHandle())
+		{
+			OwnerHandle = MoveTemp(ParentHandle);
+		}
+		if (const TSharedPtr<IPropertyHandleStruct> OwnerHandleStruct = OwnerHandle->AsStruct())
+		{
+			if (const TSharedPtr<FStructOnScope> StructOnScope = OwnerHandleStruct->GetStructData())
+			{
+				if (const UStruct* const Struct = StructOnScope->GetStruct())
+				{
+					if (Struct->IsChildOf<FDataflowNode>())
+					{
+						const FDataflowNode* const DataflowNode = reinterpret_cast<FDataflowNode*>(StructOnScope->GetStructMemory());
+
+						if (const FProperty* const Property = Struct->FindPropertyByName(CollectionPropertyName))
+						{
+							if (const FStructProperty* const StructProperty = CastField<FStructProperty>(Property))
+							{
+								if (StructProperty->GetCPPType(nullptr, CPPF_None) == TEXT("FManagedArrayCollection"))
+								{
+									if (const FDataflowInput* const DataflowInput = DataflowNode->FindInput(StructProperty->ContainerPtrToValuePtr<FManagedArrayCollection*>(DataflowNode)))
+									{
+										UE::Dataflow::FContextThreaded EmptyContext;
+										return DataflowInput->GetValue(Context.IsValid() ? *Context : EmptyContext, EmptyCollection);
+									}
+								}
+							}
+						}
+					}
+				}
+			}
+		}
+		return EmptyCollection;
+	}
+
+	bool FPropertyGroupCustomization::MakeGroupName(FString& InOutString)
+	{
+		const FString SourceString = InOutString;
+		InOutString = SlugStringForValidName(InOutString, TEXT("_")).Replace(TEXT("\\"), TEXT("_"));
+		bool bCharsWereRemoved;
+		do { InOutString.TrimCharInline(TEXT('_'), &bCharsWereRemoved); } while (bCharsWereRemoved);
+		return InOutString.Equals(SourceString);
+	}
+
+	TSharedRef<IPropertyTypeCustomization> FPropertyGroupCustomization::MakeInstance()
+	{
+		return MakeShareable(new FPropertyGroupCustomization);
+	}
+
+	void FPropertyGroupCustomization::CustomizeHeader(TSharedRef<IPropertyHandle> PropertyHandle, FDetailWidgetRow& HeaderRow, IPropertyTypeCustomizationUtils& CustomizationUtils)
+	{
+		DataflowGraphEditor = SDataflowGraphEditor::GetSelectedGraphEditor();
+
+		uint32 NumChildren;
+		const FPropertyAccess::Result Result = PropertyHandle->GetNumChildren(NumChildren);
+
+		ChildPropertyHandle = (Result == FPropertyAccess::Success && NumChildren) ? PropertyHandle->GetChildHandle(0) : nullptr;
+
+		GroupNames.Reset();
+
+		HeaderRow
+			.NameContent()
+			[
+				PropertyHandle->CreatePropertyNameWidget(PropertyHandle->GetPropertyDisplayName())
+			]
+			.ValueContent()
+			.MinDesiredWidth(250)
+			.MaxDesiredWidth(350.f)
+			[
+				SNew(SHorizontalBox)
+				+ SHorizontalBox::Slot()
+				.MaxWidth(145.f)
+				[
+					SAssignNew(ComboButton, SComboButton)
+					.ButtonStyle(FAppStyle::Get(), "NoBorder")
+					.ContentPadding(0)
+					.OnGetMenuContent(this, &FPropertyGroupCustomization::OnGetMenuContent)
+					.ButtonContent()
+					[
+						SNew(SEditableTextBox)
+						.Text(this, &FPropertyGroupCustomization::GetText)
+						.Font(FAppStyle::GetFontStyle(TEXT("PropertyWindow.NormalFont")))
+						.OnTextCommitted(this, &FPropertyGroupCustomization::OnTextCommitted)
+						.OnVerifyTextChanged(this, &FPropertyGroupCustomization::OnVerifyTextChanged)
+					]
+				]
+			];
+	}
+
+	FText FPropertyGroupCustomization::GetText() const
+	{
+		FText Text;
+		if (ChildPropertyHandle)
+		{
+			ChildPropertyHandle->GetValueAsFormattedText(Text);
+		}
+		return Text;
+	}
+
+	void FPropertyGroupCustomization::OnTextCommitted(const FText& NewText, ETextCommit::Type CommitInfo)
+	{
+		if (ChildPropertyHandle)
+		{
+			FText CurrentText;
+			ChildPropertyHandle->GetValueAsFormattedText(CurrentText);
+
+			if (!NewText.ToString().Equals(CurrentText.ToString(), ESearchCase::CaseSensitive))
+			{
+				FString String = NewText.ToString();
+				MakeGroupName(String);
+				ChildPropertyHandle->SetValueFromFormattedString(String);
+			}
+		}
+	}
+
+	void FPropertyGroupCustomization::OnSelectionChanged(TSharedPtr<FText> ItemSelected, ESelectInfo::Type /*SelectInfo*/)
+	{
+		if (ChildPropertyHandle)
+		{
+			// Set the child property's value
+			if (ItemSelected)
+			{
+				FText CurrentText;
+				ChildPropertyHandle->GetValueAsFormattedText(CurrentText);
+
+				if (!ItemSelected->EqualTo(CurrentText))
+				{
+					ChildPropertyHandle->SetValueFromFormattedString(ItemSelected->ToString());
+				}
+
+				if (TSharedPtr<SComboButton> PinnedComboButton = ComboButton.Pin())
+				{
+					PinnedComboButton->SetIsOpen(false);
+				}
+			}
+		}
+	}
+
+	bool FPropertyGroupCustomization::OnVerifyTextChanged(const FText& Text, FText& OutErrorMessage)
+	{
+		FString TextString = Text.ToString();
+		const bool bIsValidGroupName = MakeGroupName(TextString);
+		if (!bIsValidGroupName)
+		{
+			OutErrorMessage =
+				LOCTEXT("NotAValidGroupName",
+					"To be a valid group name, this text string musn't start by an underscore,\n"
+					"contain whitespaces, or any of the following character: \"',/.:|&!~@#(){}[]=;^%$`");
+		}
+		return bIsValidGroupName;
+	}
+
+	TSharedRef<ITableRow> FPropertyGroupCustomization::MakeCategoryViewWidget(TSharedPtr<FText> Item, const TSharedRef< STableViewBase >& OwnerTable)
+	{
+		if (Item)
+		{
+			return SNew(STableRow<TSharedPtr<FString>>, OwnerTable)
+			[
+				SNew(STextBlock).Text(*Item)
+			];
+		}
+		else
+		{
+			return SNew(STableRow<TSharedPtr<FString>>, OwnerTable);
+		}
+	}
+
+	TSharedRef<SWidget> FPropertyGroupCustomization::OnGetMenuContent()
+	{
+		const TSharedPtr<const SDataflowGraphEditor> DataflowGraphEditorPtr = DataflowGraphEditor.Pin();
+		const TSharedPtr<UE::Dataflow::FContext> Context = DataflowGraphEditorPtr ? DataflowGraphEditorPtr->GetDataflowContext() : TSharedPtr<UE::Dataflow::FContext>();
+
+		GroupNames.Reset();
+
+		// Retrieve collection
+		const FManagedArrayCollection& Collection = GetPropertyCollection(Context, ChildPropertyHandle, GetCollectionPropertyName());
+
+		// Find all group names in the parent node's collection
+		TArray<FName> CollectionGroupNames;
+		const TArray<FName> TargetGroupNames = GetTargetGroupNames(Collection);
+
+		// Find all group names in the parent selection node's collection
+		for (const FName& GroupName : Collection.GroupNames())
+		{
+			if (TargetGroupNames.Contains(GroupName))
+			{
+				GroupNames.Add(MakeShareable(new FText(FText::FromName(GroupName))));
+			}
+		}
+
+		return SNew(SVerticalBox)
+			+ SVerticalBox::Slot()
+			.AutoHeight()
+			.MaxHeight(400.0f)
+			[
+				SNew(SListView<TSharedPtr<FText>>)
+					.ListItemsSource(&GroupNames)
+					.SelectionMode(ESelectionMode::Single)
+					.OnGenerateRow(this, &FPropertyGroupCustomization::MakeCategoryViewWidget)
+					.OnSelectionChanged(this, &FPropertyGroupCustomization::OnSelectionChanged)
+			];
+	}
+}  // End namespace UE::Chaos::ClothAsset
+
+#undef LOCTEXT_NAMESPACE
diff --git a/Engine/Plugins/Experimental/Dataflow/Source/DataflowEditor/Private/Dataflow/ScalarVertexPropertyGroupCustomization.cpp b/Engine/Plugins/Experimental/Dataflow/Source/DataflowEditor/Private/Dataflow/ScalarVertexPropertyGroupCustomization.cpp
--- a/Engine/Plugins/Experimental/Dataflow/Source/DataflowEditor/Private/Dataflow/ScalarVertexPropertyGroupCustomization.cpp
+++ b/Engine/Plugins/Experimental/Dataflow/Source/DataflowEditor/Private/Dataflow/ScalarVertexPropertyGroupCustomization.cpp
@@ -1,14 +1,7 @@
 // Copyright Epic Games, Inc. All Rights Reserved.
 
-#include "ScalarVertexPropertyGroupCustomization.h"
+#include "Dataflow/ScalarVertexPropertyGroupCustomization.h"
 #include "Dataflow/DataflowCollectionAddScalarVertexPropertyNode.h"
-#include "Widgets/Input/SEditableTextBox.h"
-#include "Widgets/Input/SComboButton.h"
-#include "Widgets/Views/SListView.h"
-#include "Widgets/Views/STableRow.h"
-#include "DetailWidgetRow.h"
-
-#define LOCTEXT_NAMESPACE "ScalarVertexPropertyGroupCustomization"
 
 namespace UE::Dataflow
 {
@@ -17,170 +10,14 @@ namespace UE::Dataflow
 		return MakeShareable(new FScalarVertexPropertyGroupCustomization);
 	}
 
-	void FScalarVertexPropertyGroupCustomization::CustomizeHeader(TSharedRef<IPropertyHandle> PropertyHandle, FDetailWidgetRow& HeaderRow, IPropertyTypeCustomizationUtils& CustomizationUtils)
-	{
-		uint32 NumChildren;
-		const FPropertyAccess::Result Result = PropertyHandle->GetNumChildren(NumChildren);
-
-		ChildPropertyHandle = (Result == FPropertyAccess::Success && NumChildren) ? PropertyHandle->GetChildHandle(0) : nullptr;
-
-		GroupNames.Reset();
-
-		HeaderRow
-			.NameContent()
-			[
-				PropertyHandle->CreatePropertyNameWidget(PropertyHandle->GetPropertyDisplayName())
-			]
-			.ValueContent()
-			.MinDesiredWidth(250)
-			.MaxDesiredWidth(350.f)
-			[
-				SNew(SHorizontalBox)
-				+ SHorizontalBox::Slot()
-				.MaxWidth(145.f)
-				[
-					SAssignNew(ComboButton, SComboButton)
-					.ButtonStyle(FAppStyle::Get(), "NoBorder")
-					.ContentPadding(0)
-					.OnGetMenuContent(this, &FScalarVertexPropertyGroupCustomization::OnGetMenuContent)
-					.ButtonContent()
-					[
-						SNew(SEditableTextBox)
-						.Text_Raw(this, &FScalarVertexPropertyGroupCustomization::GetText)
-						.Font(FAppStyle::GetFontStyle(TEXT("PropertyWindow.NormalFont")))
-						.OnTextCommitted(this, &FScalarVertexPropertyGroupCustomization::OnTextCommitted)
-					]
-				]
-			];
-	}
-
-	FText FScalarVertexPropertyGroupCustomization::GetText() const
-	{
-		FText Text;
-		if (ChildPropertyHandle)
-		{
-			ChildPropertyHandle->GetValueAsFormattedText(Text);
-		}
-		return Text;
-	}
-
-	void FScalarVertexPropertyGroupCustomization::OnTextCommitted(const FText& NewText, ETextCommit::Type CommitInfo)
-	{
-		if (ChildPropertyHandle)
-		{
-			FText CurrentText;
-			ChildPropertyHandle->GetValueAsFormattedText(CurrentText);
-
-			if (!NewText.ToString().Equals(CurrentText.ToString(), ESearchCase::CaseSensitive))
-			{
-				FString String = NewText.ToString();
-				//FClothDataflowTools::MakeCollectionName(String);
-				ChildPropertyHandle->SetValueFromFormattedString(String);
-			}
-		}
-	}
-
-	void FScalarVertexPropertyGroupCustomization::OnSelectionChanged(TSharedPtr<FText> ItemSelected, ESelectInfo::Type /*SelectInfo*/)
+	TArray<FName> FScalarVertexPropertyGroupCustomization::GetTargetGroupNames(const FManagedArrayCollection& Collection) const
 	{
-		if (ChildPropertyHandle)
-		{
-			// Set the child property's value
-			if (ItemSelected)
-			{
-				FText CurrentText;
-				ChildPropertyHandle->GetValueAsFormattedText(CurrentText);
-
-				if (!ItemSelected->EqualTo(CurrentText))
-				{
-					ChildPropertyHandle->SetValueFromFormattedString(ItemSelected->ToString());
-				}
-
-				if (TSharedPtr<SComboButton> PinnedComboButton = ComboButton.Pin())
-				{
-					PinnedComboButton->SetIsOpen(false);
-				}
-			}
-		}
-	}
-
-	TSharedRef<ITableRow> FScalarVertexPropertyGroupCustomization::MakeCategoryViewWidget(TSharedPtr<FText> Item, const TSharedRef< STableViewBase >& OwnerTable)
-	{
-		if (Item)
-		{
-			return SNew(STableRow<TSharedPtr<FString>>, OwnerTable)
-			[
-				SNew(STextBlock).Text(*Item)
-			];
-		}
-		else
-		{
-			return SNew(STableRow<TSharedPtr<FString>>, OwnerTable);
-		}
-	}
-
-	TSharedRef<SWidget> FScalarVertexPropertyGroupCustomization::OnGetMenuContent()
-	{
-		GroupNames.Reset();
-
-		// Find all group names in the parent node's collection
-		FName NodeType;
-		TArray<FName> CollectionGroupNames;
-		if (const FDataflowCollectionAddScalarVertexPropertyNode* const ScalarVertexPropertyNode = GetOwnerStruct<FDataflowCollectionAddScalarVertexPropertyNode>())
-		{
-			// TODO: 
-			//   Fix this using the new context based paradigm (see SelectionGroupCustomization.cpp GetCollection() implementation).
-			//CollectionGroupNames = ScalarVertexPropertyNode->GetCachedCollectionGroupNames();
-			NodeType = ScalarVertexPropertyNode->GetType();
-
-			const TArray<FName> AvailableTargetGroups = DataflowAddScalarVertexPropertyCallbackRegistry::Get().GetTargetGroupNames();
-
-			CollectionGroupNames = CollectionGroupNames.FilterByPredicate([&AvailableTargetGroups](const FName& Name)
-				{
-					return AvailableTargetGroups.Contains(Name);
-				});
-
-			for (const FName& CollectionGroupName : CollectionGroupNames)
-			{
-				GroupNames.Add(MakeShareable(new FText(FText::FromName(CollectionGroupName))));
-			}
-		}
-
-		return SNew(SVerticalBox)
-			+ SVerticalBox::Slot()
-			.AutoHeight()
-			.MaxHeight(400.0f)
-			[
-				SNew(SListView<TSharedPtr<FText>>)
-					.ListItemsSource(&GroupNames)
-					.SelectionMode(ESelectionMode::Single)
-					.OnGenerateRow(this, &FScalarVertexPropertyGroupCustomization::MakeCategoryViewWidget)
-					.OnSelectionChanged(this, &FScalarVertexPropertyGroupCustomization::OnSelectionChanged)
-			];
+		return DataflowAddScalarVertexPropertyCallbackRegistry::Get().GetTargetGroupNames();
 	}
 
-	template<typename T>
-	T* FScalarVertexPropertyGroupCustomization::GetOwnerStruct() const
+	FName FScalarVertexPropertyGroupCustomization::GetCollectionPropertyName() const
 	{
-		if (ChildPropertyHandle)
-		{
-			if (const TSharedPtr<IPropertyHandle> PropertyHandle = ChildPropertyHandle->GetParentHandle())
-			{
-				if (const TSharedPtr<IPropertyHandle> OwnerHandle = PropertyHandle->GetParentHandle())  // Assume that the group struct is only used at the node struct level
-				{
-					if (const TSharedPtr<IPropertyHandleStruct> OwnerHandleStruct = OwnerHandle->AsStruct())
-					{
-						if (TSharedPtr<FStructOnScope> StructOnScope = OwnerHandleStruct->GetStructData())
-						{
-							if (StructOnScope->GetStruct() == T::StaticStruct())
-							{
-								return reinterpret_cast<T*>(StructOnScope->GetStructMemory());
-							}
-						}
-					}
-				}
-			}
-		}
-		return nullptr;
+		return GET_MEMBER_NAME_CHECKED(FDataflowCollectionAddScalarVertexPropertyNode, Collection);
 	}
 }  // End namespace UE::Chaos::ClothAsset
 
diff --git a/Engine/Plugins/Experimental/Dataflow/Source/DataflowEditor/Private/Dataflow/ScalarVertexPropertyGroupCustomization.h b/Engine/Plugins/Experimental/Dataflow/Source/DataflowEditor/Private/Dataflow/ScalarVertexPropertyGroupCustomization.h
--- a/Engine/Plugins/Experimental/Dataflow/Source/DataflowEditor/Private/Dataflow/ScalarVertexPropertyGroupCustomization.h
+++ b/Engine/Plugins/Experimental/Dataflow/Source/DataflowEditor/Private/Dataflow/ScalarVertexPropertyGroupCustomization.h
@@ -2,38 +2,22 @@
 
 #pragma once
 
-#include "IPropertyTypeCustomization.h"
-
-class IPropertyHandle;
-class ITableRow;
-class SComboButton;
-class STableViewBase;
-class SWidget;
+#include "Dataflow/PropertyGroupCustomization.h"
 
 namespace UE::Dataflow
 {
-
-	class FScalarVertexPropertyGroupCustomization : public IPropertyTypeCustomization
+	/**
+	 * Customization for scalara vertex node group names input.
+	 */
+	class FScalarVertexPropertyGroupCustomization : public FPropertyGroupCustomization
 	{
 	public:
-
 		static TSharedRef<IPropertyTypeCustomization> MakeInstance();
 
 	private:
-		virtual void CustomizeHeader(TSharedRef<IPropertyHandle> InPropertyHandle, FDetailWidgetRow& HeaderRow, IPropertyTypeCustomizationUtils& CustomizationUtils) override;
-		virtual void CustomizeChildren(TSharedRef<IPropertyHandle> /*InPropertyHandle*/, IDetailChildrenBuilder& /*ChildBuilder*/, IPropertyTypeCustomizationUtils& /*CustomizationUtils*/) override {}
-
-		FText GetText() const;
-		void OnTextCommitted(const FText& NewText, ETextCommit::Type CommitInfo);
-		void OnSelectionChanged(TSharedPtr<FText> ItemSelected, ESelectInfo::Type SelectInfo);
-		TSharedRef<ITableRow> MakeCategoryViewWidget(TSharedPtr<FText> Item, const TSharedRef<STableViewBase>& OwnerTable);
-		TSharedRef<SWidget> OnGetMenuContent();
-		
-		template<typename T>
-		T* GetOwnerStruct() const;
-
-		TSharedPtr<IPropertyHandle> ChildPropertyHandle;
-		TWeakPtr<SComboButton> ComboButton;
-		TArray<TSharedPtr<FText>> GroupNames;
+		//~ Begin FPropertyGroupCustomization interface
+		virtual TArray<FName> GetTargetGroupNames(const FManagedArrayCollection& Collection) const override;
+		virtual FName GetCollectionPropertyName() const override;
+		//~ End FPropertyGroupCustomization interface
 	};
 }
diff --git a/Engine/Plugins/Experimental/Dataflow/Source/DataflowEditor/Public/Dataflow/PropertyGroupCustomization.h b/Engine/Plugins/Experimental/Dataflow/Source/DataflowEditor/Public/Dataflow/PropertyGroupCustomization.h
new file mode 100644
--- /dev/null
+++ b/Engine/Plugins/Experimental/Dataflow/Source/DataflowEditor/Public/Dataflow/PropertyGroupCustomization.h
@@ -0,0 +1,81 @@
+// Copyright Epic Games, Inc. All Rights Reserved.
+
+#pragma once
+
+#include "GeometryCollection/ManagedArrayCollection.h"
+#include "IPropertyTypeCustomization.h"
+
+class IPropertyHandle;
+class ITableRow;
+class SComboButton;
+class STableViewBase;
+class SWidget;
+class SDataflowGraphEditor;
+
+namespace UE::Dataflow
+{
+	class FContext;
+
+	/**
+	 * Property group customization to allow selecting of a group from a list of the groups currently held by the node's collection.
+	 */
+	class FPropertyGroupCustomization : public IPropertyTypeCustomization
+	{
+	public:
+		static TSharedRef<IPropertyTypeCustomization> MakeInstance();
+
+	private:
+		//~ Begin IPropertyTypeCustomization interface
+		DATAFLOWEDITOR_API virtual void CustomizeHeader(TSharedRef<IPropertyHandle> InPropertyHandle, FDetailWidgetRow& HeaderRow, IPropertyTypeCustomizationUtils& CustomizationUtils) override;
+		DATAFLOWEDITOR_API virtual void CustomizeChildren(TSharedRef<IPropertyHandle> /*InPropertyHandle*/, IDetailChildrenBuilder& /*ChildBuilder*/, IPropertyTypeCustomizationUtils& /*CustomizationUtils*/) override {}
+		//~ End IPropertyTypeCustomization interface
+
+		/**
+		 * Return the FManagedArrayCollection with the specified name from the property held by the top level struct owner of ChildPropertyHandle.
+		 * @param Context The current editor Dataflow context.
+		 * @param ChildPropertyHandle The handle of a customized property from which to find the owner struct.
+		 * @param CollectionPropertyName
+		 * @return The collection.
+		 */
+		static DATAFLOWEDITOR_API const FManagedArrayCollection& GetPropertyCollection(
+			const TSharedPtr<UE::Dataflow::FContext>& Context,
+			const TSharedPtr<IPropertyHandle>& ChildPropertyHandle,
+			const FName CollectionPropertyName);
+
+		/**
+		 * Turn a string into a valid collection group or attribute name.
+		 * The resulting name won't contains spaces and any other special characters as listed in
+		 * INVALID_OBJECTNAME_CHARACTERS (currently "',/.:|&!~\n\r\t@#(){}[]=;^%$`).
+		 * It will also have all leading underscore removed, as these names are reserved for internal use.
+		 * @param InOutString The string to turn into a valid collection name.
+		 * @return Whether the InOutString was already a valid collection name.
+		 */
+		static DATAFLOWEDITOR_API bool MakeGroupName(FString& InOutString);
+
+		/** List of valid group names for the drop down list. Override this method to filter to a specific set of group names. */
+		virtual TArray<FName> GetTargetGroupNames(const FManagedArrayCollection& Collection) const
+		{
+			return Collection.GroupNames();  // By default returns all of the collection's group names
+		}
+
+		/** Name of the collection property. Override this method to specify your own collection property name. */
+		virtual FName GetCollectionPropertyName() const
+		{
+			return DefaultCollectionPropertyName;  // By default returns "Collection"
+		}
+
+		FText GetText() const;
+		void OnTextCommitted(const FText& NewText, ETextCommit::Type CommitInfo);
+		void OnSelectionChanged(TSharedPtr<FText> ItemSelected, ESelectInfo::Type SelectInfo);
+		bool OnVerifyTextChanged(const FText& Text, FText& OutErrorMessage);
+		TSharedRef<ITableRow> MakeCategoryViewWidget(TSharedPtr<FText> Item, const TSharedRef<STableViewBase>& OwnerTable);
+		TSharedRef<SWidget> OnGetMenuContent();
+		
+		static inline const FName DefaultCollectionPropertyName = TEXT("Collection");
+
+		TWeakPtr<const SDataflowGraphEditor> DataflowGraphEditor;
+		TSharedPtr<IPropertyHandle> ChildPropertyHandle;
+		TWeakPtr<SComboButton> ComboButton;
+		TArray<TSharedPtr<FText>> GroupNames;
+	};
+}
