Chaos Cloth Asset - Updated Cloth Terminal node refresh to also work when the node is disabled.

[RN] Chaos Cloth Asset - Fixed the Cloth Terminal node Active checkbox.

diff --git a/Engine/Plugins/ChaosClothAssetEditor/Source/ChaosClothAssetDataflowNodes/Private/ChaosClothAsset/ChaosClothAssetDataflowNodesModule.cpp b/Engine/Plugins/ChaosClothAssetEditor/Source/ChaosClothAssetDataflowNodes/Private/ChaosClothAsset/ChaosClothAssetDataflowNodesModule.cpp
--- a/Engine/Plugins/ChaosClothAssetEditor/Source/ChaosClothAssetDataflowNodes/Private/ChaosClothAsset/ChaosClothAssetDataflowNodesModule.cpp
+++ b/Engine/Plugins/ChaosClothAssetEditor/Source/ChaosClothAssetDataflowNodes/Private/ChaosClothAsset/ChaosClothAssetDataflowNodesModule.cpp
@@ -55,7 +55,6 @@
 #include "ChaosClothAsset/SkinningBlendNode.h"
 #include "ChaosClothAsset/StaticMeshImportNode.h"
 #include "ChaosClothAsset/TerminalNode.h"
-#include "ChaosClothAsset/TerminalNodeRefreshAssetCustomization.h"
 #include "ChaosClothAsset/TransferSkinWeightsNode.h"
 #include "ChaosClothAsset/TransformPositionsNode.h"
 #include "ChaosClothAsset/TransformUVsNode.h"
@@ -312,8 +311,6 @@ PRAGMA_ENABLE_DEPRECATION_WARNINGS
 				PropertyModule->RegisterCustomPropertyTypeLayout(FChaosClothAssetConnectableIOStringValue::StaticStruct()->GetFName(), FOnGetPropertyTypeCustomizationInstance::CreateStatic(&FConnectableValueCustomization::MakeInstance));
 				PropertyModule->RegisterCustomPropertyTypeLayout(FChaosClothAssetNodeSelectionGroup::StaticStruct()->GetFName(), FOnGetPropertyTypeCustomizationInstance::CreateStatic(&FSelectionGroupCustomization::MakeInstance));
 				PropertyModule->RegisterCustomPropertyTypeLayout(FChaosClothAssetImportFilePath::StaticStruct()->GetFName(), FOnGetPropertyTypeCustomizationInstance::CreateStatic(&FImportFilePathCustomization::MakeInstance));
-				PropertyModule->RegisterCustomPropertyTypeLayout(FChaosClothAssetTerminalNodeRefreshAsset::StaticStruct()->GetFName(), FOnGetPropertyTypeCustomizationInstance::CreateStatic(&FTerminalNodeRefreshAssetCustomization::MakeInstance));
-				PropertyModule->RegisterCustomPropertyTypeLayout(FChaosClothAssetImportNodeRefreshAsset::StaticStruct()->GetFName(), FOnGetPropertyTypeCustomizationInstance::CreateStatic(&FTerminalNodeRefreshAssetCustomization::MakeInstance));
 				PropertyModule->RegisterCustomPropertyTypeLayout(FChaosClothAssetNodeAttributeGroup::StaticStruct()->GetFName(), FOnGetPropertyTypeCustomizationInstance::CreateStatic(&FSelectionGroupCustomization::MakeInstance));
 				PropertyModule->RegisterCustomPropertyTypeLayout(FChaosClothAssetImportedVectorValue::StaticStruct()->GetFName(), FOnGetPropertyTypeCustomizationInstance::CreateStatic(&FImportedValueCustomization::MakeInstance));
 				PropertyModule->RegisterCustomPropertyTypeLayout(FChaosClothAssetImportedFloatValue::StaticStruct()->GetFName(), FOnGetPropertyTypeCustomizationInstance::CreateStatic(&FImportedValueCustomization::MakeInstance));
@@ -344,8 +341,6 @@ PRAGMA_ENABLE_DEPRECATION_WARNINGS
 					PropertyModule->UnregisterCustomPropertyTypeLayout(FChaosClothAssetConnectableIOStringValue::StaticStruct()->GetFName());
 					PropertyModule->UnregisterCustomPropertyTypeLayout(FChaosClothAssetNodeSelectionGroup::StaticStruct()->GetFName());
 					PropertyModule->UnregisterCustomPropertyTypeLayout(FChaosClothAssetImportFilePath::StaticStruct()->GetFName());
-					PropertyModule->UnregisterCustomPropertyTypeLayout(FChaosClothAssetTerminalNodeRefreshAsset::StaticStruct()->GetFName());
-					PropertyModule->UnregisterCustomPropertyTypeLayout(FChaosClothAssetImportNodeRefreshAsset::StaticStruct()->GetFName());
 					PropertyModule->UnregisterCustomPropertyTypeLayout(FChaosClothAssetNodeAttributeGroup::StaticStruct()->GetFName());
 					PropertyModule->UnregisterCustomPropertyTypeLayout(FChaosClothAssetImportedVectorValue::StaticStruct()->GetFName());
 					PropertyModule->UnregisterCustomPropertyTypeLayout(FChaosClothAssetImportedFloatValue::StaticStruct()->GetFName());
diff --git a/Engine/Plugins/ChaosClothAssetEditor/Source/ChaosClothAssetDataflowNodes/Private/ChaosClothAsset/ImportNode.cpp b/Engine/Plugins/ChaosClothAssetEditor/Source/ChaosClothAssetDataflowNodes/Private/ChaosClothAsset/ImportNode.cpp
--- a/Engine/Plugins/ChaosClothAssetEditor/Source/ChaosClothAssetDataflowNodes/Private/ChaosClothAsset/ImportNode.cpp
+++ b/Engine/Plugins/ChaosClothAssetEditor/Source/ChaosClothAssetDataflowNodes/Private/ChaosClothAsset/ImportNode.cpp
@@ -29,11 +29,6 @@ void FChaosClothAssetImportNode::Evaluate(UE::Dataflow::FContext& Context, const
 
 	if (Out->IsA<FManagedArrayCollection>(&Collection))
 	{
-		if (ReimportAsset.bRefreshAsset)
-		{
-			ReimportAsset.bRefreshAsset = false;
-		}
-
 		bool bSetValue = false;
 		if (ClothAsset)
 		{
@@ -100,7 +95,7 @@ void FChaosClothAssetImportNode::Evaluate(UE::Dataflow::FContext& Context, const
 			}
 		}
 		SetValue(Context, MoveTemp(*ClothCollection), &Collection);
-}
+	}
 }
 
 #undef LOCTEXT_NAMESPACE
diff --git a/Engine/Plugins/ChaosClothAssetEditor/Source/ChaosClothAssetDataflowNodes/Private/ChaosClothAsset/TerminalNode.cpp b/Engine/Plugins/ChaosClothAssetEditor/Source/ChaosClothAssetDataflowNodes/Private/ChaosClothAsset/TerminalNode.cpp
--- a/Engine/Plugins/ChaosClothAssetEditor/Source/ChaosClothAssetDataflowNodes/Private/ChaosClothAsset/TerminalNode.cpp
+++ b/Engine/Plugins/ChaosClothAssetEditor/Source/ChaosClothAssetDataflowNodes/Private/ChaosClothAsset/TerminalNode.cpp
@@ -89,6 +89,7 @@ namespace UE::Chaos::ClothAsset::Private
 
 FChaosClothAssetTerminalNode_v2::FChaosClothAssetTerminalNode_v2(const UE::Dataflow::FNodeParameters& InParam, FGuid InGuid)
 	: FDataflowTerminalNode(InParam, InGuid)
+	, RefreshAsset(FDataflowFunctionProperty::FDelegate::CreateLambda([this](UE::Dataflow::FContext& /*Context*/) { bClothCollectionChecksumValid = false; }))
 {
 	// Start with Lod0
 	for (int32 Index = 0; Index < NumInitialCollectionLods; ++Index)
@@ -104,13 +105,7 @@ void FChaosClothAssetTerminalNode_v2::SetAssetValue(TObjectPtr<UObject> Asset, U
 	{
 		using namespace UE::Chaos::ClothAsset;
 
-		if (RefreshAsset.bRefreshAsset)
-		{
-			bClothCollectionChecksumValid = false;
-			RefreshAsset.bRefreshAsset = false;
-		}
-
-		TArray<TSharedRef<FManagedArrayCollection>> InClothCollections = GetCleanedCollectionLodValues(Context);
+		TArray<TSharedRef<FManagedArrayCollection>> InClothCollections = bActive ? GetCleanedCollectionLodValues(Context) : TArray<TSharedRef<FManagedArrayCollection>>();
 		TArray<TSharedRef<FManagedArrayCollection>>& ClothCollections = ClothAsset->GetClothCollections();
 
 		const uint32 PreviousChecksum = ClothColllectionChecksum;
@@ -386,13 +381,15 @@ void FChaosClothAssetTerminalNode::SetAssetValue(TObjectPtr<UObject> Asset, UE::
 	{
 		using namespace UE::Chaos::ClothAsset;
 
+PRAGMA_DISABLE_DEPRECATION_WARNINGS
 		if (RefreshAsset.bRefreshAsset)
 		{
 			bClothCollectionChecksumValid = false;
 			RefreshAsset.bRefreshAsset = false;
 		}
+PRAGMA_ENABLE_DEPRECATION_WARNINGS
 
-		TArray<TSharedRef<FManagedArrayCollection>> InClothCollections = GetCleanedCollectionLodValues(Context);
+		TArray<TSharedRef<FManagedArrayCollection>> InClothCollections = bActive ? GetCleanedCollectionLodValues(Context) : TArray<TSharedRef<FManagedArrayCollection>>();
 		TArray<TSharedRef<FManagedArrayCollection>>& ClothCollections = ClothAsset->GetClothCollections();
 
 		const uint32 PreviousChecksum = ClothColllectionChecksum;
diff --git a/Engine/Plugins/ChaosClothAssetEditor/Source/ChaosClothAssetDataflowNodes/Private/ChaosClothAsset/TerminalNodeRefreshAssetCustomization.cpp b/Engine/Plugins/ChaosClothAssetEditor/Source/ChaosClothAssetDataflowNodes/Private/ChaosClothAsset/TerminalNodeRefreshAssetCustomization.cpp
deleted file mode 100644
diff --git a/Engine/Plugins/ChaosClothAssetEditor/Source/ChaosClothAssetDataflowNodes/Private/ChaosClothAsset/TerminalNodeRefreshAssetCustomization.h b/Engine/Plugins/ChaosClothAssetEditor/Source/ChaosClothAssetDataflowNodes/Private/ChaosClothAsset/TerminalNodeRefreshAssetCustomization.h
deleted file mode 100644
diff --git a/Engine/Plugins/ChaosClothAssetEditor/Source/ChaosClothAssetDataflowNodes/Public/ChaosClothAsset/ImportNode.h b/Engine/Plugins/ChaosClothAssetEditor/Source/ChaosClothAssetDataflowNodes/Public/ChaosClothAsset/ImportNode.h
--- a/Engine/Plugins/ChaosClothAssetEditor/Source/ChaosClothAssetDataflowNodes/Public/ChaosClothAsset/ImportNode.h
+++ b/Engine/Plugins/ChaosClothAssetEditor/Source/ChaosClothAssetDataflowNodes/Public/ChaosClothAsset/ImportNode.h
@@ -2,6 +2,7 @@
 
 #pragma once 
 
+#include "Dataflow/DataflowFunctionProperty.h"
 #include "Dataflow/DataflowNode.h"
 #include "GeometryCollection/ManagedArrayCollection.h"
 #include "ImportNode.generated.h"
@@ -9,14 +10,14 @@
 #if UE_ENABLE_INCLUDE_ORDER_DEPRECATED_IN_5_5
 namespace Dataflow = UE::Dataflow;
 #else
-namespace UE_DEPRECATED(5.5, "Use UE::Dafaflow instead.") Dataflow {}
+namespace UE_DEPRECATED(5.5, "Use UE::Dataflow instead.") Dataflow {}
 #endif
 
 class UChaosClothAsset;
 
 /** Refresh structure for push button customization. */
 USTRUCT()
-struct FChaosClothAssetImportNodeRefreshAsset
+struct UE_DEPRECATED(5.5, "Use UE::Dataflow::FunctionProperty instead.") FChaosClothAssetImportNodeRefreshAsset
 {
 	GENERATED_USTRUCT_BODY()
 
@@ -25,9 +26,11 @@ struct FChaosClothAssetImportNodeRefreshAsset
 };
 
 /** Import an existing Cloth Asset into the graph. */
+PRAGMA_DISABLE_DEPRECATION_WARNINGS  // For Reimport
 USTRUCT(Meta = (DataflowCloth))
 struct FChaosClothAssetImportNode : public FDataflowNode
 {
+PRAGMA_ENABLE_DEPRECATION_WARNINGS
 	GENERATED_USTRUCT_BODY()
 	DATAFLOW_NODE_DEFINE_INTERNAL(FChaosClothAssetImportNode, "ClothAssetImport", "Cloth", "Cloth Asset Import")
 	DATAFLOW_NODE_RENDER_TYPE("SurfaceRender", FName("FClothCollection"), "Collection")
@@ -44,11 +47,15 @@ public:
 	UPROPERTY(EditAnywhere, Category = "Cloth Asset Import", Meta = (DisplayName = "Import LOD", ClampMin = "0"))
 	int32 ImportLod = 0;
 
-	/**
-	 * Reimport the imported asset. 
-	 */
-	UPROPERTY(EditAnywhere, Category = "Cloth Asset Import")
+	/** Reimport the imported cloth asset. */
+	UPROPERTY(EditAnywhere, Category = "Cloth Asset Import", Meta = (ButtonImage = "Persona.ReimportAsset", DisplayName = "Reimport Asset"))
+	FDataflowFunctionProperty Reimport;
+
+PRAGMA_DISABLE_DEPRECATION_WARNINGS
+	UE_DEPRECATED(5.5, "Use FDataflowFunctionProperty Reimport instead.")
+	UPROPERTY()
 	mutable FChaosClothAssetImportNodeRefreshAsset ReimportAsset;
+PRAGMA_ENABLE_DEPRECATION_WARNINGS
 
 	FChaosClothAssetImportNode(const UE::Dataflow::FNodeParameters& InParam, FGuid InGuid = FGuid::NewGuid());
 
diff --git a/Engine/Plugins/ChaosClothAssetEditor/Source/ChaosClothAssetDataflowNodes/Public/ChaosClothAsset/TerminalNode.h b/Engine/Plugins/ChaosClothAssetEditor/Source/ChaosClothAssetDataflowNodes/Public/ChaosClothAsset/TerminalNode.h
--- a/Engine/Plugins/ChaosClothAssetEditor/Source/ChaosClothAssetDataflowNodes/Public/ChaosClothAsset/TerminalNode.h
+++ b/Engine/Plugins/ChaosClothAssetEditor/Source/ChaosClothAssetDataflowNodes/Public/ChaosClothAsset/TerminalNode.h
@@ -2,6 +2,7 @@
 
 #pragma once 
 
+#include "Dataflow/DataflowFunctionProperty.h"
 #include "Dataflow/DataflowTerminalNode.h"
 #include "GeometryCollection/ManagedArrayCollection.h"
 #include "ChaosClothAsset/ClothLodTransitionDataCache.h"
@@ -10,12 +11,12 @@
 #if UE_ENABLE_INCLUDE_ORDER_DEPRECATED_IN_5_5
 namespace Dataflow = UE::Dataflow;
 #else
-namespace UE_DEPRECATED(5.5, "Use UE::Dafaflow instead.") Dataflow {}
+namespace UE_DEPRECATED(5.5, "Use UE::Dataflow instead.") Dataflow {}
 #endif
 
 /** Refresh structure for push buton customization. */
 USTRUCT()
-struct FChaosClothAssetTerminalNodeRefreshAsset
+struct UE_DEPRECATED(5.5, "Use UE::Dataflow::FunctionProperty instead.") FChaosClothAssetTerminalNodeRefreshAsset
 {
 	GENERATED_USTRUCT_BODY()
 
@@ -32,6 +33,7 @@ struct FChaosClothAssetTerminalNode_v2 : public FDataflowTerminalNode
 	DATAFLOW_NODE_RENDER_TYPE("SurfaceRender", FName("FClothCollection"), "Collection")
 
 public:
+	/** Input cloth collection for this LOD. */
 	UPROPERTY()
 	TArray<FManagedArrayCollection> CollectionLods;
 
@@ -40,8 +42,8 @@ public:
 	 * Note that it is not required to manually refresh the cloth asset, this is done automatically when there is a change in the Dataflow.
 	 * This function is a developper utility used for debugging.
 	 */
-	UPROPERTY(EditAnywhere, Category = "Cloth Asset Terminal")
-	mutable FChaosClothAssetTerminalNodeRefreshAsset RefreshAsset;
+	UPROPERTY(EditAnywhere, Category = "Cloth Asset Terminal", Meta = (ButtonImage = "Icons.Refresh"))
+	FDataflowFunctionProperty RefreshAsset;
 
 	FChaosClothAssetTerminalNode_v2(const UE::Dataflow::FNodeParameters& InParam, FGuid InGuid = FGuid::NewGuid());
 
@@ -70,12 +72,12 @@ private:
 	static constexpr int32 NumInitialCollectionLods = 1;
 };
 
-
-
 /** Cloth terminal node to generate a cloth asset from a cloth collection. */
+PRAGMA_DISABLE_DEPRECATION_WARNINGS  // For RefreshAsset
 USTRUCT(Meta = (DataflowCloth, DataflowTerminal, Deprecated = 5.5))
 struct UE_DEPRECATED(5.5, "Use the newer version of this node instead.") FChaosClothAssetTerminalNode : public FDataflowTerminalNode
 {
+PRAGMA_ENABLE_DEPRECATION_WARNINGS
 	GENERATED_USTRUCT_BODY()
 	DATAFLOW_NODE_DEFINE_INTERNAL(FChaosClothAssetTerminalNode, "ClothAssetTerminal", "Cloth", "Cloth Terminal")  // TODO: Should the category be Terminal instead like all other terminal nodes
 	DATAFLOW_NODE_RENDER_TYPE("SurfaceRender", FName("FClothCollection"), "Collection")
@@ -111,8 +113,10 @@ PRAGMA_ENABLE_DEPRECATION_WARNINGS
 	 * Note that it is not required to manually refresh the cloth asset, this is done automatically when there is a change in the Dataflow.
 	 * This function is a developper utility used for debugging.
 	 */
+PRAGMA_DISABLE_DEPRECATION_WARNINGS
 	UPROPERTY(EditAnywhere, Category = "Cloth Asset Terminal")
 	mutable FChaosClothAssetTerminalNodeRefreshAsset RefreshAsset;
+PRAGMA_ENABLE_DEPRECATION_WARNINGS
 
 	FChaosClothAssetTerminalNode(const UE::Dataflow::FNodeParameters& InParam, FGuid InGuid = FGuid::NewGuid());
 
