Cloth Editor: Enable ChaosClothAssetEditor plugin for Mac and Linux. This fixes being unable to create cloth assets on these platforms. Also fix Linux build errors.

Mac preflight: 66fc3c6ea879cb8b772ff524
Linux preflight: 66fd794430288a3e457f2c6b

diff --git a/Engine/Plugins/ChaosClothAssetEditor/ChaosClothAssetEditor.uplugin b/Engine/Plugins/ChaosClothAssetEditor/ChaosClothAssetEditor.uplugin
--- a/Engine/Plugins/ChaosClothAssetEditor/ChaosClothAssetEditor.uplugin
+++ b/Engine/Plugins/ChaosClothAssetEditor/ChaosClothAssetEditor.uplugin
@@ -18,25 +18,25 @@
 			"Name": "ChaosClothAssetEditor",
 			"Type": "Editor",
 			"LoadingPhase": "Default",
-			"PlatformAllowList": [ "Win64" ]
+			"PlatformAllowList": [ "Win64", "Mac", "Linux" ]
 		},
 		{
 			"Name": "ChaosClothAssetEditorTools",
 			"Type": "Editor",
 			"LoadingPhase": "Default",
-			"PlatformAllowList": [ "Win64" ]
+			"PlatformAllowList": [ "Win64", "Mac", "Linux" ]
 		},
 		{
 			"Name": "ChaosClothAssetDataflowNodes",
 			"Type": "Editor",
 			"LoadingPhase": "Default",
-			"PlatformAllowList": [ "Win64" ]
+			"PlatformAllowList": [ "Win64", "Mac", "Linux" ]
 		},
 		{
 			"Name": "ChaosClothAssetTools",
 			"Type": "UncookedOnly",
 			"LoadingPhase": "Default",
-			"PlatformAllowList": [ "Win64" ]
+			"PlatformAllowList": [ "Win64", "Mac", "Linux" ]
 		}
 	],
 	"Plugins": [
@@ -70,7 +70,7 @@
 			"Name": "USDImporter",
 			"Enabled": true,
 			"SupportedTargetPlatforms": [
-				"Win64"
+				"Win64", "Mac", "Linux"
 			]
 		}
 	]
diff --git a/Engine/Plugins/ChaosClothAssetEditor/Source/ChaosClothAssetDataflowNodes/Private/ChaosClothAsset/ProxyDeformerNode.cpp b/Engine/Plugins/ChaosClothAssetEditor/Source/ChaosClothAssetDataflowNodes/Private/ChaosClothAsset/ProxyDeformerNode.cpp
--- a/Engine/Plugins/ChaosClothAssetEditor/Source/ChaosClothAssetDataflowNodes/Private/ChaosClothAsset/ProxyDeformerNode.cpp
+++ b/Engine/Plugins/ChaosClothAssetEditor/Source/ChaosClothAssetDataflowNodes/Private/ChaosClothAsset/ProxyDeformerNode.cpp
@@ -515,7 +515,7 @@ void FChaosClothAssetProxyDeformerNode_v2::OnPinRemoved(const UE::Dataflow::FPin
 	check(FirstInput || SecondInput);
 	const bool bIsFirstInput = FirstInput && FirstInput->GetName() == Pin.Name;
 	const bool bIsSecondInput = SecondInput && SecondInput->GetName() == Pin.Name;
-	if (bIsFirstInput && !SecondInput || bIsSecondInput && !FirstInput)
+	if ((bIsFirstInput && !SecondInput) || (bIsSecondInput && !FirstInput))
 	{
 		// Both inputs removed. Remove array index.
 		SelectionFilterSets.SetNum(Index);
diff --git a/Engine/Plugins/ChaosClothAssetEditor/Source/ChaosClothAssetDataflowNodes/Private/ChaosClothAsset/TerminalNode.cpp b/Engine/Plugins/ChaosClothAssetEditor/Source/ChaosClothAssetDataflowNodes/Private/ChaosClothAsset/TerminalNode.cpp
--- a/Engine/Plugins/ChaosClothAssetEditor/Source/ChaosClothAssetDataflowNodes/Private/ChaosClothAsset/TerminalNode.cpp
+++ b/Engine/Plugins/ChaosClothAssetEditor/Source/ChaosClothAssetDataflowNodes/Private/ChaosClothAsset/TerminalNode.cpp
@@ -5,7 +5,7 @@
 #include "ChaosClothAsset/ClothCollectionGroup.h"
 #include "ChaosClothAsset/ClothDataflowTools.h"
 #include "ChaosClothAsset/ClothGeometryTools.h"
-#include "ChaosClothAsset/ClothLODTransitionDataCache.h"
+#include "ChaosClothAsset/ClothLodTransitionDataCache.h"
 #include "ChaosClothAsset/CollectionClothFacade.h"
 #include "ChaosClothAsset/CollectionClothSelectionFacade.h"
 #include "Animation/Skeleton.h"
diff --git a/Engine/Plugins/ChaosClothAssetEditor/Source/ChaosClothAssetEditor/Private/ChaosClothAsset/ClothEditorToolkit.cpp b/Engine/Plugins/ChaosClothAssetEditor/Source/ChaosClothAssetEditor/Private/ChaosClothAsset/ClothEditorToolkit.cpp
--- a/Engine/Plugins/ChaosClothAssetEditor/Source/ChaosClothAssetEditor/Private/ChaosClothAsset/ClothEditorToolkit.cpp
+++ b/Engine/Plugins/ChaosClothAssetEditor/Source/ChaosClothAssetEditor/Private/ChaosClothAsset/ClothEditorToolkit.cpp
@@ -32,7 +32,7 @@
 #include "Dataflow/DataflowGraphEditor.h"
 #include "Dataflow/DataflowObject.h"
 #include "Dataflow/DataflowSchema.h"
-#include "Dataflow/DataFlowEdNode.h"
+#include "Dataflow/DataflowEdNode.h"
 #include "IStructureDetailsView.h"
 #include "Engine/Canvas.h"
 #include "PropertyEditorModule.h"
diff --git a/Engine/Plugins/ChaosClothAssetEditor/Source/ChaosClothAssetTools/Private/ChaosClothAsset/ClothingAssetToClothAssetExporter.cpp b/Engine/Plugins/ChaosClothAssetEditor/Source/ChaosClothAssetTools/Private/ChaosClothAsset/ClothingAssetToClothAssetExporter.cpp
--- a/Engine/Plugins/ChaosClothAssetEditor/Source/ChaosClothAssetTools/Private/ChaosClothAsset/ClothingAssetToClothAssetExporter.cpp
+++ b/Engine/Plugins/ChaosClothAssetEditor/Source/ChaosClothAssetTools/Private/ChaosClothAsset/ClothingAssetToClothAssetExporter.cpp
@@ -11,7 +11,7 @@
 #include "ToDynamicMesh.h"
 #include "Engine/SkinnedAssetCommon.h"
 #include "Engine/SkeletalMesh.h"
-#include "GeometryCollection/ManagedArraycollection.h"
+#include "GeometryCollection/ManagedArrayCollection.h"
 #include "Materials/Material.h"
 #include "Misc/MessageDialog.h"
 #include "ClothingAsset.h"
