CIS - Fixed issue 620209: Warnings in Check Asset References (RefreshAsset).

#rnx

diff --git a/Engine/Config/BaseEngine.ini b/Engine/Config/BaseEngine.ini
--- a/Engine/Config/BaseEngine.ini
+++ b/Engine/Config/BaseEngine.ini
@@ -1306,6 +1306,8 @@ DefaultWorldPartitionHLODLayer=/Engine/Maps/Templates/HLODs/HLODLayer_Instanced.
 +PropertyRedirects=(OldName="NetworkPhysicsSettingsResimulation.bOverrideResimulationErrorThreshold", NewName="bOverrideResimulationErrorPositionThreshold")
 +PropertyRedirects=(OldName="NetworkPhysicsSettingsResimulation.ResimulationErrorThreshold", NewName="ResimulationErrorPositionThreshold")
 
++PropertyRedirects=(OldName="ChaosClothAssetTerminalNode_v2.RefreshAsset", NewName="RefreshAssetOld")
+
 ; Moved UObject-based camera shakes and animations to EngineCameras
 +ClassRedirects=(OldName="/Script/GameplayCameras.CompositeCameraShakePattern", NewName="/Script/EngineCameras.CompositeCameraShakePattern")
 +ClassRedirects=(OldName="/Script/GameplayCameras.DefaultCameraShakeBase", NewName="/Script/EngineCameras.DefaultCameraShakeBase")
diff --git a/Engine/Plugins/ChaosClothAssetEditor/Source/ChaosClothAssetDataflowNodes/Public/ChaosClothAsset/TerminalNode.h b/Engine/Plugins/ChaosClothAssetEditor/Source/ChaosClothAssetDataflowNodes/Public/ChaosClothAsset/TerminalNode.h
--- a/Engine/Plugins/ChaosClothAssetEditor/Source/ChaosClothAssetDataflowNodes/Public/ChaosClothAsset/TerminalNode.h
+++ b/Engine/Plugins/ChaosClothAssetEditor/Source/ChaosClothAssetDataflowNodes/Public/ChaosClothAsset/TerminalNode.h
@@ -42,7 +42,7 @@ public:
 	 * Note that it is not required to manually refresh the cloth asset, this is done automatically when there is a change in the Dataflow.
 	 * This function is a developper utility used for debugging.
 	 */
-	UPROPERTY(EditAnywhere, Category = "Cloth Asset Terminal", Meta = (ButtonImage = "Icons.Refresh"))
+	UPROPERTY(EditAnywhere, Transient, Category = "Cloth Asset Terminal", Meta = (ButtonImage = "Icons.Refresh"))
 	FDataflowFunctionProperty RefreshAsset;
 
 	FChaosClothAssetTerminalNode_v2(const UE::Dataflow::FNodeParameters& InParam, FGuid InGuid = FGuid::NewGuid());
@@ -70,6 +70,11 @@ private:
 	mutable uint32 ClothColllectionChecksum = 0;
 	static constexpr int32 NumRequiredInputs = 0;
 	static constexpr int32 NumInitialCollectionLods = 1;
+
+PRAGMA_DISABLE_DEPRECATION_WARNINGS
+	UPROPERTY()
+	FChaosClothAssetTerminalNodeRefreshAsset RefreshAssetOld_DEPRECATED;
+PRAGMA_ENABLE_DEPRECATION_WARNINGS
 };
 
 /** Cloth terminal node to generate a cloth asset from a cloth collection. */
