Dataflow - Do not call SetAssetValue and Evaluate when a Terminal node, or any node without outputs is de-activated.
Cloth Asset - Revert the change made to the Cloth Asset Terminal in CL 36776697 that created an empty asset when the node was de-activated.

diff --git a/Engine/Plugins/ChaosClothAssetEditor/Source/ChaosClothAssetDataflowNodes/Private/ChaosClothAsset/TerminalNode.cpp b/Engine/Plugins/ChaosClothAssetEditor/Source/ChaosClothAssetDataflowNodes/Private/ChaosClothAsset/TerminalNode.cpp
--- a/Engine/Plugins/ChaosClothAssetEditor/Source/ChaosClothAssetDataflowNodes/Private/ChaosClothAsset/TerminalNode.cpp
+++ b/Engine/Plugins/ChaosClothAssetEditor/Source/ChaosClothAssetDataflowNodes/Private/ChaosClothAsset/TerminalNode.cpp
@@ -105,7 +105,7 @@ void FChaosClothAssetTerminalNode_v2::SetAssetValue(TObjectPtr<UObject> Asset, U
 	{
 		using namespace UE::Chaos::ClothAsset;
 
-		TArray<TSharedRef<FManagedArrayCollection>> InClothCollections = bActive ? GetCleanedCollectionLodValues(Context) : TArray<TSharedRef<FManagedArrayCollection>>();
+		TArray<TSharedRef<FManagedArrayCollection>> InClothCollections = GetCleanedCollectionLodValues(Context);
 		TArray<TSharedRef<FManagedArrayCollection>>& ClothCollections = ClothAsset->GetClothCollections();
 
 		const uint32 PreviousChecksum = ClothColllectionChecksum;
@@ -389,7 +389,7 @@ PRAGMA_DISABLE_DEPRECATION_WARNINGS
 		}
 PRAGMA_ENABLE_DEPRECATION_WARNINGS
 
-		TArray<TSharedRef<FManagedArrayCollection>> InClothCollections = bActive ? GetCleanedCollectionLodValues(Context) : TArray<TSharedRef<FManagedArrayCollection>>();
+		TArray<TSharedRef<FManagedArrayCollection>> InClothCollections = GetCleanedCollectionLodValues(Context);
 		TArray<TSharedRef<FManagedArrayCollection>>& ClothCollections = ClothAsset->GetClothCollections();
 
 		const uint32 PreviousChecksum = ClothColllectionChecksum;
diff --git a/Engine/Plugins/Experimental/Dataflow/Source/DataflowEditor/Private/Dataflow/DataflowEditorCommands.cpp b/Engine/Plugins/Experimental/Dataflow/Source/DataflowEditor/Private/Dataflow/DataflowEditorCommands.cpp
--- a/Engine/Plugins/Experimental/Dataflow/Source/DataflowEditor/Private/Dataflow/DataflowEditorCommands.cpp
+++ b/Engine/Plugins/Experimental/Dataflow/Source/DataflowEditor/Private/Dataflow/DataflowEditorCommands.cpp
@@ -159,11 +159,14 @@ const FDataflowNode* FDataflowEditorCommands::EvaluateNode(UE::Dataflow::FContex
 	}
 	if (Node && InOutLastNodeTimestamp < Node->GetTimestamp())
 	{
+		// Note: If the node is deactivated and has any outputs, then these outputs might still need to be forwarded.
+		//       Therefore the Evaluate method has to be called for whichever value of bActive.
+		//       This however isn't the case of SetAssetValue() for which the active state needs to be checked before the call.
 		Context.Evaluate(Node, Output);
 
-		if (Asset)
+		if (const FDataflowTerminalNode* const TerminalNode = Node->AsType<const FDataflowTerminalNode>())
 		{
-			if (const FDataflowTerminalNode* const TerminalNode = Node->AsType<const FDataflowTerminalNode>())
+			if (TerminalNode->bActive && Asset)
 			{
 				UE_LOG(LogChaosDataflow, Verbose, TEXT("FDataflowTerminalNode::SetAssetValue(): TerminalNode [%s], Asset [%s]"), *TerminalNode->GetName().ToString(), *Asset->GetName());
 				TerminalNode->SetAssetValue(Asset, Context);
diff --git a/Engine/Source/Runtime/Experimental/Dataflow/Core/Private/Dataflow/DataflowNodeParameters.cpp b/Engine/Source/Runtime/Experimental/Dataflow/Core/Private/Dataflow/DataflowNodeParameters.cpp
--- a/Engine/Source/Runtime/Experimental/Dataflow/Core/Private/Dataflow/DataflowNodeParameters.cpp
+++ b/Engine/Source/Runtime/Experimental/Dataflow/Core/Private/Dataflow/DataflowNodeParameters.cpp
@@ -94,7 +94,10 @@ namespace UE::Dataflow
 					Context.Evaluate(*NodeOutput);
 				}
 			}
-			else
+			// Note: If the node is deactivated and has an output (like above), then the output might still need to be forwarded.
+			//       Therefore the Evaluate method has to be called for whichever value of bActive.
+			//       However if the node is deactivated and has no outputs (like below), now is the time to check its bActive state.
+			else if (Node->bActive)
 			{
 				// TODO: When no outputs are specified, this call to Evaluate should really be removed.
 				//       The purpose of the node evaluation function is to evaluate outputs.
diff --git a/Engine/Source/Runtime/Experimental/Dataflow/Engine/Private/Dataflow/DataflowBlueprintLibrary.cpp b/Engine/Source/Runtime/Experimental/Dataflow/Engine/Private/Dataflow/DataflowBlueprintLibrary.cpp
--- a/Engine/Source/Runtime/Experimental/Dataflow/Engine/Private/Dataflow/DataflowBlueprintLibrary.cpp
+++ b/Engine/Source/Runtime/Experimental/Dataflow/Engine/Private/Dataflow/DataflowBlueprintLibrary.cpp
@@ -16,8 +16,11 @@ void UDataflowBlueprintLibrary::EvaluateTerminalNodeByName(UDataflow* Dataflow,
 			{
 				UE_LOG(LogChaosDataflow, Verbose, TEXT("UDataflowBlueprintLibrary::EvaluateTerminalNodeByName(): Node [%s]"), *TerminalNodeName.ToString());
 				UE::Dataflow::FEngineContext Context(ResultAsset);
+				// Note: If the node is deactivated and has any outputs, then these outputs might still need to be forwarded.
+				//       Therefore the Evaluate method has to be called for whichever value of bActive.
+				//       This however isn't the case of SetAssetValue() for which the active state needs to be checked before the call.
 				TerminalNode->Evaluate(Context);
-				if (ResultAsset)
+				if (TerminalNode->bActive && ResultAsset)
 				{
 					UE_LOG(LogChaosDataflow, Verbose, TEXT("FDataflowTerminalNode::SetAssetValue(): TerminalNode [%s], Asset [%s]"), *TerminalNodeName.ToString(), *ResultAsset->GetName());
 					TerminalNode->SetAssetValue(ResultAsset, Context);
