Chaos Cloth Asset - Fixed the endless loop when clearing the Dataflow reference from a Cloth Asset.

diff --git a/Engine/Plugins/ChaosClothAssetEditor/Source/ChaosClothAssetEditor/Private/ChaosClothAsset/AssetDefinition_ClothAsset.cpp b/Engine/Plugins/ChaosClothAssetEditor/Source/ChaosClothAssetEditor/Private/ChaosClothAsset/AssetDefinition_ClothAsset.cpp
--- a/Engine/Plugins/ChaosClothAssetEditor/Source/ChaosClothAssetEditor/Private/ChaosClothAsset/AssetDefinition_ClothAsset.cpp
+++ b/Engine/Plugins/ChaosClothAssetEditor/Source/ChaosClothAssetEditor/Private/ChaosClothAsset/AssetDefinition_ClothAsset.cpp
@@ -181,16 +181,6 @@ EAssetCommandResult UAssetDefinition_ClothAsset::OpenAssets(const FAssetOpenArgs
 
 	if (ClothObjects.Num() > 0)
 	{
-		// Validate the asset
-		UChaosClothAsset* const ClothAsset = CastChecked<UChaosClothAsset>(ClothObjects[0]);
-		if (!ClothAsset->GetDataflow())
-		{
-			if (UDataflow* const NewDataflowAsset = Cast<UDataflow>(NewOrOpenDataflowAsset(ClothAsset)))
-			{
-				ClothAsset->SetDataflow(NewDataflowAsset);
-			}
-		}
-		
 		UAssetEditorSubsystem* const AssetEditorSubsystem = GEditor->GetEditorSubsystem<UAssetEditorSubsystem>();
 		if(!UE::Chaos::ClothAsset::Private::bEnableClothDataflowEditor)
 		{
diff --git a/Engine/Plugins/ChaosClothAssetEditor/Source/ChaosClothAssetEditor/Private/ChaosClothAsset/AssetDefinition_ClothAsset.h b/Engine/Plugins/ChaosClothAssetEditor/Source/ChaosClothAssetEditor/Private/ChaosClothAsset/AssetDefinition_ClothAsset.h
--- a/Engine/Plugins/ChaosClothAssetEditor/Source/ChaosClothAssetEditor/Private/ChaosClothAsset/AssetDefinition_ClothAsset.h
+++ b/Engine/Plugins/ChaosClothAssetEditor/Source/ChaosClothAssetEditor/Private/ChaosClothAsset/AssetDefinition_ClothAsset.h
@@ -14,6 +14,7 @@ class UAssetDefinition_ClothAsset : public UAssetDefinitionDefault
 
 public:
 
+	UE_DEPRECATED(5.5, "Dataflow assets are now created through the physics menu.")
 	static UObject* NewOrOpenDataflowAsset(const UChaosClothAsset* ClothAsset);
 
 private:
diff --git a/Engine/Plugins/ChaosClothAssetEditor/Source/ChaosClothAssetEditor/Private/ChaosClothAsset/ChaosClothAssetEditorModule.cpp b/Engine/Plugins/ChaosClothAssetEditor/Source/ChaosClothAssetEditor/Private/ChaosClothAsset/ChaosClothAssetEditorModule.cpp
--- a/Engine/Plugins/ChaosClothAssetEditor/Source/ChaosClothAssetEditor/Private/ChaosClothAsset/ChaosClothAssetEditorModule.cpp
+++ b/Engine/Plugins/ChaosClothAssetEditor/Source/ChaosClothAssetEditor/Private/ChaosClothAsset/ChaosClothAssetEditorModule.cpp
@@ -63,10 +63,12 @@ namespace UE::Chaos::ClothAsset
 							UChaosClothAsset* const ClothAsset = CastChecked<UChaosClothAsset>(AssetsToEdit[0]);
 							if (!ClothAsset->GetDataflow())
 							{
+PRAGMA_DISABLE_DEPRECATION_WARNINGS
 								if (UDataflow* const NewDataflowAsset = Cast<UDataflow>(UAssetDefinition_ClothAsset::NewOrOpenDataflowAsset(ClothAsset)))
 								{
 									ClothAsset->SetDataflow(NewDataflowAsset);
 								}
+PRAGMA_ENABLE_DEPRECATION_WARNINGS
 							}
 
 							const TSubclassOf<AActor> PreviewActorClass = StaticLoadClass(AActor::StaticClass(), nullptr,
diff --git a/Engine/Plugins/ChaosClothAssetEditor/Source/ChaosClothAssetTools/ChaosClothAssetTools.Build.cs b/Engine/Plugins/ChaosClothAssetEditor/Source/ChaosClothAssetTools/ChaosClothAssetTools.Build.cs
--- a/Engine/Plugins/ChaosClothAssetEditor/Source/ChaosClothAssetTools/ChaosClothAssetTools.Build.cs
+++ b/Engine/Plugins/ChaosClothAssetEditor/Source/ChaosClothAssetTools/ChaosClothAssetTools.Build.cs
@@ -24,6 +24,7 @@ public class ChaosClothAssetTools : ModuleRules
 				"Chaos",
 				"ChaosClothAsset",
 				"ChaosClothAssetEngine",
+				"DataflowEngine",
 				"SlateCore",
 				"UnrealEd",
 				"Projects",
diff --git a/Engine/Plugins/ChaosClothAssetEditor/Source/ChaosClothAssetTools/Private/ChaosClothAsset/ClothAssetFactory.cpp b/Engine/Plugins/ChaosClothAssetEditor/Source/ChaosClothAssetTools/Private/ChaosClothAsset/ClothAssetFactory.cpp
--- a/Engine/Plugins/ChaosClothAssetEditor/Source/ChaosClothAssetTools/Private/ChaosClothAsset/ClothAssetFactory.cpp
+++ b/Engine/Plugins/ChaosClothAssetEditor/Source/ChaosClothAssetTools/Private/ChaosClothAsset/ClothAssetFactory.cpp
@@ -1,10 +1,12 @@
 // Copyright Epic Games, Inc. All Rights Reserved.
 
 #include "ClothAssetFactory.h"
+#include "AssetRegistry/AssetRegistryModule.h"
 #include "ChaosClothAsset/ClothAsset.h"
 #include "ChaosClothAsset/CollectionClothFacade.h"
+#include "Dataflow/DataflowObject.h"
 #include "GeometryCollection/ManagedArrayCollection.h"
-#include "Animation/Skeleton.h"
+#include "Misc/PackageName.h"
 
 #include UE_INLINE_GENERATED_CPP_BY_NAME(ClothAssetFactory)
 
@@ -23,5 +25,30 @@ UObject* UChaosClothAssetFactory::FactoryCreateNew(UClass* Class, UObject* Paren
 	UChaosClothAsset* const ClothAsset = NewObject<UChaosClothAsset>(Parent, Class, Name, Flags | RF_Transactional | RF_Public | RF_Standalone);
 	ClothAsset->MarkPackageDirty();
 
+	// Create a new Dataflow asset
+	const FString DataflowPath = FPackageName::GetLongPackagePath(ClothAsset->GetOutermost()->GetName());
+	const FString ClothAssetName = ClothAsset->GetName();
+	const FString DataflowName = FString(TEXT("DF_")) + (ClothAssetName.StartsWith(TEXT("CA_")) ? ClothAssetName.RightChop(3) : ClothAssetName);
+	const FString DataflowPackageName = FPaths::Combine(DataflowPath, DataflowName);
+
+	UPackage* const DataflowPackage = CreatePackage(*DataflowPackageName);
+
+	// Load the cloth template into the new Dataflow asset
+	UDataflow* const Template = LoadObject<UDataflow>(DataflowPackage, TEXT("/ChaosClothAssetEditor/ClothAssetTemplate.ClothAssetTemplate"));
+	UDataflow* const Dataflow = DuplicateObject(Template, DataflowPackage, FName(DataflowName));
+
+	Dataflow->MarkPackageDirty();
+
+	// Notify the asset registry
+	FAssetRegistryModule::AssetCreated(Dataflow);
+
+	// Set the Dataflow to the cloth asset
+	ClothAsset->SetDataflow(Dataflow);
+
 	return ClothAsset;
 }
+
+FString UChaosClothAssetFactory::GetDefaultNewAssetName() const
+{
+	return FString(TEXT("CA_NewChaosClothAsset"));
+}
diff --git a/Engine/Plugins/ChaosClothAssetEditor/Source/ChaosClothAssetTools/Private/ChaosClothAsset/ClothAssetFactory.h b/Engine/Plugins/ChaosClothAssetEditor/Source/ChaosClothAssetTools/Private/ChaosClothAsset/ClothAssetFactory.h
--- a/Engine/Plugins/ChaosClothAssetEditor/Source/ChaosClothAssetTools/Private/ChaosClothAsset/ClothAssetFactory.h
+++ b/Engine/Plugins/ChaosClothAssetEditor/Source/ChaosClothAssetTools/Private/ChaosClothAsset/ClothAssetFactory.h
@@ -23,5 +23,6 @@ public:
 	virtual bool FactoryCanImport(const FString& Filename) override { return false; }
 	virtual bool ShouldShowInNewMenu() const override { return true; }
 	virtual UObject* FactoryCreateNew(UClass* InClass, UObject* InParent, FName InName, EObjectFlags Flags, UObject* Context, FFeedbackContext* Warn) override;
+	virtual FString GetDefaultNewAssetName() const override;
 	/** End UFactory Interface */
 };
