Chaos Cloth Asset: set NumTexCoords based on the imported value when importing meshes.

Tested this fix on an asset which requires more than one TexCoord. It draws correctly now in the construction and preview viewports (cloth and dataflow editors) as well as in PIE.

diff --git a/Engine/Plugins/ChaosClothAsset/Source/ChaosClothAssetEngine/Private/ChaosClothAsset/ClothAssetBuilder.cpp b/Engine/Plugins/ChaosClothAsset/Source/ChaosClothAssetEngine/Private/ChaosClothAsset/ClothAssetBuilder.cpp
--- a/Engine/Plugins/ChaosClothAsset/Source/ChaosClothAssetEngine/Private/ChaosClothAsset/ClothAssetBuilder.cpp
+++ b/Engine/Plugins/ChaosClothAsset/Source/ChaosClothAssetEngine/Private/ChaosClothAsset/ClothAssetBuilder.cpp
@@ -269,10 +269,13 @@ void UChaosClothAsset::FBuilder::BuildLod(FSkeletalMeshLODModel& LODModel, const
 			SoftVertex.Color = PatternRenderColor[VertexIndex].ToFColor(bSRGB);
 
 			const TArray<FVector2f>& RenderUVs = PatternRenderUVs[VertexIndex];
-			for (int32 TexCoord = 0; TexCoord < FMath::Min(RenderUVs.Num(), (int32)MAX_TEXCOORDS); ++TexCoord)
+			FMemory::Memzero(SoftVertex.UVs, sizeof(SoftVertex.UVs));
+			const int32 NumTexCoords = FMath::Min(RenderUVs.Num(), (int32)MAX_TEXCOORDS);
+			for (int32 TexCoord = 0; TexCoord < NumTexCoords; ++TexCoord)
 			{
 				SoftVertex.UVs[TexCoord] = RenderUVs[TexCoord];
 			}
+			LODModel.NumTexCoords = FMath::Max(LODModel.NumTexCoords, (uint32)NumTexCoords);
 
 			// Conform to the platform max number of influences
 			const int32 MaxBoneInfluencesFromUnlimitedBoneInfluences = FGPUBaseSkinVertexFactory::GetUnlimitedBoneInfluences(TargetPlatform) ? MAX_TOTAL_INFLUENCES : EXTRA_BONE_INFLUENCES;
diff --git a/Engine/Plugins/ChaosClothAssetEditor/Source/ChaosClothAssetDataflowNodes/Private/ChaosClothAsset/ClothDataflowTools.cpp b/Engine/Plugins/ChaosClothAssetEditor/Source/ChaosClothAssetDataflowNodes/Private/ChaosClothAsset/ClothDataflowTools.cpp
--- a/Engine/Plugins/ChaosClothAssetEditor/Source/ChaosClothAssetDataflowNodes/Private/ChaosClothAsset/ClothDataflowTools.cpp
+++ b/Engine/Plugins/ChaosClothAssetEditor/Source/ChaosClothAssetDataflowNodes/Private/ChaosClothAsset/ClothDataflowTools.cpp
@@ -394,6 +394,7 @@ namespace UE::Chaos::ClothAsset
 		TArrayView<FLinearColor> RenderColor = ClothPatternFacade.GetRenderColor();
 		TArrayView<TArray<int32>> RenderBoneIndices = ClothPatternFacade.GetRenderBoneIndices();
 		TArrayView<TArray<float>> RenderBoneWeights = ClothPatternFacade.GetRenderBoneWeights();
+		const uint32 NumTexCoords = FMath::Min((uint32)MAX_TEXCOORDS, SkeletalMeshModel.NumTexCoords);
 		for (int32 VertexIndex = 0; VertexIndex < Section.NumVertices; ++VertexIndex)
 		{
 			const FSoftSkinVertex& SoftVertex = Section.SoftVertices[VertexIndex];
@@ -402,8 +403,8 @@ namespace UE::Chaos::ClothAsset
 			RenderNormal[VertexIndex] = SoftVertex.TangentZ;
 			RenderTangentU[VertexIndex] = SoftVertex.TangentX;
 			RenderTangentV[VertexIndex] = SoftVertex.TangentY;
-			RenderUVs[VertexIndex].SetNum(MAX_TEXCOORDS);
-			for (int32 TexCoordIndex = 0; TexCoordIndex < MAX_TEXCOORDS; ++TexCoordIndex)
+			RenderUVs[VertexIndex].SetNum(NumTexCoords);
+			for (uint32 TexCoordIndex = 0; TexCoordIndex < NumTexCoords; ++TexCoordIndex)
 			{
 				RenderUVs[VertexIndex][TexCoordIndex] = SoftVertex.UVs[TexCoordIndex];
 			}
@@ -520,6 +521,7 @@ namespace UE::Chaos::ClothAsset
 		IMeshUtilities::MeshBuildOptions BuildOptions;
 		BuildOptions.TargetPlatform = GetTargetPlatformManagerRef().GetRunningTargetPlatform();
 		BuildOptions.FillOptions(BuildSettings);
+		SkeletalMeshModel.NumTexCoords = SkeletalMeshImportGeometry.NumTexCoords;
 
 		static const FString SkeletalMeshName("ClothAssetStaticMeshImportConvert"); // This is only used by warning messages in the mesh builder.
 		// Build a RefSkeleton with just a root bone. The BuildSkeletalMesh code expects you have a reference skeleton with at least one bone to work.
diff --git a/Engine/Plugins/ChaosClothAssetEditor/Source/ChaosClothAssetTools/Private/ChaosClothAsset/ClothPatternToDynamicMesh.cpp b/Engine/Plugins/ChaosClothAssetEditor/Source/ChaosClothAssetTools/Private/ChaosClothAsset/ClothPatternToDynamicMesh.cpp
--- a/Engine/Plugins/ChaosClothAssetEditor/Source/ChaosClothAssetTools/Private/ChaosClothAsset/ClothPatternToDynamicMesh.cpp
+++ b/Engine/Plugins/ChaosClothAssetEditor/Source/ChaosClothAssetTools/Private/ChaosClothAsset/ClothPatternToDynamicMesh.cpp
@@ -35,6 +35,7 @@ public:
 		VertexDataType(VertexDataType),
 		Cloth(ClothFacade)
 	{
+		NumTexCoords = 0;
 		if (PatternIndex == INDEX_NONE)
 		{
 			// All patterns in one dynamic mesh
@@ -43,10 +44,12 @@ public:
 			case EClothPatternVertexType::Render:
 			{
 				const int32 NumVertices = Cloth.GetNumRenderVertices();
-				VertIDs.SetNum(NumVertices);
+				VertIDs.SetNum(NumVertices); 
+				TConstArrayView<TArray<FVector2f>> UVs = Cloth.GetRenderUVs();
 				for (int32 VtxIndex = 0; VtxIndex < NumVertices; ++VtxIndex)
 				{
-					VertIDs[VtxIndex] = VtxIndex;
+					VertIDs[VtxIndex] = VtxIndex; 
+					NumTexCoords = FMath::Max(NumTexCoords, UVs[VtxIndex].Num());
 				}
 				NormalIDs = VertIDs;
 
@@ -132,9 +135,11 @@ public:
 				const int32 NumVertices = Pattern.GetNumRenderVertices();
 				const int32 VertexOffset = Pattern.GetRenderVerticesOffset();
 				VertIDs.SetNum(NumVertices);
+				TConstArrayView<TArray<FVector2f>> UVs = Pattern.GetRenderUVs();
 				for (int32 VtxIndex = 0; VtxIndex < NumVertices; ++VtxIndex)
 				{
 					VertIDs[VtxIndex] = VtxIndex + VertexOffset;
+					NumTexCoords = FMath::Max(NumTexCoords, UVs[VtxIndex].Num());
 				}
 				NormalIDs = VertIDs;
 
@@ -243,7 +248,7 @@ public:
 
 	int32 NumUVLayers() const
 	{
-		return (VertexDataType == EClothPatternVertexType::Render) ? 1 : 0;		// No UVs for Sim mesh
+		return NumTexCoords;
 	}
 
 	int32 NumWeightMapLayers() const
@@ -539,7 +544,7 @@ public:
 	{ 
 		checkf(VertexDataType == EClothPatternVertexType::Render, TEXT("Requested UVs from a Sim mesh"));
 		const TConstArrayView<FVector2f> VertexUVs = Cloth.GetRenderUVs()[UVID];
-		return VertexUVs[LayerID];
+		return VertexUVs.IsValidIndex(LayerID) ? VertexUVs[LayerID] : FVector2f(0.f);
 	}
 
 	bool GetUVTri(int32 LayerID, const TriIDType& TriID, UVIDType& ID0, UVIDType& ID1, UVIDType& ID2) const 
@@ -603,6 +608,7 @@ private:
 	TArray<TriIDType> TriIDs;		// indices into Indices
 	TArray<VertIDType> VertIDs;		// indices into Positions2D or Positions3D
 	TArray<NormalIDType> NormalIDs; // indices into Normals (and all per vertex data other than positions)
+	int32 NumTexCoords;
 	
 	TArray<FName> WeightMapNames;
 
