Chaos Cloth Asset: Add Blend Vertices node.

diff --git a/Engine/Plugins/ChaosClothAssetEditor/Source/ChaosClothAssetDataflowNodes/Private/ChaosClothAsset/BlendVerticesNode.cpp b/Engine/Plugins/ChaosClothAssetEditor/Source/ChaosClothAssetDataflowNodes/Private/ChaosClothAsset/BlendVerticesNode.cpp
new file mode 100644
--- /dev/null
+++ b/Engine/Plugins/ChaosClothAssetEditor/Source/ChaosClothAssetDataflowNodes/Private/ChaosClothAsset/BlendVerticesNode.cpp
@@ -0,0 +1,210 @@
+// Copyright Epic Games, Inc. All Rights Reserved.
+
+#include "ChaosClothAsset/BlendVerticesNode.h"
+#include "ChaosClothAsset/ClothDataflowTools.h"
+#include "ChaosClothAsset/CollectionClothFacade.h"
+#include "Dataflow/DataflowInputOutput.h"
+
+#include UE_INLINE_GENERATED_CPP_BY_NAME(BlendVerticesNode)
+
+#define LOCTEXT_NAMESPACE "ChaosClothAssetBlendVerticesNode"
+
+namespace UE::Chaos::ClothAsset::Private
+{
+	template<typename T>
+	void BlendValues(float BlendingWeight, const TArrayView<T>& Values, const TConstArrayView<T>& BlendingValues)
+	{
+		const int32 NumValues = FMath::Min(Values.Num(), BlendingValues.Num());
+		const float OneMinusScalar = 1.f - BlendingWeight;
+		for (int32 Index = 0; Index < NumValues; ++Index)
+		{
+			Values[Index] = OneMinusScalar * Values[Index] + BlendingWeight * BlendingValues[Index];
+		}
+	}
+
+	template<typename T>
+	void BlendNormalizedValues(float BlendingWeight, const TArrayView<T>& Values, const TConstArrayView<T>& BlendingValues)
+	{
+		const int32 NumValues = FMath::Min(Values.Num(), BlendingValues.Num());
+		const float OneMinusScalar = 1.f - BlendingWeight;
+		for (int32 Index = 0; Index < NumValues; ++Index)
+		{
+			Values[Index] = (OneMinusScalar * Values[Index] + BlendingWeight * BlendingValues[Index]).GetSafeNormal();
+		}
+	}
+
+	static void BlendUVSets(float BlendingWeight, const TArrayView<TArray<FVector2f>>& Values, const TConstArrayView<TArray<FVector2f>>& BlendingValues)
+	{
+		const int32 NumValues = FMath::Min(Values.Num(), BlendingValues.Num());
+		for (int32 Index = 0; Index < NumValues; ++Index)
+		{
+			BlendValues(BlendingWeight, TArrayView<FVector2f>(Values[Index]), TConstArrayView<FVector2f>(BlendingValues[Index]));
+		}
+	}
+}
+
+FChaosClothAssetBlendVerticesNode::FChaosClothAssetBlendVerticesNode(const UE::Dataflow::FNodeParameters& InParam, FGuid InGuid)
+	: FDataflowNode(InParam, InGuid)
+{
+	RegisterInputConnection(&Collection);
+	RegisterInputConnection(&BlendCollection);
+	RegisterOutputConnection(&Collection, &Collection);
+}
+
+void FChaosClothAssetBlendVerticesNode::Evaluate(UE::Dataflow::FContext& Context, const FDataflowOutput* Out) const
+{
+	if (Out->IsA<FManagedArrayCollection>(&Collection))
+	{
+		using namespace UE::Chaos::ClothAsset;
+
+		if (BlendingWeight == 0.f)
+		{
+			SafeForwardInput(Context, &Collection, &Collection);
+			return;
+		}
+
+		// Evaluate in collection
+		FManagedArrayCollection InCollection = GetValue<FManagedArrayCollection>(Context, &Collection);
+		const TSharedRef<FManagedArrayCollection> ClothCollection = MakeShared<FManagedArrayCollection>(MoveTemp(InCollection));
+
+		FManagedArrayCollection InBlendCollection = GetValue<FManagedArrayCollection>(Context, &BlendCollection);  // Can't use a const reference here sadly since the facade needs a SharedRef to be created
+		const TSharedRef<const FManagedArrayCollection> BlendClothCollection = MakeShared<const FManagedArrayCollection>(MoveTemp(InBlendCollection));
+
+		// Always check for a valid cloth collection/facade to avoid processing non cloth collections
+		FCollectionClothFacade ClothFacade(ClothCollection);
+		const FCollectionClothConstFacade BlendClothFacade(BlendClothCollection);
+
+		auto LogVertexCountMismatch = [this](const FText& VertexType, int32 CollectionCount, int32 BlendCollectionCount)
+			{
+				UE::Chaos::ClothAsset::FClothDataflowTools::LogAndToastWarning(*this,
+					FText::Format(LOCTEXT("VertexCountMismatchHeadline", "Failed to blend {0} attributes."), VertexType),
+					FText::Format(LOCTEXT("VertexCountMismatchDetails", "Vertex count mismatch {0} != {1}. Set 'Require Same Vertex Counts' to false to disable this check."),
+						CollectionCount, BlendCollectionCount));
+			};
+
+		if (ClothFacade.IsValid())
+		{
+			if (BlendClothFacade.IsValid())
+			{
+				if (bBlendSimMesh)
+				{
+					if (bBlend2DSimPositions)
+					{
+						const int32 NumSimVertices2D = ClothFacade.GetNumSimVertices2D();
+						const int32 NumBlendSimVertices2D = BlendClothFacade.GetNumSimVertices2D();
+						if (!bRequireSameVertexCounts || NumSimVertices2D == NumBlendSimVertices2D)
+						{
+							TArrayView<FVector2f> SimPosition2D = ClothFacade.GetSimPosition2D();
+							TConstArrayView<FVector2f> BlendSimPosition2D = BlendClothFacade.GetSimPosition2D();
+							ensure(SimPosition2D.Num() == NumSimVertices2D);
+							ensure(BlendSimPosition2D.Num() == NumBlendSimVertices2D);
+							Private::BlendValues(BlendingWeight, SimPosition2D, BlendSimPosition2D);
+						}
+						else
+						{
+							LogVertexCountMismatch(FText::FromString(TEXT("Sim Vertices 2D")), NumSimVertices2D, NumBlendSimVertices2D);
+						}
+					}
+
+					if (bBlend3DSimPositions || bBlendSimNormals)
+					{
+						const int32 NumSimVertices3D = ClothFacade.GetNumSimVertices3D();
+						const int32 NumBlendSimVertices3D = BlendClothFacade.GetNumSimVertices3D();
+						if (!bRequireSameVertexCounts || NumSimVertices3D == NumBlendSimVertices3D)
+						{
+							if (bBlend3DSimPositions)
+							{
+								TArrayView<FVector3f> SimPosition3D = ClothFacade.GetSimPosition3D();
+								TConstArrayView<FVector3f> BlendSimPosition3D = BlendClothFacade.GetSimPosition3D();
+								ensure(SimPosition3D.Num() == NumSimVertices3D);
+								ensure(BlendSimPosition3D.Num() == NumBlendSimVertices3D);
+								Private::BlendValues(BlendingWeight, SimPosition3D, BlendSimPosition3D);
+							}
+							if (bBlendSimNormals)
+							{
+								TArrayView<FVector3f> SimNormal = ClothFacade.GetSimNormal();
+								TConstArrayView<FVector3f> BlendSimNormal = BlendClothFacade.GetSimNormal();
+								ensure(SimNormal.Num() == NumSimVertices3D);
+								ensure(BlendSimNormal.Num() == NumBlendSimVertices3D);
+								Private::BlendNormalizedValues(BlendingWeight, SimNormal, BlendSimNormal);
+							}
+						}
+						else
+						{
+							LogVertexCountMismatch(FText::FromString(TEXT("Sim Vertices 3D")), NumSimVertices3D, NumBlendSimVertices3D);
+						}
+					}
+				}
+
+				if (bBlendRenderMesh)
+				{
+					const int32 NumRenderVertices = ClothFacade.GetNumRenderVertices();
+					const int32 NumBlendRenderVertices = BlendClothFacade.GetNumRenderVertices();
+					if (!bRequireSameVertexCounts || NumRenderVertices == NumBlendRenderVertices)
+					{
+						if (bBlendRenderPositions)
+						{
+							TArrayView<FVector3f> RenderPosition = ClothFacade.GetRenderPosition();
+							TConstArrayView<FVector3f> BlendRenderPosition = BlendClothFacade.GetRenderPosition();
+							ensure(RenderPosition.Num() == NumRenderVertices);
+							ensure(BlendRenderPosition.Num() == NumBlendRenderVertices);
+							Private::BlendValues(BlendingWeight, RenderPosition, BlendRenderPosition);
+						}
+
+						if (bBlendRenderNormalsAndTangents)
+						{
+							TArrayView<FVector3f> RenderNormal = ClothFacade.GetRenderNormal();
+							TConstArrayView<FVector3f> BlendRenderNormal = BlendClothFacade.GetRenderNormal();
+							ensure(RenderNormal.Num() == NumRenderVertices);
+							ensure(BlendRenderNormal.Num() == NumBlendRenderVertices);
+							Private::BlendNormalizedValues(BlendingWeight, RenderNormal, BlendRenderNormal);
+							TArrayView<FVector3f> RenderTangentU = ClothFacade.GetRenderTangentU();
+							TConstArrayView<FVector3f> BlendRenderTangentU = BlendClothFacade.GetRenderTangentU();
+							ensure(RenderTangentU.Num() == NumRenderVertices);
+							ensure(BlendRenderTangentU.Num() == NumBlendRenderVertices);
+							Private::BlendNormalizedValues(BlendingWeight, RenderTangentU, BlendRenderTangentU);
+							TArrayView<FVector3f> RenderTangentV = ClothFacade.GetRenderTangentV();
+							TConstArrayView<FVector3f> BlendRenderTangentV = BlendClothFacade.GetRenderTangentV();
+							ensure(RenderTangentV.Num() == NumRenderVertices);
+							ensure(BlendRenderTangentV.Num() == NumBlendRenderVertices);
+							Private::BlendNormalizedValues(BlendingWeight, RenderTangentV, BlendRenderTangentV);
+						}
+
+						if (bBlendRenderUVs)
+						{
+							TArrayView<TArray<FVector2f>> RenderUVs = ClothFacade.GetRenderUVs();
+							TConstArrayView<TArray<FVector2f>> BlendRenderUVs = BlendClothFacade.GetRenderUVs();
+							ensure(RenderUVs.Num() == NumRenderVertices);
+							ensure(BlendRenderUVs.Num() == NumBlendRenderVertices);
+							Private::BlendUVSets(BlendingWeight, RenderUVs, BlendRenderUVs);
+						}
+
+						if (bBlendRenderColors)
+						{
+							TArrayView<FLinearColor> RenderColor = ClothFacade.GetRenderColor();
+							TConstArrayView<FLinearColor> BlendRenderColor = BlendClothFacade.GetRenderColor();
+							ensure(RenderColor.Num() == NumRenderVertices);
+							ensure(BlendRenderColor.Num() == NumBlendRenderVertices);
+							Private::BlendValues(BlendingWeight, RenderColor, BlendRenderColor);
+						}
+					}
+					else
+					{
+						LogVertexCountMismatch(FText::FromString(TEXT("Render Vertices")), NumRenderVertices, NumBlendRenderVertices);
+					}
+				}
+			}
+			else
+			{
+				UE::Chaos::ClothAsset::FClothDataflowTools::LogAndToastWarning(*this,
+					LOCTEXT("InvalidBlendCollectionHeadline", "Invalid Blend Collection"),
+					LOCTEXT("InvalidBlendCollectionDetails", "Input Blend Collection is not a valid Cloth Collection.")
+				);
+			}
+		}
+
+		SetValue(Context, MoveTemp(*ClothCollection), &Collection);
+	}
+}
+
+#undef LOCTEXT_NAMESPACE
diff --git a/Engine/Plugins/ChaosClothAssetEditor/Source/ChaosClothAssetDataflowNodes/Private/ChaosClothAsset/ChaosClothAssetDataflowNodesModule.cpp b/Engine/Plugins/ChaosClothAssetEditor/Source/ChaosClothAssetDataflowNodes/Private/ChaosClothAsset/ChaosClothAssetDataflowNodesModule.cpp
--- a/Engine/Plugins/ChaosClothAssetEditor/Source/ChaosClothAssetDataflowNodes/Private/ChaosClothAsset/ChaosClothAssetDataflowNodesModule.cpp
+++ b/Engine/Plugins/ChaosClothAssetEditor/Source/ChaosClothAssetDataflowNodes/Private/ChaosClothAsset/ChaosClothAssetDataflowNodesModule.cpp
@@ -4,6 +4,7 @@
 #include "ChaosClothAsset/AddWeightMapNode.h"
 #include "ChaosClothAsset/AttributeNode.h"
 #include "ChaosClothAsset/BindToRootBoneNode.h"
+#include "ChaosClothAsset/BlendVerticesNode.h"
 #include "ChaosClothAsset/ClothCollectionGroup.h"
 #include "ChaosClothAsset/ColorScheme.h"
 #include "ChaosClothAsset/CopySimulationToRenderMeshNode.h"
@@ -85,6 +86,7 @@ namespace UE::Chaos::ClothAsset
 			DATAFLOW_NODE_REGISTER_CREATION_FACTORY(FChaosClothAssetWeightMapNode);
 			DATAFLOW_NODE_REGISTER_CREATION_FACTORY(FChaosClothAssetAttributeNode_v2);
 			DATAFLOW_NODE_REGISTER_CREATION_FACTORY(FChaosClothAssetBindToRootBoneNode);
+			DATAFLOW_NODE_REGISTER_CREATION_FACTORY(FChaosClothAssetBlendVerticesNode);
 			DATAFLOW_NODE_REGISTER_CREATION_FACTORY(FChaosClothAssetCopySimulationToRenderMeshNode);
 			DATAFLOW_NODE_REGISTER_CREATION_FACTORY(FChaosClothAssetDatasmithImportNode);
 			DATAFLOW_NODE_REGISTER_CREATION_FACTORY(FChaosClothAssetDeleteElementNode);
diff --git a/Engine/Plugins/ChaosClothAssetEditor/Source/ChaosClothAssetDataflowNodes/Public/ChaosClothAsset/BlendVerticesNode.h b/Engine/Plugins/ChaosClothAssetEditor/Source/ChaosClothAssetDataflowNodes/Public/ChaosClothAsset/BlendVerticesNode.h
new file mode 100644
--- /dev/null
+++ b/Engine/Plugins/ChaosClothAssetEditor/Source/ChaosClothAssetDataflowNodes/Public/ChaosClothAsset/BlendVerticesNode.h
@@ -0,0 +1,75 @@
+// Copyright Epic Games, Inc. All Rights Reserved.
+
+#pragma once
+
+#include "Dataflow/DataflowNode.h"
+#include "GeometryCollection/ManagedArrayCollection.h"
+#include "BlendVerticesNode.generated.h"
+
+/** Blend vertex values from another cloth collection. The topology of the Collection will remain the same.*/
+USTRUCT(Meta = (DataflowCloth))
+struct FChaosClothAssetBlendVerticesNode final : public FDataflowNode 
+{
+	GENERATED_USTRUCT_BODY()
+	DATAFLOW_NODE_DEFINE_INTERNAL(FChaosClothAssetBlendVerticesNode, "BlendVertices", "Cloth", "Cloth Simulation Blend Vertices")
+	DATAFLOW_NODE_RENDER_TYPE("SurfaceRender", FName("FClothCollection"), "Collection")
+
+public:
+
+	/** Input/output collection */
+	UPROPERTY(Meta = (Dataflowinput, DataflowOutput, DataflowPassthrough = "Collection"))
+	FManagedArrayCollection Collection;
+
+	/** Collection to blend in. */
+	UPROPERTY(Meta = (Dataflowinput))
+	FManagedArrayCollection BlendCollection;
+
+	/** Require same vertex counts between Collection and BlendCollection in order to blend a vertex type. Otherwise the shared subset will be blended. */
+	UPROPERTY(EditAnywhere, Category = "Blend Vertices")
+	bool bRequireSameVertexCounts = true;
+
+	/** Blending Weight. 0 = Keep existing values in Collection. 1 = Set values from BlendCollection. */
+	UPROPERTY(EditAnywhere, Category = "Blend Vertices", Meta = (ClampMin = 0, ClampMax = 1))
+	float BlendingWeight = 1.f;
+
+	/** Blend Sim Mesh.*/
+	UPROPERTY(EditAnywhere, Category = "Blend Vertices - Sim Mesh", Meta = (EditCondition = "BlendingWeight > 0"))
+	bool bBlendSimMesh = true;
+
+	/** Blend 2D Sim Positions.*/
+	UPROPERTY(EditAnywhere, Category = "Blend Vertices - Sim Mesh", Meta = (EditCondition = "bBlendSimMesh && BlendingWeight > 0"))
+	bool bBlend2DSimPositions = true;
+
+	/** Blend 3D Sim Positions.*/
+	UPROPERTY(EditAnywhere, Category = "Blend Vertices - Sim Mesh", Meta = (EditCondition = "bBlendSimMesh && BlendingWeight > 0"))
+	bool bBlend3DSimPositions = true;
+
+	/** Blend 3D Sim Normals.*/
+	UPROPERTY(EditAnywhere, Category = "Blend Vertices - Sim Mesh", Meta = (EditCondition = "bBlendSimMesh && BlendingWeight > 0"))
+	bool bBlendSimNormals = true;
+
+	/** Blend Render Mesh.*/
+	UPROPERTY(EditAnywhere, Category = "Blend Vertices - Render Mesh", Meta = (EditCondition = "BlendingWeight > 0.f"))
+	bool bBlendRenderMesh = true;
+
+	/** Blend Render Positions.*/
+	UPROPERTY(EditAnywhere, Category = "Blend Vertices - Render Mesh", Meta = (EditCondition = "bBlendRenderMesh && BlendingWeight > 0"))
+	bool bBlendRenderPositions = true;
+
+	/** Blend Render Normals and Tangents.*/
+	UPROPERTY(EditAnywhere, Category = "Blend Vertices - Render Mesh", Meta = (EditCondition = "bBlendRenderMesh && BlendingWeight > 0"))
+	bool bBlendRenderNormalsAndTangents = true;
+
+	/** Blend Render UVs. Only existing UV sets on Collection will be updated.*/
+	UPROPERTY(EditAnywhere, Category = "Blend Vertices - Render Mesh", Meta = (EditCondition = "bBlendRenderMesh && BlendingWeight > 0"))
+	bool bBlendRenderUVs = true;
+
+	/** Blend Render Colors.*/
+	UPROPERTY(EditAnywhere, Category = "Blend Vertices - Render Mesh", Meta = (EditCondition = "bBlendRenderMesh && BlendingWeight > 0"))
+	bool bBlendRenderColors = true;
+
+	FChaosClothAssetBlendVerticesNode(const UE::Dataflow::FNodeParameters& InParam, FGuid InGuid = FGuid::NewGuid());
+
+private:
+	virtual void Evaluate(UE::Dataflow::FContext& Context, const FDataflowOutput* Out) const override;
+};
