Chaos Cloth Asset - Improved the Proxy Deformer node documentation.

#okforgihub public
[RN] Chaos Cloth Asset - Improved the Proxy Deformer node documentation.

diff --git a/Engine/Plugins/ChaosClothAssetEditor/Source/ChaosClothAssetDataflowNodes/Public/ChaosClothAsset/ProxyDeformerNode.h b/Engine/Plugins/ChaosClothAssetEditor/Source/ChaosClothAssetDataflowNodes/Public/ChaosClothAsset/ProxyDeformerNode.h
--- a/Engine/Plugins/ChaosClothAssetEditor/Source/ChaosClothAssetDataflowNodes/Public/ChaosClothAsset/ProxyDeformerNode.h
+++ b/Engine/Plugins/ChaosClothAssetEditor/Source/ChaosClothAssetDataflowNodes/Public/ChaosClothAsset/ProxyDeformerNode.h
@@ -27,7 +27,36 @@ struct FChaosClothAssetSelectionFilterSet
 	FChaosClothAssetConnectableIStringValue SimSelection = { TEXT("SelectionSimFilterSet") };
 };
 
-/** Add the proxy deformer information to this cloth collection's render data. */
+/**
+ * Adds the proxy deformer information to this cloth collection's render data.
+ * 
+ * This node is only required to selectively assign specific areas of the render mesh to the sim mesh in the proxy deformer data.
+ * When this node isn't provided to the construction graph, default proxy deformer data are automatically generated.
+ * This node can be used to either override or remove the default proxy deformer data.
+ * 
+ * To enable deformer areas, selection filter sets must be provided.
+ * With no selection filter sets, the Cloth Asset is by default fully skinned.
+ * When more selection filter sets are needed, right-click on the node and do AddOptionPin to add more sets.
+ * Each selection filter set is composed by two selections:
+ * * A render mesh selection (natively a RenderVertices selection, but RenderFaces selections can also be used).
+ * * A simulation mesh selection (natively a SimFaces selection, but SimVertices3D and SimVertices2D selections can also be used).
+ * The render mesh selection of vertices will then only be deformed by the associated simulation mesh selection of faces.
+ * The Proxy Deformer node evaluations are usually slower than the equivalent default proxy deformer due to the added processing cost
+ * of the selection filter sets. When in use, it is therefore better to place it earlier in the graph.
+ * 
+ * A default SkinningBlend weight map is also generated by the node that masks the areas for which selection filters aren't provided.
+ * The map can then be painted, or replaced by a smoothed transition map generated by the SkinningBlend node.
+ * The SkinningBlend node can leverage information about the kinematic areas added by the MaxDistanceConfig node.
+ * Unlike the default SkinningBlend weight map provided by the Proxy Deformer node, the SkinningBlend node generates areas of smooth
+ * transition where the render mesh vertices can be simultaneously skinned and deformed by the cloth.
+ * 
+ * To sum up:
+ * * No ProxyDeformer node -> The Cloth Asset is fully deformed.
+ * * ProxyDeformer node without selection filters -> The Cloth Asset is fully skinned.
+ * * ProxyDeformer node with selection filters -> the Cloth Asset is partially deformed/skinned.
+ * * The ProxyDeformer can be slow to evaluate, always place it as early as possible in the graph to avoid costly re-evaluations.
+ * * A SkinningBlend node placed after the MaxDistanceConfig adds smooth transitions between the skinned and proxy deformed areas.
+ */
 USTRUCT(Meta = (DataflowCloth))
 struct FChaosClothAssetProxyDeformerNode_v2 : public FDataflowNode
 {
