Chaos Cloth Asset - Fixed normals on skeletal and static mesh imports that only worked when importing all sections of a Static Mesh.

[RN] Chaos Cloth Asset - Fixed normals on skeletal and static mesh imports that only worked when importing all sections of a Static Mesh.

diff --git a/Engine/Plugins/ChaosClothAssetEditor/Source/ChaosClothAssetDataflowNodes/Private/ChaosClothAsset/ClothDataflowTools.cpp b/Engine/Plugins/ChaosClothAssetEditor/Source/ChaosClothAssetDataflowNodes/Private/ChaosClothAsset/ClothDataflowTools.cpp
--- a/Engine/Plugins/ChaosClothAssetEditor/Source/ChaosClothAssetDataflowNodes/Private/ChaosClothAsset/ClothDataflowTools.cpp
+++ b/Engine/Plugins/ChaosClothAssetEditor/Source/ChaosClothAssetDataflowNodes/Private/ChaosClothAsset/ClothDataflowTools.cpp
@@ -452,7 +452,7 @@ namespace UE::Chaos::ClothAsset
 		UE::Geometry::FNonManifoldMappingSupport::AttachNonManifoldVertexMappingData(SkelMeshSectionToDynamicMesh.ToSrcVertIDMap, DynamicMesh);
 
 		constexpr bool bAppend = true;
-		FClothGeometryTools::BuildSimMeshFromDynamicMesh(ClothCollection, DynamicMesh, UVChannelIndex, UVScale, bAppend);
+		FClothGeometryTools::BuildSimMeshFromDynamicMesh(ClothCollection, DynamicMesh, UVChannelIndex, UVScale, bAppend, bImportNormals);
 	}
 
 	void FClothDataflowTools::LogAndToastWarning(const FDataflowNode& DataflowNode, const FText& Headline, const FText& Details)
