Deprecate UClothEditorContextObject. It has been replaced everywhere by UDataflowContextObject.

diff --git a/Engine/Plugins/ChaosClothAssetEditor/Source/ChaosClothAssetEditor/Private/ChaosClothAsset/ClothEditorMode.cpp b/Engine/Plugins/ChaosClothAssetEditor/Source/ChaosClothAssetEditor/Private/ChaosClothAsset/ClothEditorMode.cpp
--- a/Engine/Plugins/ChaosClothAssetEditor/Source/ChaosClothAssetEditor/Private/ChaosClothAsset/ClothEditorMode.cpp
+++ b/Engine/Plugins/ChaosClothAssetEditor/Source/ChaosClothAssetEditor/Private/ChaosClothAsset/ClothEditorMode.cpp
@@ -9,7 +9,6 @@
 #include "ChaosClothAsset/ClothEditorRestSpaceViewportClient.h"
 #include "ChaosClothAsset/ClothPatternToDynamicMesh.h"
 #include "ChaosClothAsset/ClothEditorPreviewScene.h"
-#include "ChaosClothAsset/ClothEditorContextObject.h"
 #include "ChaosClothAsset/ClothEditorToolBuilders.h"
 #include "ChaosClothAsset/WeightMapNode.h"
 #include "ChaosClothAsset/TransferSkinWeightsNode.h"
diff --git a/Engine/Plugins/ChaosClothAssetEditor/Source/ChaosClothAssetEditorTools/Private/ChaosClothAsset/ClothWeightMapPaintTool.cpp b/Engine/Plugins/ChaosClothAssetEditor/Source/ChaosClothAssetEditorTools/Private/ChaosClothAsset/ClothWeightMapPaintTool.cpp
--- a/Engine/Plugins/ChaosClothAssetEditor/Source/ChaosClothAssetEditorTools/Private/ChaosClothAsset/ClothWeightMapPaintTool.cpp
+++ b/Engine/Plugins/ChaosClothAssetEditor/Source/ChaosClothAssetEditorTools/Private/ChaosClothAsset/ClothWeightMapPaintTool.cpp
@@ -32,7 +32,6 @@
 #include "CanvasTypes.h"
 #include "CanvasItem.h"
 #include "ChaosClothAsset/ClothCollectionGroup.h"
-#include "ChaosClothAsset/ClothEditorContextObject.h"
 #include "ChaosClothAsset/ClothEditorToolBuilders.h"
 #include "ChaosClothAsset/ClothPatternVertexType.h"
 #include "ChaosClothAsset/CollectionClothFacade.h"
diff --git a/Engine/Plugins/ChaosClothAssetEditor/Source/ChaosClothAssetEditorTools/Public/ChaosClothAsset/ClothEditorContextObject.h b/Engine/Plugins/ChaosClothAssetEditor/Source/ChaosClothAssetEditorTools/Public/ChaosClothAsset/ClothEditorContextObject.h
--- a/Engine/Plugins/ChaosClothAssetEditor/Source/ChaosClothAssetEditorTools/Public/ChaosClothAsset/ClothEditorContextObject.h
+++ b/Engine/Plugins/ChaosClothAssetEditor/Source/ChaosClothAssetEditorTools/Public/ChaosClothAsset/ClothEditorContextObject.h
@@ -25,7 +25,7 @@ struct FManagedArrayCollection;
 
 
 UCLASS()
-class CHAOSCLOTHASSETEDITORTOOLS_API UClothEditorContextObject : public UObject
+class UE_DEPRECATED(5.6, "Use UDataflowContextObject instead") CHAOSCLOTHASSETEDITORTOOLS_API UClothEditorContextObject : public UObject
 {
 	GENERATED_BODY()
 
