Dataflow Editor / Cloth Selection Tool: if we are selecting Sim faces, allow switching between 2DSim and 3DSim view modes without stopping the tool

diff --git a/Engine/Plugins/ChaosClothAssetEditor/Source/ChaosClothAssetEditorTools/Private/ChaosClothAsset/ClothEditorToolBuilders.cpp b/Engine/Plugins/ChaosClothAssetEditor/Source/ChaosClothAssetEditorTools/Private/ChaosClothAsset/ClothEditorToolBuilders.cpp
--- a/Engine/Plugins/ChaosClothAssetEditor/Source/ChaosClothAssetEditorTools/Private/ChaosClothAsset/ClothEditorToolBuilders.cpp
+++ b/Engine/Plugins/ChaosClothAssetEditor/Source/ChaosClothAssetEditorTools/Private/ChaosClothAsset/ClothEditorToolBuilders.cpp
@@ -163,7 +163,19 @@ void UClothMeshSelectionToolBuilder::GetSupportedConstructionViewModes(const UDa
 			return;
 		}
 		
-		if (SelectionNode->Group.Name == ClothCollectionGroup::SimVertices2D.ToString() || SelectionNode->Group.Name == ClothCollectionGroup::SimVertices3D.ToString() || SelectionNode->Group.Name == ClothCollectionGroup::SimFaces.ToString())
+		if (SelectionNode->Group.Name == ClothCollectionGroup::SimVertices2D.ToString())
+		{
+			Modes.Add(Sim2DMode);
+			return;
+		}
+
+		if (SelectionNode->Group.Name == ClothCollectionGroup::SimVertices3D.ToString())
+		{
+			Modes.Add(Sim3DMode);
+			return;
+		}
+
+		if (SelectionNode->Group.Name == ClothCollectionGroup::SimFaces.ToString())
 		{
 			Modes.Add(Sim2DMode);
 			Modes.Add(Sim3DMode);
@@ -188,12 +200,27 @@ void UClothMeshSelectionToolBuilder::GetSupportedConstructionViewModes(const UDa
 
 bool UClothMeshSelectionToolBuilder::CanSceneStateChange(const UInteractiveTool* ActiveTool, const FToolBuilderState& SceneState) const
 {
-	return false;
+	return ActiveTool->IsA<UClothMeshSelectionTool>();
 }
 
 void UClothMeshSelectionToolBuilder::SceneStateChanged(UInteractiveTool* ActiveTool, const FToolBuilderState& SceneState)
 {
 	check(CanSceneStateChange(ActiveTool, SceneState));
+
+	UClothMeshSelectionTool* const SelectionTool = Cast<UClothMeshSelectionTool>(ActiveTool);
+	checkf(SelectionTool, TEXT("Expected the ActiveTool to be UClothMeshSelectionTool"));
+
+	UToolTarget* const Target = SceneState.TargetManager->BuildFirstSelectedTargetable(SceneState, GetTargetRequirements());
+	check(Target);
+	check(Target->IsValid());
+	SelectionTool->SetTarget(Target);
+	check(SelectionTool->GetTargetWorld() == SceneState.World);
+	SelectionTool->NotifyTargetChanged();
+
+	// These are likely to be empty functions but are called here for completeness (see UInteractiveToolManager::ActivateToolInternal())
+	PostBuildTool(ActiveTool, SceneState);
+	PostSetupTool(ActiveTool, SceneState);
+
 }
 
 void UClothMeshSelectionToolBuilder::GetSupportedViewModes(const UDataflowContextObject& ContextObject, TArray<UE::Chaos::ClothAsset::EClothPatternVertexType>& Modes) const
diff --git a/Engine/Plugins/ChaosClothAssetEditor/Source/ChaosClothAssetEditorTools/Private/ChaosClothAsset/ClothMeshSelectionTool.cpp b/Engine/Plugins/ChaosClothAssetEditor/Source/ChaosClothAssetEditorTools/Private/ChaosClothAsset/ClothMeshSelectionTool.cpp
--- a/Engine/Plugins/ChaosClothAssetEditor/Source/ChaosClothAssetEditorTools/Private/ChaosClothAsset/ClothMeshSelectionTool.cpp
+++ b/Engine/Plugins/ChaosClothAssetEditor/Source/ChaosClothAssetEditorTools/Private/ChaosClothAsset/ClothMeshSelectionTool.cpp
@@ -148,7 +148,8 @@
 
 // ------------------- Tool -------------------
 	
-void UClothMeshSelectionTool::Setup()
+
+void UClothMeshSelectionTool::InitializeSculptMeshFromTarget()
 {
 	using namespace UE::Geometry;
 
@@ -156,7 +157,11 @@
 	// Preview 
 	//
 
-	// Create the preview mesh
+	// (Re-)Create the preview mesh
+	if (PreviewMesh)
+	{
+		PreviewMesh->Disconnect();
+	}
 	PreviewMesh = NewObject<UPreviewMesh>(this, TEXT("PreviewMesh"));
 	PreviewMesh->CreateInWorld(GetTargetWorld(), FTransform::Identity);
 
@@ -185,6 +190,51 @@
 	PreviewMesh->UpdatePreview(UE::ToolTarget::GetDynamicMeshCopy(Target));
 	PreviewMesh->SetVisible(true);
 
+	// Hide input target mesh
+	UE::ToolTarget::HideSourceObject(Target);
+
+	// Setup non-manifold mapping if necessary
+	if (DataflowContextObject)
+	{
+		ensure(DataflowContextObject->IsUsingInputCollection());
+		if (const TSharedPtr<const FManagedArrayCollection> ClothCollection = DataflowContextObject->GetSelectedCollection())
+		{
+			PreviewMesh->ProcessMesh([this, ClothCollection](const FDynamicMesh3& Mesh)
+			{
+				using namespace UE::Chaos::ClothAsset;
+				const UE::Geometry::FNonManifoldMappingSupport NonManifoldMapping(Mesh);
+
+				bHasNonManifoldMapping = NonManifoldMapping.IsNonManifoldVertexInSource();
+
+				if (bHasNonManifoldMapping)
+				{
+					const FCollectionClothConstFacade Cloth(ClothCollection.ToSharedRef());
+					check(Cloth.IsValid());
+
+					const TConstArrayView<int32> SimVertex3DLookup = Cloth.GetSimVertex3DLookup();
+
+					DynamicMeshToSelection.SetNumUninitialized(Mesh.VertexCount());
+					SelectionToDynamicMesh.Reset();
+					SelectionToDynamicMesh.SetNum(Cloth.GetNumSimVertices3D());
+					for (int32 DynamicMeshVert = 0; DynamicMeshVert < Mesh.VertexCount(); ++DynamicMeshVert)
+					{
+						DynamicMeshToSelection[DynamicMeshVert] = NonManifoldMapping.GetOriginalNonManifoldVertexID(DynamicMeshVert);
+						SelectionToDynamicMesh[DynamicMeshToSelection[DynamicMeshVert]].Add(DynamicMeshVert);
+					}
+				}
+			});
+		}
+	}
+}
+
+
+
+void UClothMeshSelectionTool::Setup()
+{
+	using namespace UE::Geometry;
+
+	InitializeSculptMeshFromTarget();
+
 	//
 	// SelectionMechanic
 	//
@@ -252,41 +302,6 @@
 		return SelectionMechanic->GetActiveSelection().IsSelectedTriangle(Mesh, Topology.Get(), TriangleID);
 	});
 
-	// Hide input target mesh
-	UE::ToolTarget::HideSourceObject(Target);
-
-	// Setup non-manifold mapping if necessary
-	if (DataflowContextObject)
-	{
-		ensure(DataflowContextObject->IsUsingInputCollection());
-		if (const TSharedPtr<const FManagedArrayCollection> ClothCollection = DataflowContextObject->GetSelectedCollection())
-		{
-			PreviewMesh->ProcessMesh([this, ClothCollection](const FDynamicMesh3& Mesh)
-			{
-				using namespace UE::Chaos::ClothAsset;
-				const FNonManifoldMappingSupport NonManifoldMapping(Mesh);
-
-				bHasNonManifoldMapping = NonManifoldMapping.IsNonManifoldVertexInSource();
-
-				if (bHasNonManifoldMapping)
-				{
-					const FCollectionClothConstFacade Cloth(ClothCollection.ToSharedRef());
-					check(Cloth.IsValid());
-
-					const TConstArrayView<int32> SimVertex3DLookup = Cloth.GetSimVertex3DLookup();
-
-					DynamicMeshToSelection.SetNumUninitialized(Mesh.VertexCount());
-					SelectionToDynamicMesh.Reset();
-					SelectionToDynamicMesh.SetNum(Cloth.GetNumSimVertices3D());
-					for (int32 DynamicMeshVert = 0; DynamicMeshVert < Mesh.VertexCount(); ++DynamicMeshVert)
-					{
-						DynamicMeshToSelection[DynamicMeshVert] = NonManifoldMapping.GetOriginalNonManifoldVertexID(DynamicMeshVert);
-						SelectionToDynamicMesh[DynamicMeshToSelection[DynamicMeshVert]].Add(DynamicMeshVert);
-					}
-				}
-			});
-		}
-	}
 
 
 	//
@@ -634,5 +649,48 @@
 	}
 }
 
+void UClothMeshSelectionTool::NotifyTargetChanged()
+{
+	using namespace UE::Chaos::ClothAsset;
 
+	//
+	// The target mesh has changed due to a view mode change. We will transfer the current in-progress selection to the new mesh.
+	// First, temporarily save the existing selection
+	// 
+
+	const FGroupTopologySelection CurrentSelection = SelectionMechanic->GetActiveSelection();
+
+	//
+	// Now re-initialize everything that depends on the mesh
+	//
+
+	InitializeSculptMeshFromTarget();
+
+	PreviewMesh->ProcessMesh([this](const FDynamicMesh3& Mesh)
+	{
+		Topology = MakeUnique<UE::Geometry::FTriangleGroupTopology>(&Mesh, true);
+
+		SelectionMechanic->Initialize(&Mesh,
+			FTransform3d(),
+			GetTargetWorld(),
+			Topology.Get(),
+			[this]() { return PreviewMesh->GetSpatial(); });
+	});
+
+	PreviewMesh->EnableSecondaryTriangleBuffers([this](const FDynamicMesh3* Mesh, int32 TriangleID)
+	{
+		return SelectionMechanic->GetActiveSelection().IsSelectedTriangle(Mesh, Topology.Get(), TriangleID);
+	});
+
+	SelectionMechanic->NotifyMeshChanged(true);
+
+	//
+	// Copy saved values back to the new preview mesh
+	//
+
+	constexpr bool bBroadcastChange = false;
+	SelectionMechanic->SetSelection(CurrentSelection, bBroadcastChange);
+	GetToolManager()->PostInvalidation();
+}
+
 #undef LOCTEXT_NAMESPACE
diff --git a/Engine/Plugins/ChaosClothAssetEditor/Source/ChaosClothAssetEditorTools/Private/ChaosClothAsset/ClothMeshSelectionTool.h b/Engine/Plugins/ChaosClothAssetEditor/Source/ChaosClothAssetEditorTools/Private/ChaosClothAsset/ClothMeshSelectionTool.h
--- a/Engine/Plugins/ChaosClothAssetEditor/Source/ChaosClothAssetEditorTools/Private/ChaosClothAsset/ClothMeshSelectionTool.h
+++ b/Engine/Plugins/ChaosClothAssetEditor/Source/ChaosClothAssetEditorTools/Private/ChaosClothAsset/ClothMeshSelectionTool.h
@@ -155,13 +155,13 @@ private:
 
 	FChaosClothAssetSelectionNode_v2* SelectionNodeToUpdate = nullptr;
 	TSet<int32> InputSelectionSet;
-	//
-	// Action support
-	//
 
 public:
+
 	virtual void RequestAction(EClothMeshSelectionToolActions ActionType);
 
+	void NotifyTargetChanged();
+
 	UPROPERTY()
 	TObjectPtr<UClothMeshSelectionToolActions> ActionsProps;
 
@@ -170,5 +170,6 @@ private:
 	EClothMeshSelectionToolActions PendingAction;
 	virtual void ApplyAction(EClothMeshSelectionToolActions ActionType);
 
+	void InitializeSculptMeshFromTarget();
 };
 
