Cloth Asset - Deprecate secondary selection set in cloth selection facade.

#rnx

diff --git a/Engine/Plugins/ChaosClothAsset/Source/ChaosClothAsset/Private/ChaosClothAsset/ClothGeometryTools.cpp b/Engine/Plugins/ChaosClothAsset/Source/ChaosClothAsset/Private/ChaosClothAsset/ClothGeometryTools.cpp
--- a/Engine/Plugins/ChaosClothAsset/Source/ChaosClothAsset/Private/ChaosClothAsset/ClothGeometryTools.cpp
+++ b/Engine/Plugins/ChaosClothAsset/Source/ChaosClothAsset/Private/ChaosClothAsset/ClothGeometryTools.cpp
@@ -1299,6 +1299,7 @@ namespace UE::Chaos::ClothAsset
 	{
 		FCollectionClothSelectionConstFacade SelectionFacade(ClothCollection);
 		FCollectionClothConstFacade ClothFacade(ClothCollection);
+PRAGMA_DISABLE_DEPRECATION_WARNINGS
 		if (!SelectionFacade.IsValid() || !ClothFacade.IsValid() || (bSecondarySelection ? !SelectionFacade.HasSelectionSecondarySet(SelectionName) : !SelectionFacade.HasSelection(SelectionName)))
 		{
 			return false;
@@ -1306,6 +1307,7 @@ namespace UE::Chaos::ClothAsset
 
 		const TSet<int32>& OrigSelectionSet = bSecondarySelection ? SelectionFacade.GetSelectionSecondarySet(SelectionName) : SelectionFacade.GetSelectionSet(SelectionName);
 		const FName OrigSelectionGroup = bSecondarySelection ? SelectionFacade.GetSelectionSecondaryGroup(SelectionName) : SelectionFacade.GetSelectionGroup(SelectionName);
+PRAGMA_ENABLE_DEPRECATION_WARNINGS
 
 		if (OrigSelectionGroup == GroupName)
 		{
diff --git a/Engine/Plugins/ChaosClothAsset/Source/ChaosClothAsset/Private/ChaosClothAsset/CollectionClothSelectionFacade.cpp b/Engine/Plugins/ChaosClothAsset/Source/ChaosClothAsset/Private/ChaosClothAsset/CollectionClothSelectionFacade.cpp
--- a/Engine/Plugins/ChaosClothAsset/Source/ChaosClothAsset/Private/ChaosClothAsset/CollectionClothSelectionFacade.cpp
+++ b/Engine/Plugins/ChaosClothAsset/Source/ChaosClothAsset/Private/ChaosClothAsset/CollectionClothSelectionFacade.cpp
@@ -79,7 +79,9 @@ namespace UE::Chaos::ClothAsset
 	FName FCollectionClothSelectionConstFacade::GetSelectionSecondaryGroup(const FName& Name) const
 	{
 		check(IsValid());
+PRAGMA_DISABLE_DEPRECATION_WARNINGS
 		check(HasSelectionSecondarySet(Name));	// Also checks the primary set exists
+PRAGMA_ENABLE_DEPRECATION_WARNINGS
 		return ManagedArrayCollection->GetDependency(Name, Private::SelectionSecondaryGroup);
 	}
 
diff --git a/Engine/Plugins/ChaosClothAsset/Source/ChaosClothAsset/Public/ChaosClothAsset/CollectionClothSelectionFacade.h b/Engine/Plugins/ChaosClothAsset/Source/ChaosClothAsset/Public/ChaosClothAsset/CollectionClothSelectionFacade.h
--- a/Engine/Plugins/ChaosClothAsset/Source/ChaosClothAsset/Public/ChaosClothAsset/CollectionClothSelectionFacade.h
+++ b/Engine/Plugins/ChaosClothAsset/Source/ChaosClothAsset/Public/ChaosClothAsset/CollectionClothSelectionFacade.h
@@ -48,15 +48,19 @@ namespace UE::Chaos::ClothAsset
 
 
 		/** Return whether a selection with the given name currently exists in the collection and has a secondary set associated with it. */
+		UE_DEPRECATED(5.6, "Secondary selection sets are no longer supported.")
 		CHAOSCLOTHASSET_API bool HasSelectionSecondarySet(const FName& Name) const;
 
 		/** Return the secondary selection set group dependency. The selection must exist and have a secondary set to call this function. */
+		UE_DEPRECATED(5.6, "Secondary selection sets are no longer supported.")
 		CHAOSCLOTHASSET_API FName GetSelectionSecondaryGroup(const FName& Name) const;
 
 		/** Get the secondary selection set for the given selection name. The selection must exist and have a secondary set to call this function. */
+		UE_DEPRECATED(5.6, "Secondary selection sets are no longer supported.")
 		CHAOSCLOTHASSET_API const TSet<int32>& GetSelectionSecondarySet(const FName& Name) const;
 
 		/** Find the secondary selection set for the given selection name. Returns nullptr if no such selection exists or if it does not have an associated secondary set. */
+		UE_DEPRECATED(5.6, "Secondary selection sets are no longer supported.")
 		CHAOSCLOTHASSET_API const TSet<int32>* FindSelectionSecondarySet(const FName& Name) const;
 
 	protected:
@@ -92,12 +96,15 @@ namespace UE::Chaos::ClothAsset
 		CHAOSCLOTHASSET_API void RemoveSelectionSet(const FName& Name);
 
 		/** Get the secondary selection set for the given selection name. The selection must exist and have a secondary set to call this function. */
+		UE_DEPRECATED(5.6, "Secondary selection sets are no longer supported.")
 		CHAOSCLOTHASSET_API TSet<int32>& GetSelectionSecondarySet(const FName& Name);
 
 		/** Find the secondary selection set for the given selection name. Returns nullptr if no such selection exists or if it does not have an associated secondary set. */
+		UE_DEPRECATED(5.6, "Secondary selection sets are no longer supported.")
 		CHAOSCLOTHASSET_API TSet<int32>* FindSelectionSecondarySet(const FName& Name);
 
 		/** Remove the secondary selection set from the selection with the given name if it exists. */
+		UE_DEPRECATED(5.6, "Secondary selection sets are no longer supported.")
 		CHAOSCLOTHASSET_API void RemoveSelectionSecondarySet(const FName& Name);
 
 		/** 
@@ -136,6 +143,7 @@ namespace UE::Chaos::ClothAsset
 		 * If the group doesn't already exists, this function will create it.
 		 * If the secondary selection set already exists, but depends on a different group, the old secondary selection set will be deleted and a new one recreated with the new group dependency.
 		 */
+		UE_DEPRECATED(5.6, "Secondary selection sets are no longer supported.")
 		CHAOSCLOTHASSET_API TSet<int32>& FindOrAddSelectionSecondarySet(const FName& Name, const FName& GroupName);
 
 	private:
diff --git a/Engine/Plugins/ChaosClothAssetEditor/Source/ChaosClothAssetDataflowNodes/Private/ChaosClothAsset/ProxyDeformerNode.cpp b/Engine/Plugins/ChaosClothAssetEditor/Source/ChaosClothAssetDataflowNodes/Private/ChaosClothAsset/ProxyDeformerNode.cpp
--- a/Engine/Plugins/ChaosClothAssetEditor/Source/ChaosClothAssetDataflowNodes/Private/ChaosClothAsset/ProxyDeformerNode.cpp
+++ b/Engine/Plugins/ChaosClothAssetEditor/Source/ChaosClothAssetDataflowNodes/Private/ChaosClothAsset/ProxyDeformerNode.cpp
@@ -361,13 +361,16 @@ namespace UE::Chaos::ClothAsset::Private
 			{
 				if (const TSet<int32>* SelectionSet = SelectionFacade.FindSelectionSet(SelectionName))
 				{
+PRAGMA_DISABLE_DEPRECATION_WARNINGS
 					if (const TSet<int32>* SecondarySelectionSet = SelectionFacade.FindSelectionSecondarySet(SelectionName))
+PRAGMA_ENABLE_DEPRECATION_WARNINGS
 					{
 						if (SelectionSet->Num() && SecondarySelectionSet->Num())
 						{
 							FName SelectionGroup = SelectionFacade.GetSelectionGroup(SelectionName);
+PRAGMA_DISABLE_DEPRECATION_WARNINGS
 							FName SelectionSecondaryGroup = SelectionFacade.GetSelectionSecondaryGroup(SelectionName);
-
+PRAGMA_ENABLE_DEPRECATION_WARNINGS
 							// Retrieve the sim face selection
 							TSet<int32> SimFaceSelection = GetSimFaceSelection(SelectionGroup, *SelectionSet);
 							if (!SimFaceSelection.Num())
diff --git a/Engine/Plugins/ChaosClothAssetEditor/Source/ChaosClothAssetDataflowNodes/Private/ChaosClothAsset/SelectionNode.cpp b/Engine/Plugins/ChaosClothAssetEditor/Source/ChaosClothAssetDataflowNodes/Private/ChaosClothAsset/SelectionNode.cpp
--- a/Engine/Plugins/ChaosClothAssetEditor/Source/ChaosClothAssetDataflowNodes/Private/ChaosClothAsset/SelectionNode.cpp
+++ b/Engine/Plugins/ChaosClothAssetEditor/Source/ChaosClothAssetDataflowNodes/Private/ChaosClothAsset/SelectionNode.cpp
@@ -324,6 +324,7 @@ void FChaosClothAssetSelectionNode_v2::OnImportSecondary(UE::Dataflow::FContext&
 	const FCollectionClothSelectionConstFacade SelectionFacade(ClothCollection);
 
 	const FName InInputName = GetInputName(Context);
+PRAGMA_DISABLE_DEPRECATION_WARNINGS
 	if (const TSet<int32>* const SelectionSet = SelectionFacade.IsValid() ? SelectionFacade.FindSelectionSecondarySet(InInputName) : nullptr)
 	{
 		Indices = *SelectionSet;
@@ -331,6 +332,7 @@ void FChaosClothAssetSelectionNode_v2::OnImportSecondary(UE::Dataflow::FContext&
 		Group.Name = SelectionFacade.GetSelectionSecondaryGroup(InInputName).ToString();
 		SelectionOverrideType = EChaosClothAssetSelectionOverrideType::ReplaceAll;
 	}
+PRAGMA_ENABLE_DEPRECATION_WARNINGS
 	else
 	{
 		FClothDataflowTools::LogAndToastWarning(*this,
@@ -529,12 +531,14 @@ void FChaosClothAssetSelectionNode::SetAssetValue(TObjectPtr<UObject> Asset, UE:
 						const TArray<int32> SelectionAsArray = TransferSelectionFacade.GetSelectionSet(InInputName).Array();
 						InTransferCollectionHash = GetArrayHash(SelectionAsArray.GetData(), SelectionAsArray.Num(), InTransferCollectionHash);
 
+PRAGMA_DISABLE_DEPRECATION_WARNINGS
 						if (TransferSelectionFacade.HasSelectionSecondarySet(InInputName))
 						{
 							InTransferCollectionHash = HashCombineFast(InTransferCollectionHash, GetTypeHash(TransferSelectionFacade.GetSelectionSecondaryGroup(InInputName)));
 							const TArray<int32> SecondarySelectionAsArray = TransferSelectionFacade.GetSelectionSecondarySet(InInputName).Array();
 							InTransferCollectionHash = GetArrayHash(SecondarySelectionAsArray.GetData(), SecondarySelectionAsArray.Num(), InTransferCollectionHash);
 						}
+PRAGMA_ENABLE_DEPRECATION_WARNINGS
 					}
 					else
 					{
@@ -561,11 +565,11 @@ void FChaosClothAssetSelectionNode::SetAssetValue(TObjectPtr<UObject> Asset, UE:
 							{
 								TSet<int32> InputSelection;
 
-								PRAGMA_DISABLE_DEPRECATION_WARNINGS
+PRAGMA_DISABLE_DEPRECATION_WARNINGS
 								FClothGeometryTools::ConvertSelectionToNewGroupType(ClothCollection, InInputName, SelectionGroupName, true, InputSelection);
-								PRAGMA_ENABLE_DEPRECATION_WARNINGS
 
 								MutableThis->SetSecondaryIndices(InputSelection, SecondaryFinalSelection);
+PRAGMA_ENABLE_DEPRECATION_WARNINGS
 							}
 						}
 					}
@@ -653,12 +657,12 @@ void FChaosClothAssetSelectionNode::Evaluate(UE::Dataflow::FContext& Context, co
 			const FName SecondarySelectionGroupName = (*SecondaryGroup.Name);
 
 			InputSelectionSet.Reset();
-			PRAGMA_DISABLE_DEPRECATION_WARNINGS
+PRAGMA_DISABLE_DEPRECATION_WARNINGS
 			FClothGeometryTools::ConvertSelectionToNewGroupType(SelectionCollection, InInputName, SelectionGroupName, true, InputSelectionSet);
-			PRAGMA_ENABLE_DEPRECATION_WARNINGS
 			FinalSet.Reset();
 			CalculateFinalSecondarySet(InputSelectionSet, FinalSet);
 			TSet<int32>& SecondarySelectionSet = SelectionFacade.FindOrAddSelectionSecondarySet(SecondarySelectionName, SecondarySelectionGroupName);
+PRAGMA_ENABLE_DEPRECATION_WARNINGS
 
 			CopyIntoSelection(SelectionCollection, SecondarySelectionGroupName, FinalSet, SecondarySelectionSet);
 		}
diff --git a/Engine/Plugins/ChaosClothAssetEditor/Source/ChaosClothAssetDataflowNodes/Public/ChaosClothAsset/SelectionNode.h b/Engine/Plugins/ChaosClothAssetEditor/Source/ChaosClothAssetDataflowNodes/Public/ChaosClothAsset/SelectionNode.h
--- a/Engine/Plugins/ChaosClothAssetEditor/Source/ChaosClothAssetDataflowNodes/Public/ChaosClothAsset/SelectionNode.h
+++ b/Engine/Plugins/ChaosClothAssetEditor/Source/ChaosClothAssetDataflowNodes/Public/ChaosClothAsset/SelectionNode.h
@@ -242,9 +242,11 @@ public:
 
 	// Set Indices based on SelectionOverrideType
 	void CHAOSCLOTHASSETDATAFLOWNODES_API SetIndices(const TSet<int32>& InputSet, const TSet<int32>& FinalSet);
+	UE_DEPRECATED(5.6, "Secondary selection sets are no longer supported.")
 	void CHAOSCLOTHASSETDATAFLOWNODES_API SetSecondaryIndices(const TSet<int32>& InputSet, const TSet<int32>& FinalSet);
 
 	void CHAOSCLOTHASSETDATAFLOWNODES_API CalculateFinalSet(const TSet<int32>& InputSet, TSet<int32>& FinalSet) const;
+	UE_DEPRECATED(5.6, "Secondary selection sets are no longer supported.")
 	void CHAOSCLOTHASSETDATAFLOWNODES_API CalculateFinalSecondarySet(const TSet<int32>& InputSet, TSet<int32>& FinalSet) const;
 
 private:
