Chaos Cloth Asset: save/load the value of the cvar controlling whether cloth assets should open in the cloth editor or the dataflow editor

diff --git a/Engine/Plugins/ChaosClothAssetEditor/Source/ChaosClothAssetEditor/Private/ChaosClothAsset/ChaosClothAssetEditorModule.cpp b/Engine/Plugins/ChaosClothAssetEditor/Source/ChaosClothAssetEditor/Private/ChaosClothAsset/ChaosClothAssetEditorModule.cpp
--- a/Engine/Plugins/ChaosClothAssetEditor/Source/ChaosClothAssetEditor/Private/ChaosClothAsset/ChaosClothAssetEditorModule.cpp
+++ b/Engine/Plugins/ChaosClothAssetEditor/Source/ChaosClothAssetEditor/Private/ChaosClothAsset/ChaosClothAssetEditorModule.cpp
@@ -8,6 +8,7 @@
 #include "ChaosClothAsset/ClothSimulationNodeDetailExtender.h"
 #include "ChaosClothAsset/ClothEditorCommands.h"
 #include "ChaosClothAsset/ClothEditorMode.h"
+#include "ChaosClothAsset/ClothEditorOptions.h"
 #include "ChaosClothAsset/ClothEditorStyle.h"
 #include "Modules/ModuleManager.h"
 #include "EditorModeRegistry.h"
@@ -103,10 +104,29 @@ PRAGMA_ENABLE_DEPRECATION_WARNINGS
 		Dataflow::FDataflowConstructionVisualizationRegistry::GetInstance().RegisterVisualization(MakeUnique<FClothDataflowConstructionVisualization>());
 		Dataflow::FDataflowSimulationVisualizationRegistry::GetInstance().RegisterVisualization(MakeUnique<FClothDataflowSimulationVisualization>());
 		Dataflow::FDataflowNodeDetailExtensionRegistry::GetInstance().RegisterExtension(MakeUnique<FClothSimulationNodeDetailExtender>());
+
+		if (IConsoleVariable* const Var = IConsoleManager::Get().FindConsoleVariable(TEXT("p.ChaosCloth.EnableDataflowEditor")))	// declared in AssetDefinition_ClothAsset.cpp
+		{
+			if (UChaosClothEditorOptions* const Options = UChaosClothEditorOptions::StaticClass()->GetDefaultObject<UChaosClothEditorOptions>())
+			{
+				Var->Set(Options->bClothAssetsOpenInDataflowEditor);
+
+				OnCVarChangedDelegateHandle = Var->OnChangedDelegate().AddLambda([Options](IConsoleVariable* Variable)
+				{
+					Options->bClothAssetsOpenInDataflowEditor = Variable->GetBool();
+					Options->SaveConfig();
+				});
+			}
+		}
 	}
 
 	void FChaosClothAssetEditorModule::ShutdownModule()
 	{
+		if (IConsoleVariable* const Var = IConsoleManager::Get().FindConsoleVariable(TEXT("p.ChaosCloth.EnableDataflowEditor")))
+		{
+			Var->OnChangedDelegate().Remove(OnCVarChangedDelegateHandle);
+		}
+
 		Dataflow::FDataflowNodeDetailExtensionRegistry::GetInstance().DeregisterExtension(FClothSimulationNodeDetailExtender::Name);
 		Dataflow::FDataflowConstructionVisualizationRegistry::GetInstance().DeregisterVisualization(FClothDataflowConstructionVisualization::Name);
 		Dataflow::FDataflowSimulationVisualizationRegistry::GetInstance().DeregisterVisualization(FClothDataflowSimulationVisualization::Name);
diff --git a/Engine/Plugins/ChaosClothAssetEditor/Source/ChaosClothAssetEditor/Private/ChaosClothAsset/ChaosClothAssetEditorModule.h b/Engine/Plugins/ChaosClothAssetEditor/Source/ChaosClothAssetEditor/Private/ChaosClothAsset/ChaosClothAssetEditorModule.h
--- a/Engine/Plugins/ChaosClothAssetEditor/Source/ChaosClothAssetEditor/Private/ChaosClothAsset/ChaosClothAssetEditorModule.h
+++ b/Engine/Plugins/ChaosClothAssetEditor/Source/ChaosClothAssetEditor/Private/ChaosClothAsset/ChaosClothAssetEditorModule.h
@@ -16,6 +16,7 @@ namespace UE::Chaos::ClothAsset
 	private:
 
 		FDelegateHandle StartupCallbackDelegateHandle;
+		FDelegateHandle OnCVarChangedDelegateHandle;
 
 		void RegisterMenus();
 		
diff --git a/Engine/Plugins/ChaosClothAssetEditor/Source/ChaosClothAssetEditor/Private/ChaosClothAsset/ClothEditorOptions.cpp b/Engine/Plugins/ChaosClothAssetEditor/Source/ChaosClothAssetEditor/Private/ChaosClothAsset/ClothEditorOptions.cpp
new file mode 100644
--- /dev/null
+++ b/Engine/Plugins/ChaosClothAssetEditor/Source/ChaosClothAssetEditor/Private/ChaosClothAsset/ClothEditorOptions.cpp
@@ -0,0 +1,11 @@
+// Copyright Epic Games, Inc. All Rights Reserved.
+
+#include "ChaosClothAsset/ClothEditorOptions.h"
+
+#include UE_INLINE_GENERATED_CPP_BY_NAME(ClothEditorOptions)
+
+UChaosClothEditorOptions::UChaosClothEditorOptions(class FObjectInitializer const& ObjectInitializer)
+{
+	bClothAssetsOpenInDataflowEditor = false;
+}
+
diff --git a/Engine/Plugins/ChaosClothAssetEditor/Source/ChaosClothAssetEditor/Private/ChaosClothAsset/ClothEditorOptions.h b/Engine/Plugins/ChaosClothAssetEditor/Source/ChaosClothAssetEditor/Private/ChaosClothAsset/ClothEditorOptions.h
new file mode 100644
--- /dev/null
+++ b/Engine/Plugins/ChaosClothAssetEditor/Source/ChaosClothAssetEditor/Private/ChaosClothAsset/ClothEditorOptions.h
@@ -0,0 +1,19 @@
+// Copyright Epic Games, Inc. All Rights Reserved.
+
+#pragma once
+
+#include "UObject/Object.h"
+#include "ClothEditorOptions.generated.h"
+
+/** A settings class used to store and retreive user settings related to the Cloth Editor */
+UCLASS(hidecategories = Object, config = EditorPerProjectUserSettings, MinimalAPI)
+class UChaosClothEditorOptions : public UObject
+{
+	GENERATED_UCLASS_BODY()
+
+public:
+
+	/** Whether Chaos Cloth Assets open in the Dataflow Editor (vs in the Cloth Editor) */
+	UPROPERTY(config)
+	bool bClothAssetsOpenInDataflowEditor;
+};
