Cloth Editor: enable middle mouse button panning (really trucking) in the construction viewport in 2D mode. Toggling between right/middle/both can be done in the project settings window

diff --git a/Engine/Plugins/ChaosClothAssetEditor/Source/ChaosClothAssetEditor/ChaosClothAssetEditor.Build.cs b/Engine/Plugins/ChaosClothAssetEditor/Source/ChaosClothAssetEditor/ChaosClothAssetEditor.Build.cs
--- a/Engine/Plugins/ChaosClothAssetEditor/Source/ChaosClothAssetEditor/ChaosClothAssetEditor.Build.cs
+++ b/Engine/Plugins/ChaosClothAssetEditor/Source/ChaosClothAssetEditor/ChaosClothAssetEditor.Build.cs
@@ -70,7 +70,8 @@ public class ChaosClothAssetEditor : ModuleRules
 				"AssetDefinition",
 				"EditorWidgets",
 				"KismetWidgets",      // SScrubControlPanel
-				"AnimGraph"           // UAnimSingleNodeInstance
+				"AnimGraph",          // UAnimSingleNodeInstance
+				"DeveloperSettings"
 				// ... add private dependencies that you statically link with here ...	
 			}
 			);
diff --git a/Engine/Plugins/ChaosClothAssetEditor/Source/ChaosClothAssetEditor/Private/ChaosClothAsset/ClothEditorOptions.cpp b/Engine/Plugins/ChaosClothAssetEditor/Source/ChaosClothAssetEditor/Private/ChaosClothAsset/ClothEditorOptions.cpp
--- a/Engine/Plugins/ChaosClothAssetEditor/Source/ChaosClothAssetEditor/Private/ChaosClothAsset/ClothEditorOptions.cpp
+++ b/Engine/Plugins/ChaosClothAssetEditor/Source/ChaosClothAssetEditor/Private/ChaosClothAsset/ClothEditorOptions.cpp
@@ -7,5 +7,19 @@
 UChaosClothEditorOptions::UChaosClothEditorOptions(class FObjectInitializer const& ObjectInitializer)
 {
 	bClothAssetsOpenInDataflowEditor = false;
+	ConstructionViewportMousePanButton = EConstructionViewportMousePanButton::Right;
 }
 
+FName UChaosClothEditorOptions::GetCategoryName() const
+{
+	return TEXT("Plugins");
+}
+
+#if WITH_EDITOR
+
+FText UChaosClothEditorOptions::GetSectionText() const
+{
+	return NSLOCTEXT("ChaosClothEditorPlugin", "ChaosClothEditorSettingsSection", "Chaos Cloth Editor");
+}
+
+#endif
diff --git a/Engine/Plugins/ChaosClothAssetEditor/Source/ChaosClothAssetEditor/Private/ChaosClothAsset/ClothEditorOptions.h b/Engine/Plugins/ChaosClothAssetEditor/Source/ChaosClothAssetEditor/Private/ChaosClothAsset/ClothEditorOptions.h
--- a/Engine/Plugins/ChaosClothAssetEditor/Source/ChaosClothAssetEditor/Private/ChaosClothAsset/ClothEditorOptions.h
+++ b/Engine/Plugins/ChaosClothAssetEditor/Source/ChaosClothAssetEditor/Private/ChaosClothAsset/ClothEditorOptions.h
@@ -2,18 +2,44 @@
 
 #pragma once
 
-#include "UObject/Object.h"
+#include "Engine/DeveloperSettings.h"
 #include "ClothEditorOptions.generated.h"
 
+UENUM(BlueprintType)
+enum class EConstructionViewportMousePanButton : uint8
+{
+	/** Right Mouse Button */
+	Right = 0,
+
+	/** Middle Mouse Button */
+	Middle = 1,
+
+	/** Either Right or Middle Mouse Buttons */
+	RightOrMiddle = 2,
+};
+
+
 /** A settings class used to store and retreive user settings related to the Cloth Editor */
-UCLASS(hidecategories = Object, config = EditorPerProjectUserSettings, MinimalAPI)
-class UChaosClothEditorOptions : public UObject
+UCLASS(HideCategories = Object, Config = EditorPerProjectUserSettings, Meta = (DisplayName = "Chaos Cloth Editor"), MinimalAPI)
+class UChaosClothEditorOptions : public UDeveloperSettings
 {
 	GENERATED_UCLASS_BODY()
 
 public:
 
 	/** Whether Chaos Cloth Assets open in the Dataflow Editor (vs in the Cloth Editor) */
-	UPROPERTY(config)
+	UPROPERTY(Config)
 	bool bClothAssetsOpenInDataflowEditor;
+
+	/** Which mouse button controls camera panning in the Construction Viewport in 2D view mode */
+	UPROPERTY(Config, EditAnywhere, Category = UI)
+	EConstructionViewportMousePanButton ConstructionViewportMousePanButton;
+
+	 // UDeveloperSettings overrides
+	CHAOSCLOTHASSETEDITOR_API virtual FName GetCategoryName() const override;
+
+#if WITH_EDITOR
+	CHAOSCLOTHASSETEDITOR_API virtual FText GetSectionText() const override;
+#endif
+
 };
diff --git a/Engine/Plugins/ChaosClothAssetEditor/Source/ChaosClothAssetEditor/Private/ChaosClothAsset/ClothEditorRestSpaceViewportClient.cpp b/Engine/Plugins/ChaosClothAssetEditor/Source/ChaosClothAssetEditor/Private/ChaosClothAsset/ClothEditorRestSpaceViewportClient.cpp
--- a/Engine/Plugins/ChaosClothAssetEditor/Source/ChaosClothAssetEditor/Private/ChaosClothAsset/ClothEditorRestSpaceViewportClient.cpp
+++ b/Engine/Plugins/ChaosClothAssetEditor/Source/ChaosClothAssetEditor/Private/ChaosClothAsset/ClothEditorRestSpaceViewportClient.cpp
@@ -2,6 +2,7 @@
 
 #include "ChaosClothAsset/ClothEditorRestSpaceViewportClient.h"
 
+#include "ChaosClothAsset/ClothEditorOptions.h"
 #include "BaseBehaviors/ClickDragBehavior.h"
 #include "BaseBehaviors/MouseWheelBehavior.h"
 #include "EngineUtils.h"
@@ -34,11 +35,32 @@ FChaosClothEditorRestSpaceViewportClient::FChaosClothEditorRestSpaceViewportClie
 	constexpr int ViewportBehaviorPriority = 150;
 
 	ScrollBehaviorTarget = MakeUnique<FEditor2DScrollBehaviorTarget>(this);
-	UClickDragInputBehavior* const ScrollBehavior = NewObject<UClickDragInputBehavior>();
-	ScrollBehavior->Initialize(ScrollBehaviorTarget.Get());
-	ScrollBehavior->SetDefaultPriority(ViewportBehaviorPriority);
-	ScrollBehavior->SetUseRightMouseButton();
-	BehaviorsFor2DMode.Add(ScrollBehavior);
+
+	bool bUseRightMouseButton = true;
+	bool bUseMiddleMouseButton = true;
+	if (const UChaosClothEditorOptions* const Options = GetDefault<UChaosClothEditorOptions>())
+	{
+		bUseRightMouseButton = (Options->ConstructionViewportMousePanButton == EConstructionViewportMousePanButton::Right || Options->ConstructionViewportMousePanButton == EConstructionViewportMousePanButton::RightOrMiddle);
+		bUseMiddleMouseButton = (Options->ConstructionViewportMousePanButton == EConstructionViewportMousePanButton::Middle || Options->ConstructionViewportMousePanButton == EConstructionViewportMousePanButton::RightOrMiddle);
+	}
+
+	if (bUseRightMouseButton)
+	{
+		UClickDragInputBehavior* const RightMouseClickDragInputBehavior = NewObject<UClickDragInputBehavior>();
+		RightMouseClickDragInputBehavior->Initialize(ScrollBehaviorTarget.Get());
+		RightMouseClickDragInputBehavior->SetDefaultPriority(ViewportBehaviorPriority);
+		RightMouseClickDragInputBehavior->SetUseRightMouseButton();
+		BehaviorsFor2DMode.Add(RightMouseClickDragInputBehavior);
+	}
+
+	if (bUseMiddleMouseButton)
+	{
+		UClickDragInputBehavior* const MiddleMouseClickDragInputBehavior = NewObject<UClickDragInputBehavior>();
+		MiddleMouseClickDragInputBehavior->Initialize(ScrollBehaviorTarget.Get());
+		MiddleMouseClickDragInputBehavior->SetDefaultPriority(ViewportBehaviorPriority);
+		MiddleMouseClickDragInputBehavior->SetUseMiddleMouseButton();
+		BehaviorsFor2DMode.Add(MiddleMouseClickDragInputBehavior);
+	}
 
 	ZoomBehaviorTarget = MakeUnique<FEditor2DMouseWheelZoomBehaviorTarget>(this);
 	constexpr float CameraFarPlaneWorldZ = -10.0f;
