Cloth Weight Map Paint Tool: Fix GetCurrentWeightValueUnderBrush function -- it was just using the weight at the closest vertex rather than barycentric interpolation

diff --git a/Engine/Plugins/ChaosClothAssetEditor/Source/ChaosClothAssetEditorTools/Private/ChaosClothAsset/ClothWeightMapPaintTool.cpp b/Engine/Plugins/ChaosClothAssetEditor/Source/ChaosClothAssetEditorTools/Private/ChaosClothAsset/ClothWeightMapPaintTool.cpp
--- a/Engine/Plugins/ChaosClothAssetEditor/Source/ChaosClothAssetEditorTools/Private/ChaosClothAsset/ClothWeightMapPaintTool.cpp
+++ b/Engine/Plugins/ChaosClothAssetEditor/Source/ChaosClothAssetEditorTools/Private/ChaosClothAsset/ClothWeightMapPaintTool.cpp
@@ -50,6 +50,7 @@
 #include "Selection/PolygonSelectionMechanic.h"
 #include "BaseGizmos/BrushStampIndicator.h"
 #include "DynamicSubmesh3.h"
+#include "Intersection/IntrRay3Triangle3.h"
 
 #include UE_INLINE_GENERATED_CPP_BY_NAME(ClothWeightMapPaintTool)
 
@@ -1690,17 +1691,19 @@ void UClothEditorWeightMapPaintTool::ApplyVisibilityFilter(const TArray<int32>&
 
 int32 UClothEditorWeightMapPaintTool::FindHitSculptMeshTriangle(const FRay3d& LocalRay)
 {
-	// TODO: Figure out what the actual position on the triangle is when hit.
-	CurrentBaryCentricCoords = FVector3d(1.0 / 3.0, 1.0 / 3.0, 1.0 / 3.0);
+	CurrentBaryCentricCoords = FVector3d::Zero();
 
 	if (!IsInBrushSubMode())
 	{
 		return IndexConstants::InvalidID;
 	}
 
+	int32 HitTID;
+	const FDynamicMesh3* const Mesh = GetSculptMesh();
+
 	if (FilterProperties->bHitBackFaces)
 	{
-		const int32 HitTID = Octree.FindNearestHitObject(LocalRay,
+		HitTID = Octree.FindNearestHitObject(LocalRay,
 			[this](int TriangleID)
 			{ 
 				if (HiddenTriangles.Contains(TriangleID))
@@ -1709,16 +1712,13 @@ int32 UClothEditorWeightMapPaintTool::FindHitSculptMeshTriangle(const FRay3d& Lo
 				}
 				return true;
 			});
-		return HitTID;
 	}
 	else
 	{
-		FDynamicMesh3* Mesh = GetSculptMesh();
-
 		FViewCameraState StateOut;
 		GetToolManager()->GetContextQueriesAPI()->GetCurrentViewState(StateOut);
 		FVector3d LocalEyePosition(CurTargetTransform.InverseTransformPosition((FVector3d)StateOut.Position));
-		const int32 HitTID = Octree.FindNearestHitObject(LocalRay,
+		HitTID = Octree.FindNearestHitObject(LocalRay,
 			[this, Mesh, &LocalEyePosition](int TriangleID) 
 			{
 				if (HiddenTriangles.Contains(TriangleID))
@@ -1730,8 +1730,18 @@ int32 UClothEditorWeightMapPaintTool::FindHitSculptMeshTriangle(const FRay3d& Lo
 				Mesh->GetTriInfo(TriangleID, Normal, Area, Centroid);
 				return Normal.Dot((Centroid - LocalEyePosition)) < 0;
 			});
-		return HitTID;
 	}
+
+	if (Mesh->IsTriangle(HitTID))
+	{
+		FTriangle3d Triangle;
+		Mesh->GetTriVertices(HitTID, Triangle.V[0], Triangle.V[1], Triangle.V[2]);
+		FIntrRay3Triangle3d Query(LocalRay, Triangle);
+		Query.Find();
+		CurrentBaryCentricCoords = Query.TriangleBaryCoords;
+	}
+
+	return HitTID;
 }
 
 int32 UClothEditorWeightMapPaintTool::FindHitTargetMeshTriangle(const FRay3d& LocalRay)
@@ -1908,6 +1918,11 @@ void UClothEditorWeightMapPaintTool::OnTick(float DeltaTime)
 	{
 		if (GetSculptMesh()->IsTriangle(GetBrushTriangleID()))
 		{
+			if (FilterProperties->ValueAtBrushQueryType == EClothEditorWeightMapPaintQueryType::NearestVertexAccurate)
+			{
+				UpdateROI(HoverStamp);
+			}
+
 			const double HitWeightValue = GetCurrentWeightValueUnderBrush();
 
 			if (bShouldPickWeight)
@@ -2385,11 +2400,54 @@ double UClothEditorWeightMapPaintTool::GetCurrentWeightValue(int32 VertexId) con
 
 double UClothEditorWeightMapPaintTool::GetCurrentWeightValueUnderBrush() const
 {
-	float WeightValue = 0.0;
-	int32 VertexID = GetBrushNearestVertex();
-	if (ActiveWeightMap && VertexID != IndexConstants::InvalidID)
+	if (!ActiveWeightMap)
 	{
-		ActiveWeightMap->GetValue(VertexID, &WeightValue);
+		return -1.0;
+	}
+
+	float WeightValue = -1.0f;
+
+	switch (FilterProperties->ValueAtBrushQueryType)
+	{
+
+		case EClothEditorWeightMapPaintQueryType::Interpolated:
+		{
+			const int32 Tid = GetBrushTriangleID();
+			if (Tid != IndexConstants::InvalidID)
+			{
+				const FDynamicMesh3* Mesh = GetSculptMesh();
+				const FIndex3i Vertices = Mesh->GetTriangle(Tid);
+				WeightValue = 0.0f;
+				for (int32 TriangleVertexIndex = 0; TriangleVertexIndex < 3; ++TriangleVertexIndex)
+				{
+					float VertexWeight;
+					ActiveWeightMap->GetValue(Vertices[TriangleVertexIndex], &VertexWeight);
+					WeightValue += CurrentBaryCentricCoords[TriangleVertexIndex] * VertexWeight;
+				}
+			}
+		}
+		break;
+
+		case EClothEditorWeightMapPaintQueryType::NearestVertexFast:
+		{
+			const int32 VertexID = GetBrushNearestVertex();
+			if (VertexID != IndexConstants::InvalidID)
+			{
+				ActiveWeightMap->GetValue(VertexID, &WeightValue);
+			}
+		}
+		break;
+
+		case EClothEditorWeightMapPaintQueryType::NearestVertexAccurate:
+		{
+			const int32 VertexID = GetBrushNearestVertexAccurate();
+			if (VertexID != IndexConstants::InvalidID)
+			{
+				ActiveWeightMap->GetValue(VertexID, &WeightValue);
+			}
+		}
+
+		break;
 	}
 	return WeightValue;
 }
@@ -2425,6 +2483,39 @@ int32 UClothEditorWeightMapPaintTool::GetBrushNearestVertex() const
 	return Vertices[TriangleVertex];
 }
 
+int32 UClothEditorWeightMapPaintTool::GetBrushNearestVertexAccurate() const
+{
+	int32 NearestVertexIndex = INDEX_NONE;
+	const int32 Tid = GetBrushTriangleID();
+
+	if (Tid != IndexConstants::InvalidID)
+	{
+		const FDynamicMesh3* const Mesh = GetSculptMesh();
+
+		FVector3d PointOnSurface(0, 0, 0);
+		const FIndex3i Vertices = Mesh->GetTriangle(Tid);
+		for (int32 TriangleVertexIndex = 0; TriangleVertexIndex < 3; ++TriangleVertexIndex)
+		{
+			PointOnSurface += CurrentBaryCentricCoords[TriangleVertexIndex] * Mesh->GetVertex(Vertices[TriangleVertexIndex]);
+		}
+
+		double MinDist = UE_BIG_NUMBER;
+		for (const int32 VertexIndex : VertexROI)
+		{
+			const FVector3d& VertexPosition = Mesh->GetVertex(VertexIndex);
+			const double CurrDist = FVector3d::Distance(VertexPosition, PointOnSurface);
+			if (CurrDist < MinDist)
+			{
+				MinDist = CurrDist;
+				NearestVertexIndex = VertexIndex;
+			}
+		}
+	}
+
+	return NearestVertexIndex;
+}
+
+
 void UClothEditorWeightMapPaintTool::GetCurrentWeightMap(TArray<float>& OutWeights) const
 {
 	if (ActiveWeightMap)
diff --git a/Engine/Plugins/ChaosClothAssetEditor/Source/ChaosClothAssetEditorTools/Private/ChaosClothAsset/ClothWeightMapPaintTool.h b/Engine/Plugins/ChaosClothAssetEditor/Source/ChaosClothAssetEditorTools/Private/ChaosClothAsset/ClothWeightMapPaintTool.h
--- a/Engine/Plugins/ChaosClothAssetEditor/Source/ChaosClothAssetEditorTools/Private/ChaosClothAsset/ClothWeightMapPaintTool.h
+++ b/Engine/Plugins/ChaosClothAssetEditor/Source/ChaosClothAssetEditorTools/Private/ChaosClothAsset/ClothWeightMapPaintTool.h
@@ -108,6 +108,19 @@ enum class EClothEditorWeightMapPaintVisibilityType : uint8
 	Unoccluded
 };
 
+/** Value Query mode  */
+UENUM()
+enum class EClothEditorWeightMapPaintQueryType : uint8
+{
+	/** Value is interpolated from triangle vertices */
+	Interpolated UMETA(DisplayName = "Interpolated"),
+
+	/** Return the value at the closest vertex of the triangle under the mouse cursor */
+	NearestVertexFast UMETA(DisplayName = "Nearest Vertex (Fast)"),
+
+	/** Return the value of the nearest vertex inside the brush radius, even if the vertex is not on the triangle under the mouse cursor */
+	NearestVertexAccurate UMETA(DisplayName = "Nearest Vertex (Accurate)")
+};
 
 
 // TODO: Look at EditConditions for all these properties. Which ones make sense for which SubToolType?
@@ -195,6 +208,11 @@ public:
 		HideEditConditionToggle, EditConditionHides, EditCondition = "SubToolType == EClothEditorWeightMapPaintInteractionType::Brush || SubToolType == EClothEditorWeightMapPaintInteractionType::Fill"))
 	double ValueAtBrush = 0;
 
+	/** How to compute the value under brush indicator */
+	UPROPERTY(EditAnywhere, Transient, Category = Query, meta = (HideEditConditionToggle, EditConditionHides, 
+		EditCondition = "SubToolType == EClothEditorWeightMapPaintInteractionType::Brush || SubToolType == EClothEditorWeightMapPaintInteractionType::Fill"))
+	EClothEditorWeightMapPaintQueryType ValueAtBrushQueryType = EClothEditorWeightMapPaintQueryType::Interpolated;
+
 };
 
 
@@ -508,6 +526,7 @@ protected:
 	double GetCurrentWeightValueUnderBrush() const;
 	FVector3d CurrentBaryCentricCoords;
 	int32 GetBrushNearestVertex() const;
+	int32 GetBrushNearestVertexAccurate() const;
 
 	void GetCurrentWeightMap(TArray<float>& OutWeights) const;
 
