Register ChaosClothAsset IComponentAssetBroker so editor tools can spawn Chaos Cloth Assets.

diff --git a/Engine/Plugins/ChaosClothAssetEditor/Source/ChaosClothAssetEditor/Private/ChaosClothAsset/ChaosClothAssetEditorModule.cpp b/Engine/Plugins/ChaosClothAssetEditor/Source/ChaosClothAssetEditor/Private/ChaosClothAsset/ChaosClothAssetEditorModule.cpp
--- a/Engine/Plugins/ChaosClothAssetEditor/Source/ChaosClothAssetEditor/Private/ChaosClothAsset/ChaosClothAssetEditorModule.cpp
+++ b/Engine/Plugins/ChaosClothAssetEditor/Source/ChaosClothAssetEditor/Private/ChaosClothAsset/ChaosClothAssetEditorModule.cpp
@@ -3,6 +3,7 @@
 #include "ChaosClothAsset/ChaosClothAssetEditorModule.h"
 #include "ChaosClothAsset/AssetDefinition_ClothAsset.h"
 #include "ChaosClothAsset/ClothAsset.h"
+#include "ChaosClothAsset/ClothComponent.h"
 #include "ChaosClothAsset/ClothDataflowConstructionVisualization.h"
 #include "ChaosClothAsset/ClothDataflowSimulationVisualization.h"
 #include "ChaosClothAsset/ClothSimulationNodeDetailExtender.h"
@@ -16,12 +17,44 @@
 #include "ContentBrowserMenuContexts.h"
 #include "Dataflow/DataflowEditorModeUILayer.h"
 #include "Dataflow/DataflowEditor.h"
+#include "ComponentAssetBroker.h"
 #include "ToolMenus.h"
 
 #define LOCTEXT_NAMESPACE "ChaosClothAssetEditorModule"
 
 namespace UE::Chaos::ClothAsset
 {
+	struct FClothAssetComponentBroker : public IComponentAssetBroker
+	{
+	public:
+		UClass* GetSupportedAssetClass() override
+		{
+			return UChaosClothAsset::StaticClass();
+		}
+
+		virtual bool AssignAssetToComponent(UActorComponent* InComponent, UObject* InAsset) override
+		{
+			if (UChaosClothComponent* ClothComponent = Cast<UChaosClothComponent>(InComponent))
+			{
+				if (UChaosClothAsset* ClothAsset = Cast<UChaosClothAsset>(InAsset))
+				{
+					ClothComponent->SetClothAsset(ClothAsset);
+					return true;
+				}
+			}
+			return false;
+		}
+
+		virtual UObject* GetAssetFromComponent(UActorComponent* InComponent) override
+		{
+			if (UChaosClothComponent* ClothComponent = Cast<UChaosClothComponent>(InComponent))
+			{
+				return ClothComponent->GetClothAsset();
+			}
+			return nullptr;
+		}
+	};
+
 	void FChaosClothAssetEditorModule::RegisterMenus()
 	{
 		// Allows cleanup when module unloads.
@@ -118,10 +151,18 @@ PRAGMA_ENABLE_DEPRECATION_WARNINGS
 				});
 			}
 		}
+
+		ClothAssetComponentBroker = MakeShareable(new FClothAssetComponentBroker);
+		FComponentAssetBrokerage::RegisterBroker(ClothAssetComponentBroker, UChaosClothComponent::StaticClass(), true, true);
 	}
 
 	void FChaosClothAssetEditorModule::ShutdownModule()
 	{
+		if (UObjectInitialized())
+		{
+			FComponentAssetBrokerage::UnregisterBroker(ClothAssetComponentBroker);
+		}
+
 		if (IConsoleVariable* const Var = IConsoleManager::Get().FindConsoleVariable(TEXT("p.ChaosCloth.EnableDataflowEditor")))
 		{
 			Var->OnChangedDelegate().Remove(OnCVarChangedDelegateHandle);
diff --git a/Engine/Plugins/ChaosClothAssetEditor/Source/ChaosClothAssetEditor/Private/ChaosClothAsset/ChaosClothAssetEditorModule.h b/Engine/Plugins/ChaosClothAssetEditor/Source/ChaosClothAssetEditor/Private/ChaosClothAsset/ChaosClothAssetEditorModule.h
--- a/Engine/Plugins/ChaosClothAssetEditor/Source/ChaosClothAssetEditor/Private/ChaosClothAsset/ChaosClothAssetEditorModule.h
+++ b/Engine/Plugins/ChaosClothAssetEditor/Source/ChaosClothAssetEditor/Private/ChaosClothAsset/ChaosClothAssetEditorModule.h
@@ -6,6 +6,8 @@
 
 namespace UE::Chaos::ClothAsset
 {
+	struct FClothAssetComponentBroker;
+
 	class FChaosClothAssetEditorModule final : public FBaseCharacterFXEditorModule
 	{
 	public:
@@ -20,5 +22,6 @@ namespace UE::Chaos::ClothAsset
 
 		void RegisterMenus();
 		
+		TSharedPtr<FClothAssetComponentBroker> ClothAssetComponentBroker;
 	};
 } // namespace UE::Chaos::ClothAsset
