Cloth Weight Map Paint Tool: change the default ValueAtBrushQueryType to NearestVertexFast

diff --git a/Engine/Plugins/ChaosClothAssetEditor/Source/ChaosClothAssetEditorTools/Private/ChaosClothAsset/ClothWeightMapPaintTool.h b/Engine/Plugins/ChaosClothAssetEditor/Source/ChaosClothAssetEditorTools/Private/ChaosClothAsset/ClothWeightMapPaintTool.h
--- a/Engine/Plugins/ChaosClothAssetEditor/Source/ChaosClothAssetEditorTools/Private/ChaosClothAsset/ClothWeightMapPaintTool.h
+++ b/Engine/Plugins/ChaosClothAssetEditor/Source/ChaosClothAssetEditorTools/Private/ChaosClothAsset/ClothWeightMapPaintTool.h
@@ -211,7 +211,7 @@ public:
 	/** How to compute the value under brush indicator */
 	UPROPERTY(EditAnywhere, Transient, Category = Query, meta = (HideEditConditionToggle, EditConditionHides, 
 		EditCondition = "SubToolType == EClothEditorWeightMapPaintInteractionType::Brush || SubToolType == EClothEditorWeightMapPaintInteractionType::Fill"))
-	EClothEditorWeightMapPaintQueryType ValueAtBrushQueryType = EClothEditorWeightMapPaintQueryType::Interpolated;
+	EClothEditorWeightMapPaintQueryType ValueAtBrushQueryType = EClothEditorWeightMapPaintQueryType::NearestVertexFast;
 
 };
 
