Dataflow Editor: when starting the DF Editor with a Cloth Asset, check that the Cloth has a Dataflow Assset. If it does not, start a simple asset editor (essentially just a details panel)

diff --git a/Engine/Plugins/ChaosClothAssetEditor/Source/ChaosClothAssetEditor/Private/ChaosClothAsset/AssetDefinition_ClothAsset.cpp b/Engine/Plugins/ChaosClothAssetEditor/Source/ChaosClothAssetEditor/Private/ChaosClothAsset/AssetDefinition_ClothAsset.cpp
--- a/Engine/Plugins/ChaosClothAssetEditor/Source/ChaosClothAssetEditor/Private/ChaosClothAsset/AssetDefinition_ClothAsset.cpp
+++ b/Engine/Plugins/ChaosClothAssetEditor/Source/ChaosClothAssetEditor/Private/ChaosClothAsset/AssetDefinition_ClothAsset.cpp
@@ -17,6 +17,7 @@
 #include "ChaosClothAsset/TerminalNode.h"
 #include "Dataflow/DataflowEditor.h"
 #include "Dataflow/DataflowSNode.h"
+#include "Toolkits/SimpleAssetEditor.h"
 
 #define LOCTEXT_NAMESPACE "AssetDefinition_ClothAsset"
 
@@ -189,10 +190,20 @@ EAssetCommandResult UAssetDefinition_ClothAsset::OpenAssets(const FAssetOpenArgs
 		}
 		else
 		{
-			UDataflowEditor* const AssetEditor = NewObject<UDataflowEditor>(AssetEditorSubsystem, NAME_None, RF_Transient);
-            const TSubclassOf<AActor> ActorClass = StaticLoadClass(AActor::StaticClass(), nullptr,
-            	TEXT("/ChaosClothAssetEditor/BP_ClothPreview.BP_ClothPreview_C"), nullptr, LOAD_None, nullptr);
-			AssetEditor->Initialize({ ClothObjects[0] }, ActorClass);
+			// Check if the cloth asset has a Dataflow asset
+			UChaosClothAsset* const ClothAsset = CastChecked<UChaosClothAsset>(ClothObjects[0]);
+			if (!ClothAsset->GetDataflow())
+			{
+				// No Dataflow asset, just open a basic properties panel
+				FSimpleAssetEditor::CreateEditor(EToolkitMode::Standalone, OpenArgs.ToolkitHost, ClothObjects[0] );
+			}
+			else
+			{
+				UDataflowEditor* const AssetEditor = NewObject<UDataflowEditor>(AssetEditorSubsystem, NAME_None, RF_Transient);
+				const TSubclassOf<AActor> ActorClass = StaticLoadClass(AActor::StaticClass(), nullptr,
+					TEXT("/ChaosClothAssetEditor/BP_ClothPreview.BP_ClothPreview_C"), nullptr, LOAD_None, nullptr);
+				AssetEditor->Initialize({ ClothObjects[0] }, ActorClass);
+			}
 		}
 
 		return EAssetCommandResult::Handled;
