Chaos Cloth Asset - Fixed ensure when selecting multiple cloth nodes with connectable value customization.

#rnx

diff --git a/Engine/Plugins/ChaosClothAssetEditor/Source/ChaosClothAssetDataflowNodes/Private/ChaosClothAsset/ConnectableValueCustomization.cpp b/Engine/Plugins/ChaosClothAssetEditor/Source/ChaosClothAssetDataflowNodes/Private/ChaosClothAsset/ConnectableValueCustomization.cpp
--- a/Engine/Plugins/ChaosClothAssetEditor/Source/ChaosClothAssetDataflowNodes/Private/ChaosClothAsset/ConnectableValueCustomization.cpp
+++ b/Engine/Plugins/ChaosClothAssetEditor/Source/ChaosClothAssetDataflowNodes/Private/ChaosClothAsset/ConnectableValueCustomization.cpp
@@ -179,7 +179,7 @@ PRAGMA_ENABLE_DEPRECATION_WARNINGS
 							if (ensure(DataflowNode))
 							{
 								void* Data;
-								if (ensure(HandlePtr->GetValueData(Data) == FPropertyAccess::Success))
+								if (HandlePtr->GetValueData(Data) == FPropertyAccess::Success)  // GetValueData could return false with multiple selections 
 								{
 									Text = *static_cast<FString*>(Data);  // Default value if the property isn't an input, or isn't connected
 									if (const FDataflowInput* const DataflowInput = DataflowNode->FindInput(Data))
