Add experimental nodes to convert cloth collections to and from DynamicMesh objects. This is support prototyping resizing with geometry tools on DynamicMeshes.

This will now always set the UVs on the Dynamic Mesh to be the SimPositions2D when requesting a Sim3D dynamic mesh.

Added mapping data between the cloth pattern and dynamic mesh. This continues to use the NonManifoldMappingSupport for vertices, and a new attribute for triangles.

diff --git a/Engine/Plugins/ChaosClothAssetEditor/Source/ChaosClothAssetDataflowNodes/ChaosClothAssetDataflowNodes.Build.cs b/Engine/Plugins/ChaosClothAssetEditor/Source/ChaosClothAssetDataflowNodes/ChaosClothAssetDataflowNodes.Build.cs
--- a/Engine/Plugins/ChaosClothAssetEditor/Source/ChaosClothAssetDataflowNodes/ChaosClothAssetDataflowNodes.Build.cs
+++ b/Engine/Plugins/ChaosClothAssetEditor/Source/ChaosClothAssetDataflowNodes/ChaosClothAssetDataflowNodes.Build.cs
@@ -33,6 +33,7 @@ public class ChaosClothAssetDataflowNodes : ModuleRules
 				"DynamicMesh",
 				"Engine",
 				"GeometryCore",
+				"GeometryFramework",
 				"InputCore",
 				"MeshConversion",
 				"MeshDescription",
diff --git a/Engine/Plugins/ChaosClothAssetEditor/Source/ChaosClothAssetDataflowNodes/Private/ChaosClothAsset/ChaosClothAssetDataflowNodesModule.cpp b/Engine/Plugins/ChaosClothAssetEditor/Source/ChaosClothAssetDataflowNodes/Private/ChaosClothAsset/ChaosClothAssetDataflowNodesModule.cpp
--- a/Engine/Plugins/ChaosClothAssetEditor/Source/ChaosClothAssetDataflowNodes/Private/ChaosClothAsset/ChaosClothAssetDataflowNodesModule.cpp
+++ b/Engine/Plugins/ChaosClothAssetEditor/Source/ChaosClothAssetDataflowNodes/Private/ChaosClothAsset/ChaosClothAssetDataflowNodesModule.cpp
@@ -6,16 +6,16 @@
 #include "ChaosClothAsset/BindToRootBoneNode.h"
 #include "ChaosClothAsset/BlendVerticesNode.h"
 #include "ChaosClothAsset/ClothCollectionGroup.h"
+#include "ChaosClothAsset/ClothCollectionToDynamicMeshNode.h"
 #include "ChaosClothAsset/ColorScheme.h"
 #include "ChaosClothAsset/CopySimulationToRenderMeshNode.h"
-#include "ChaosClothAsset/SelectionGroupCustomization.h"
 #include "ChaosClothAsset/DatasmithImportNode.h"
 #include "ChaosClothAsset/DeleteElementNode.h"
 #include "ChaosClothAsset/ImportFilePathCustomization.h"
-#include "ChaosClothAsset/RemeshNode.h"
 #include "ChaosClothAsset/ImportNode.h"
 #include "ChaosClothAsset/MergeClothCollectionsNode.h"
 #include "ChaosClothAsset/ProxyDeformerNode.h"
+#include "ChaosClothAsset/RemeshNode.h"
 #include "ChaosClothAsset/ReverseNormalsNode.h"
 #include "ChaosClothAsset/SelectionNode.h"
 #include "ChaosClothAsset/SelectionToIntMapNode.h"
@@ -64,11 +64,12 @@
 #include "ChaosClothAsset/WeightMapNode.h"
 #include "ChaosClothAsset/WeightedValueCustomization.h"
 
-#include "ChaosClothAsset/ImportedValueCustomization.h"
-#include "ChaosClothAsset/WeightMapToSelectionNode.h"
+#include "ChaosClothAsset/ClothDataflowViewModes.h"
 #include "ChaosClothAsset/ConnectableValueCustomization.h"
 #include "ChaosClothAsset/ConnectableValue.h"
-#include "ChaosClothAsset/ClothDataflowViewModes.h"
+#include "ChaosClothAsset/ImportedValueCustomization.h"
+#include "ChaosClothAsset/SelectionGroupCustomization.h"
+#include "ChaosClothAsset/WeightMapToSelectionNode.h"
 #include "Dataflow/DataflowCollectionAddScalarVertexPropertyNode.h"
 #include "Dataflow/DataflowNodeColorsRegistry.h"
 #include "Dataflow/DataflowNodeFactory.h"
@@ -87,6 +88,7 @@ namespace UE::Chaos::ClothAsset
 			DATAFLOW_NODE_REGISTER_CREATION_FACTORY(FChaosClothAssetAttributeNode_v2);
 			DATAFLOW_NODE_REGISTER_CREATION_FACTORY(FChaosClothAssetBindToRootBoneNode);
 			DATAFLOW_NODE_REGISTER_CREATION_FACTORY(FChaosClothAssetBlendVerticesNode);
+			DATAFLOW_NODE_REGISTER_CREATION_FACTORY(FChaosClothAssetCollectionToDynamicMeshNode);
 			DATAFLOW_NODE_REGISTER_CREATION_FACTORY(FChaosClothAssetCopySimulationToRenderMeshNode);
 			DATAFLOW_NODE_REGISTER_CREATION_FACTORY(FChaosClothAssetDatasmithImportNode);
 			DATAFLOW_NODE_REGISTER_CREATION_FACTORY(FChaosClothAssetDeleteElementNode);
@@ -137,6 +139,7 @@ namespace UE::Chaos::ClothAsset
 			DATAFLOW_NODE_REGISTER_CREATION_FACTORY(FChaosClothAssetTransferSkinWeightsNode);
 			DATAFLOW_NODE_REGISTER_CREATION_FACTORY(FChaosClothAssetTransformPositionsNode);
 			DATAFLOW_NODE_REGISTER_CREATION_FACTORY(FChaosClothAssetTransformUVsNode);
+			DATAFLOW_NODE_REGISTER_CREATION_FACTORY(FChaosClothAssetUpdateClothFromDynamicMeshNode);
 			DATAFLOW_NODE_REGISTER_CREATION_FACTORY(FChaosClothAssetUSDImportNode_v2);
 			DATAFLOW_NODE_REGISTER_CREATION_FACTORY(FChaosClothAssetWeightMapToSelectionNode);
 			// Deprecated nodes
diff --git a/Engine/Plugins/ChaosClothAssetEditor/Source/ChaosClothAssetDataflowNodes/Private/ChaosClothAsset/ClothCollectionToDynamicMeshNode.cpp b/Engine/Plugins/ChaosClothAssetEditor/Source/ChaosClothAssetDataflowNodes/Private/ChaosClothAsset/ClothCollectionToDynamicMeshNode.cpp
new file mode 100644
--- /dev/null
+++ b/Engine/Plugins/ChaosClothAssetEditor/Source/ChaosClothAssetDataflowNodes/Private/ChaosClothAsset/ClothCollectionToDynamicMeshNode.cpp
@@ -0,0 +1,175 @@
+// Copyright Epic Games, Inc. All Rights Reserved.
+
+#include "ChaosClothAsset/ClothCollectionToDynamicMeshNode.h"
+#include "ChaosClothAsset/ClothPatternToDynamicMesh.h"
+#include "ChaosClothAsset/ClothPatternToDynamicMeshMappingSupport.h"
+#include "ChaosClothAsset/CollectionClothFacade.h"
+#include "Dataflow/DataflowInputOutput.h"
+#include "DynamicMesh/DynamicMesh3.h"
+#include "Materials/MaterialInterface.h"
+#include "UDynamicMesh.h"
+
+#include UE_INLINE_GENERATED_CPP_BY_NAME(ClothCollectionToDynamicMeshNode)
+
+#define LOCTEXT_NAMESPACE "FChaosClothAssetCollectionToDynamicMeshNode"
+
+FChaosClothAssetCollectionToDynamicMeshNode::FChaosClothAssetCollectionToDynamicMeshNode(const UE::Dataflow::FNodeParameters& InParam, FGuid InGuid)
+	: FDataflowNode(InParam, InGuid)
+{
+	RegisterInputConnection(&Collection);
+	RegisterOutputConnection(&SimDynamicMesh);
+	RegisterOutputConnection(&RenderDynamicMesh);
+	RegisterOutputConnection(&RenderMaterials);
+}
+
+void FChaosClothAssetCollectionToDynamicMeshNode::Evaluate(UE::Dataflow::FContext& Context, const FDataflowOutput* Out) const
+{
+	using namespace UE::Geometry;
+	using namespace UE::Chaos::ClothAsset;
+
+	if (Out->IsA(&SimDynamicMesh))
+	{
+		FManagedArrayCollection InCollection = GetValue<FManagedArrayCollection>(Context, &Collection);
+		const TSharedRef<FManagedArrayCollection> ClothCollection = MakeShared<FManagedArrayCollection>(MoveTemp(InCollection));
+		FCollectionClothConstFacade ClothFacade(ClothCollection);
+		if (ClothFacade.IsValid())
+		{
+			TObjectPtr<UDynamicMesh> NewMesh = NewObject<UDynamicMesh>();
+			NewMesh->Reset();
+			FDynamicMesh3& DynamicMesh = NewMesh->GetMeshRef();
+			FClothPatternToDynamicMesh Converter;
+			Converter.Convert(ClothCollection, INDEX_NONE, EClothPatternVertexType::Sim3D, DynamicMesh);
+			SetValue(Context, NewMesh, &SimDynamicMesh);
+			return;
+		}
+
+		SetValue(Context, TObjectPtr<UDynamicMesh>(NewObject<UDynamicMesh>()), &SimDynamicMesh);
+	}
+	if (Out->IsA(&RenderDynamicMesh) || Out->IsA(&RenderMaterials))
+	{
+		FManagedArrayCollection InCollection = GetValue<FManagedArrayCollection>(Context, &Collection);
+		const TSharedRef<FManagedArrayCollection> ClothCollection = MakeShared<FManagedArrayCollection>(MoveTemp(InCollection));
+		FCollectionClothConstFacade ClothFacade(ClothCollection);
+		if (ClothFacade.IsValid())
+		{
+			TObjectPtr<UDynamicMesh> NewMesh = NewObject<UDynamicMesh>();
+			NewMesh->Reset();
+			FDynamicMesh3& DynamicMesh = NewMesh->GetMeshRef();
+			FClothPatternToDynamicMesh Converter;
+			Converter.Convert(ClothCollection, INDEX_NONE, EClothPatternVertexType::Render, DynamicMesh);
+			SetValue(Context, NewMesh, &RenderDynamicMesh);
+
+			TArray<TObjectPtr<UMaterialInterface>> OutMaterials;
+			const TConstArrayView<FString> RenderMaterialPathNames = ClothFacade.GetRenderMaterialPathName();
+			OutMaterials.Reserve(RenderMaterialPathNames.Num());
+			for (int32 RenderPatternIndex = 0; RenderPatternIndex < RenderMaterialPathNames.Num(); ++RenderPatternIndex)
+			{
+				const FString& RenderMaterialPathName = RenderMaterialPathNames[RenderPatternIndex];
+				OutMaterials.Add(LoadObject<UMaterialInterface>((UObject*)GetTransientPackage(), *RenderMaterialPathName, nullptr, LOAD_None, nullptr));
+			}
+			SetValue(Context, OutMaterials, &RenderMaterials);
+			return;
+		}
+
+		SetValue(Context, TObjectPtr<UDynamicMesh>(NewObject<UDynamicMesh>()), &RenderDynamicMesh);
+		SetValue(Context, TArray<TObjectPtr<UMaterialInterface>>(), &RenderMaterials);
+	}
+}
+
+FChaosClothAssetUpdateClothFromDynamicMeshNode::FChaosClothAssetUpdateClothFromDynamicMeshNode(const UE::Dataflow::FNodeParameters& InParam, FGuid InGuid)
+	: FDataflowNode(InParam, InGuid)
+{
+	RegisterInputConnection(&Collection);
+	RegisterInputConnection(&DynamicMesh);
+	RegisterOutputConnection(&Collection, &Collection);
+}
+
+void FChaosClothAssetUpdateClothFromDynamicMeshNode::Evaluate(UE::Dataflow::FContext& Context, const FDataflowOutput* Out) const
+{
+	if (Out->IsA(&Collection))
+	{
+		using namespace UE::Chaos::ClothAsset;
+
+		FManagedArrayCollection InCollection = GetValue<FManagedArrayCollection>(Context, &Collection);
+		const TSharedRef<FManagedArrayCollection> ClothCollection = MakeShared<FManagedArrayCollection>(MoveTemp(InCollection));
+		FCollectionClothFacade ClothFacade(ClothCollection);
+		if (ClothFacade.IsValid())
+		{
+			if (const TObjectPtr<const UDynamicMesh> InMesh = GetValue(Context, &DynamicMesh))
+			{
+				const UE::Geometry::FDynamicMesh3& DynMesh = InMesh->GetMeshRef();
+				FClothPatternToDynamicMeshMappingSupport ClothMapping(DynMesh);
+
+				if (bCopyToRenderPositions)
+				{
+					TArrayView<FVector3f> RenderPositions = ClothFacade.GetRenderPosition();
+					for (int32 VertexID : DynMesh.VertexIndicesItr())
+					{
+						const int32 RenderMeshID = ClothMapping.GetOriginalVertexID(VertexID);
+						if (RenderPositions.IsValidIndex(RenderMeshID))
+						{
+							RenderPositions[RenderMeshID] = FVector3f(DynMesh.GetVertexRef(VertexID));
+						}
+					}
+				}
+				if (bCopyToSim3DPositions)
+				{
+					TArrayView<FVector3f> SimPositions = ClothFacade.GetSimPosition3D();
+					for (int32 VertexID : DynMesh.VertexIndicesItr())
+					{
+						const int32 SimMeshID = ClothMapping.GetOriginalVertexID(VertexID);
+						if (SimPositions.IsValidIndex(SimMeshID))
+						{
+							SimPositions[SimMeshID] = FVector3f(DynMesh.GetVertexRef(VertexID));
+						}
+					}
+				}
+				if (bCopyUVsToSim2DPositions)
+				{
+					if (const UE::Geometry::FDynamicMeshAttributeSet* const AttributeSet = DynMesh.Attributes())
+					{
+						if (const UE::Geometry::FDynamicMeshUVOverlay* const UVOverlay = AttributeSet->GetUVLayer(UVChannelIndex))
+						{
+							TArrayView<FVector2f> SimPositions = ClothFacade.GetSimPosition2D();
+							const int32 NumSimVertices3D = ClothFacade.GetNumSimVertices3D();
+							TConstArrayView<FIntVector> SimIndices2D = ClothFacade.GetSimIndices2D();
+							TConstArrayView<FIntVector> SimIndices3D = ClothFacade.GetSimIndices3D();
+							for (int32 VertexID : DynMesh.VertexIndicesItr())
+							{
+								const int32 SimMeshVertID = ClothMapping.GetOriginalVertexID(VertexID);
+								if (SimMeshVertID >= 0 && SimMeshVertID < NumSimVertices3D)
+								{
+									UVOverlay->EnumerateVertexElements(VertexID,
+										[&SimPositions, &SimIndices2D, &SimIndices3D, &ClothMapping, SimMeshVertID](int32 TriangleID, int32 ElementID, const FVector2f& UVValue)->bool
+										{
+											const int32 SimMeshTriID = ClothMapping.GetOriginalTriangleID(TriangleID);
+											if (SimIndices3D.IsValidIndex(SimMeshTriID))
+											{
+												const FIntVector& Index3D = SimIndices3D[SimMeshTriID];
+												const FIntVector& Index2D = SimIndices2D[SimMeshTriID];
+												for (int32 LocalIndex = 0; LocalIndex < 3; ++LocalIndex)
+												{
+													if (Index3D[LocalIndex] == SimMeshVertID)
+													{
+														if (SimPositions.IsValidIndex(Index2D[LocalIndex]))
+														{
+															SimPositions[Index2D[LocalIndex]] = UVValue;
+														}
+														break;
+													}
+												}
+											}
+											return true;
+										});
+								}
+							}
+						}
+					}
+				}
+			}
+		}
+		SetValue(Context, MoveTemp(*ClothCollection), &Collection);
+	}
+}
+
+#undef LOCTEXT_NAMESPACE
diff --git a/Engine/Plugins/ChaosClothAssetEditor/Source/ChaosClothAssetDataflowNodes/Public/ChaosClothAsset/ClothCollectionToDynamicMeshNode.h b/Engine/Plugins/ChaosClothAssetEditor/Source/ChaosClothAssetDataflowNodes/Public/ChaosClothAsset/ClothCollectionToDynamicMeshNode.h
new file mode 100644
--- /dev/null
+++ b/Engine/Plugins/ChaosClothAssetEditor/Source/ChaosClothAssetDataflowNodes/Public/ChaosClothAsset/ClothCollectionToDynamicMeshNode.h
@@ -0,0 +1,89 @@
+// Copyright Epic Games, Inc. All Rights Reserved.
+
+#pragma once 
+
+#include "Dataflow/DataflowNode.h"
+#include "GeometryCollection/ManagedArrayCollection.h"
+#include "ClothCollectionToDynamicMeshNode.generated.h"
+
+
+class UDynamicMesh;
+class UMaterialInterface;
+
+/** Convert a Cloth Collection mesh to a dynamic mesh.
+ */
+USTRUCT(Meta = (DataflowCloth, Experimental))
+struct FChaosClothAssetCollectionToDynamicMeshNode : public FDataflowNode
+{
+	GENERATED_USTRUCT_BODY()
+	DATAFLOW_NODE_DEFINE_INTERNAL(FChaosClothAssetCollectionToDynamicMeshNode, "ClothCollectionToDynamicMesh", "Cloth", "Cloth Collection Dynamic Mesh")
+	DATAFLOW_NODE_RENDER_TYPE_START()
+	DATAFLOW_NODE_RENDER_TYPE_ADD("SurfaceRender", FName("FDynamicMesh3"), {"SimDynamicMesh"})
+	DATAFLOW_NODE_RENDER_TYPE_ADD("SurfaceRender", FName("FDynamicMesh3"), {"RenderDynamicMesh", "RenderMaterials"})
+	DATAFLOW_NODE_RENDER_TYPE_END()
+
+public:
+
+	FChaosClothAssetCollectionToDynamicMeshNode(const UE::Dataflow::FNodeParameters& InParam, FGuid InGuid = FGuid::NewGuid());
+
+	virtual void Evaluate(UE::Dataflow::FContext& Context, const FDataflowOutput* Out) const override;
+
+private:
+
+	UPROPERTY(Meta = (DataflowInput))
+	FManagedArrayCollection Collection;
+
+	/** Output sim collection dynamic mesh. */
+	UPROPERTY(meta = (DataflowOutput))
+	TObjectPtr<UDynamicMesh> SimDynamicMesh;
+
+	/** Output render collection dynamic mesh. */
+	UPROPERTY(meta = (DataflowOutput))
+	TObjectPtr<UDynamicMesh> RenderDynamicMesh;
+
+	/** Render materials.*/
+	UPROPERTY(meta = (DataflowOutput))
+	TArray<TObjectPtr<UMaterialInterface>> RenderMaterials;
+
+};
+
+
+/** Update cloth collection attributes from a DynamicMesh
+ */
+USTRUCT(Meta = (DataflowCloth, Experimental))
+struct FChaosClothAssetUpdateClothFromDynamicMeshNode : public FDataflowNode
+{
+	GENERATED_USTRUCT_BODY()
+	DATAFLOW_NODE_DEFINE_INTERNAL(FChaosClothAssetUpdateClothFromDynamicMeshNode, "UpdateClothFromDynamicMesh", "Cloth", "Update Cloth Collection Dynamic Mesh Cloth")
+	DATAFLOW_NODE_RENDER_TYPE("SurfaceRender", FName("FClothCollection"), "Collection")
+
+public:
+
+	FChaosClothAssetUpdateClothFromDynamicMeshNode(const UE::Dataflow::FNodeParameters& InParam, FGuid InGuid = FGuid::NewGuid());
+
+	virtual void Evaluate(UE::Dataflow::FContext& Context, const FDataflowOutput* Out) const override;
+
+private:
+
+	UPROPERTY(Meta = (DataflowInput, DataflowOutput, DataflowPassthrough = "Collection"))
+	FManagedArrayCollection Collection;
+
+	UPROPERTY(Meta = (DataflowInput))
+	TObjectPtr<UDynamicMesh> DynamicMesh;
+
+	/** Copy DynamicMesh Vertex Positions to Render Positions */
+	UPROPERTY(EditAnywhere, Category = "DynamicMesh To Cloth Collection")
+	bool bCopyToRenderPositions = false;
+
+	/** Copy DynamicMesh Vertex Positions to Sim3D Positions */
+	UPROPERTY(EditAnywhere, Category = "DynamicMesh To Cloth Collection")
+	bool bCopyToSim3DPositions = false;
+
+	/** Copy DynamicMesh UVs to Sim2D Positions */
+	UPROPERTY(EditAnywhere, Category = "DynamicMesh To Cloth Collection")
+	bool bCopyUVsToSim2DPositions = false;
+
+	/** Which UV Channel to use at Sim2D Positions */
+	UPROPERTY(EditAnywhere, Category = "DynamicMesh To Cloth Collection", meta = (EditCondition = "bCopyUVsToSim2DPositions", EditConditionHides))
+	int32 UVChannelIndex = 0;
+};
diff --git a/Engine/Plugins/ChaosClothAssetEditor/Source/ChaosClothAssetTools/Private/ChaosClothAsset/ClothPatternToDynamicMesh.cpp b/Engine/Plugins/ChaosClothAssetEditor/Source/ChaosClothAssetTools/Private/ChaosClothAsset/ClothPatternToDynamicMesh.cpp
--- a/Engine/Plugins/ChaosClothAssetEditor/Source/ChaosClothAssetTools/Private/ChaosClothAsset/ClothPatternToDynamicMesh.cpp
+++ b/Engine/Plugins/ChaosClothAssetEditor/Source/ChaosClothAssetTools/Private/ChaosClothAsset/ClothPatternToDynamicMesh.cpp
@@ -3,12 +3,12 @@
 #if WITH_EDITOR
 
 #include "ChaosClothAsset/ClothPatternToDynamicMesh.h"
+#include "ChaosClothAsset/ClothPatternToDynamicMeshMappingSupport.h"
 #include "ChaosClothAsset/ClothAsset.h"
 #include "ChaosClothAsset/ClothComponent.h"
 #include "ChaosClothAsset/CollectionClothFacade.h"
 #include "ChaosClothAsset/ClothCollectionGroup.h"
 #include "DynamicMesh/MeshNormals.h"
-#include "DynamicMesh/NonManifoldMappingSupport.h"
 #include "Engine/SkeletalMesh.h"
 #include "SkeletalMeshAttributes.h"
 #include "ToDynamicMesh.h"
@@ -120,6 +120,7 @@ public:
 						TriIDs.Add(TriIndex);
 					}
 				}
+				NumTexCoords = 1; // SimPositions2D will be stored as UVs
 			}break;
 			default:
 				checkNoEntry();
@@ -210,6 +211,7 @@ public:
 						TriIDs.Add(TriIndex + FaceOffset);
 					}
 				}
+				NumTexCoords = 1; // SimPositions2D will be stored as UVs
 			} break;
 			default:
 				checkNoEntry();
@@ -378,13 +380,18 @@ public:
 	// -- Access to per-wedge attributes -- //
 	void GetWedgeIDs(const TriIDType& TriID, WedgeIDType& WID0, WedgeIDType& WID1, WedgeIDType& WID2) const
 	{
-		checkf(false, TEXT("FClothPatternWrapper: ClothPatterns are not expected to use Wedges"));
+		check(VertexDataType == EClothPatternVertexType::Sim3D);
+		const FIntVector& Face = Cloth.GetSimIndices2D()[TriID];
+
+		WID0 = Face[0];
+		WID1 = Face[1];
+		WID2 = Face[2];
 	}
 
 	FVector2f GetWedgeUV(int32 UVLayerIndex, WedgeIDType WID) const
 	{
-		checkf(false, TEXT("FClothPatternWrapper: ClothPatterns are not expected to use Wedges"));
-		return FVector2f();
+		check(VertexDataType == EClothPatternVertexType::Sim3D);
+		return Cloth.GetSimPosition2D()[WID];
 	}
 
 	FVector3f GetWedgeNormal(WedgeIDType WID) const
@@ -675,9 +682,27 @@ void FClothPatternToDynamicMesh::Convert(const TSharedRef<const FManagedArrayCol
 	if (!bVertexIndicesMatch)
 	{
 		MeshOut.EnableAttributes();
-		UE::Geometry::FNonManifoldMappingSupport::AttachNonManifoldVertexMappingData(PatternToDynamicMesh.ToSrcVertIDMap, MeshOut);
+		FClothPatternToDynamicMeshMappingSupport::AttachVertexMappingData(PatternToDynamicMesh.ToSrcVertIDMap, MeshOut);
+	}
+	const int32 NumTriangles = VertexDataType == EClothPatternVertexType::Render ? ClothFacade.GetNumRenderFaces() : ClothFacade.GetNumSimFaces();
+	bool bTriangleIndicesMatch = false;
+	if (PatternToDynamicMesh.ToSrcTriIDMap.Num() == NumTriangles)
+	{
+		bTriangleIndicesMatch = true;
+		for (int32 TriId = 0; TriId < NumTriangles; ++TriId)
+		{
+			if (PatternToDynamicMesh.ToSrcTriIDMap[TriId] != TriId)
+			{
+				bTriangleIndicesMatch = false;
+				break;
+			}
+		}
+	}
+	if (!bTriangleIndicesMatch)
+	{
+		MeshOut.EnableAttributes();
+		FClothPatternToDynamicMeshMappingSupport::AttachTriangleMappingData(PatternToDynamicMesh.ToSrcTriIDMap, MeshOut);
 	}
-
 }
 
 void FClothPatternToDynamicMesh::Convert(const UChaosClothAsset* ClothAssetMeshIn, int32 LODIndex, int32 PatternIndex, EClothPatternVertexType VertexDataType, UE::Geometry::FDynamicMesh3& MeshOut)
diff --git a/Engine/Plugins/ChaosClothAssetEditor/Source/ChaosClothAssetTools/Private/ChaosClothAsset/ClothPatternToDynamicMeshMappingSupport.cpp b/Engine/Plugins/ChaosClothAssetEditor/Source/ChaosClothAssetTools/Private/ChaosClothAsset/ClothPatternToDynamicMeshMappingSupport.cpp
new file mode 100644
--- /dev/null
+++ b/Engine/Plugins/ChaosClothAssetEditor/Source/ChaosClothAssetTools/Private/ChaosClothAsset/ClothPatternToDynamicMeshMappingSupport.cpp
@@ -0,0 +1,341 @@
+// Copyright Epic Games, Inc. All Rights Reserved. 
+
+#include "ChaosClothAsset/ClothPatternToDynamicMeshMappingSupport.h"
+#include "DynamicMesh/DynamicMeshAttributeSet.h"
+
+namespace UE::Chaos::ClothAsset
+{
+
+class FDynamicMeshSourceTriangleIdAttributeChange final :
+	public UE::Geometry::FDynamicMeshAttributeChangeBase
+{
+	friend class FDynamicMeshSourceTriangleIdAttribute;
+	TArray<TPair<int32, int32>> OldSourceTriangleIds, NewSourceTriangleIds;
+
+	bool bOldInvalidState = false;
+	bool bNewInvalidState = false;
+
+public:
+	FDynamicMeshSourceTriangleIdAttributeChange() = default;
+
+	virtual ~FDynamicMeshSourceTriangleIdAttributeChange() override = default;
+
+	void SaveInitialTriangle(const UE::Geometry::FDynamicMeshAttributeBase* Attribute, int TriangleID) override;
+
+	void StoreAllFinalTriangles(const UE::Geometry::FDynamicMeshAttributeBase* Attribute, const TArray<int>& TriangleIDs) override;
+
+	bool Apply(UE::Geometry::FDynamicMeshAttributeBase* Attribute, bool bRevert) const override;
+};
+
+class FDynamicMeshSourceTriangleIdAttribute final :
+	public UE::Geometry::FDynamicMeshAttributeBase
+{
+	friend class FDynamicMeshSourceTriangleIdAttributeChange;
+
+	UE::Geometry::FDynamicMesh3* Parent = nullptr;
+	UE::Geometry::TDynamicVector<int32> SourceTriangleIds;
+	bool bValid = false;
+
+public:
+	FDynamicMeshSourceTriangleIdAttribute() = default;
+
+	FDynamicMeshSourceTriangleIdAttribute(UE::Geometry::FDynamicMesh3* InParent) : Parent(InParent)
+	{
+	}
+
+	virtual ~FDynamicMeshSourceTriangleIdAttribute() override = default;
+
+	const UE::Geometry::FDynamicMesh3* GetParent() const
+	{
+		return Parent;
+	}
+	UE::Geometry::FDynamicMesh3* GetParent()
+	{
+		return Parent;
+	}
+
+	bool IsValid() const
+	{
+		return bValid;
+	}
+
+	UE::Geometry::FDynamicMeshAttributeBase* MakeCopy(UE::Geometry::FDynamicMesh3* ParentIn) const override
+	{
+		FDynamicMeshSourceTriangleIdAttribute* Attribute = new FDynamicMeshSourceTriangleIdAttribute(ParentIn);
+		Attribute->Copy(*this);
+		return Attribute;
+	}
+
+	UE::Geometry::FDynamicMeshAttributeBase* MakeNew(UE::Geometry::FDynamicMesh3* ParentIn) const override
+	{
+		FDynamicMeshSourceTriangleIdAttribute* Attribute = new FDynamicMeshSourceTriangleIdAttribute(ParentIn);
+		Attribute->Initialize();
+		return Attribute;
+	}
+
+	void CompactInPlace(const UE::Geometry::FCompactMaps& CompactMaps) override
+	{
+		for (int32 FromTID = 0, NumTID = CompactMaps.NumTriangleMappings(); FromTID < NumTID; FromTID++)
+		{
+			const int32 ToTID = CompactMaps.GetTriangleMapping(FromTID);
+			if (ToTID == UE::Geometry::FCompactMaps::InvalidID)
+			{
+				continue;
+			}
+			if (ensure(ToTID <= FromTID))
+			{
+				SourceTriangleIds[ToTID] = SourceTriangleIds[FromTID];
+			}
+		}
+		SourceTriangleIds.Resize(Parent->MaxTriangleID());
+	}
+
+	void Reparent(UE::Geometry::FDynamicMesh3* NewParent) override
+	{
+		Parent = NewParent;
+	}
+
+	bool CopyThroughMapping(const UE::Geometry::TDynamicAttributeBase<UE::Geometry::FDynamicMesh3>* Source, const UE::Geometry::FMeshIndexMappings& Mapping) override
+	{
+		for (const TPair<int32, int32>& MapTID : Mapping.GetTriangleMap().GetForwardMap())
+		{
+			int32 SourceTId;
+			if (!ensure(Source->CopyOut(MapTID.Key, &SourceTId, sizeof(SourceTId))))
+			{
+				return false;
+			}
+			SetValue(MapTID.Value, SourceTId);
+		}
+		return true;
+	}
+
+	bool CopyOut(int RawID, void* Buffer, int BufferSize) const override
+	{
+		if (BufferSize != sizeof(int32) || !Parent->IsTriangle(RawID))
+		{
+			return false;
+		}
+
+		*static_cast<int32*>(Buffer) = SourceTriangleIds[RawID];
+		return true;
+	}
+
+	bool CopyIn(int RawID, void* Buffer, int BufferSize) override
+	{
+		if (BufferSize != sizeof(int32) || !Parent->IsTriangle(RawID))
+		{
+			return false;
+		}
+
+		SourceTriangleIds[RawID] = *static_cast<int32*>(Buffer);
+		return true;
+	}
+
+	TUniquePtr<UE::Geometry::TDynamicAttributeChangeBase<UE::Geometry::FDynamicMesh3>> NewBlankChange() const override
+	{
+		return MakeUnique<FDynamicMeshSourceTriangleIdAttributeChange>();
+	}
+
+	void Initialize()
+	{
+		SourceTriangleIds.Resize(Parent->MaxTriangleID());
+		SourceTriangleIds.Fill(INDEX_NONE);
+	}
+
+	void InitializeFromArray(TConstArrayView<int32> TriangleToNonManifoldVertexIDMap)
+	{
+		check(TriangleToNonManifoldVertexIDMap.Num() == Parent->MaxTriangleID());
+		SourceTriangleIds.Resize(Parent->MaxTriangleID());
+
+		for (int32 TriangleID : Parent->TriangleIndicesItr())
+		{
+			SetValue(TriangleID, TriangleToNonManifoldVertexIDMap[TriangleID]);
+		}
+		bValid = true;
+	}
+
+	void Copy(const FDynamicMeshSourceTriangleIdAttribute& Copy)
+	{
+		UE::Geometry::FDynamicMeshAttributeBase::CopyParentClassData(Copy);
+		SourceTriangleIds = Copy.SourceTriangleIds;
+		bValid = Copy.bValid;
+	}
+
+	int32 GetValue(int32 InTriangleId) const
+	{
+		return SourceTriangleIds[InTriangleId];
+	}
+
+	void SetValue(int32 InTriangleId, int32 InSourceTriangled)
+	{
+		SourceTriangleIds[InTriangleId] = InSourceTriangled;
+	}
+
+	// Any topo operation on the mesh will invalidate the non-manifold information.
+	virtual void OnSplitEdge(const DynamicMeshInfo::FEdgeSplitInfo&) override
+	{
+		bValid = false;
+	}
+
+	virtual void OnFlipEdge(const DynamicMeshInfo::FEdgeFlipInfo& FlipInfo) override
+	{
+		bValid = false;
+	}
+
+	virtual void OnCollapseEdge(const DynamicMeshInfo::FEdgeCollapseInfo&) override
+	{
+		bValid = false;
+	}
+
+	virtual void OnPokeTriangle(const DynamicMeshInfo::FPokeTriangleInfo&) override
+	{
+		bValid = false;
+	}
+
+	virtual void OnMergeEdges(const DynamicMeshInfo::FMergeEdgesInfo&) override
+	{
+		bValid = false;
+	}
+
+	virtual void OnMergeVertices(const DynamicMeshInfo::FMergeVerticesInfo&) override
+	{
+		bValid = false;
+	}
+
+	virtual void OnSplitVertex(const DynamicMeshInfo::FVertexSplitInfo&, const TArrayView<const int>&) override
+	{
+		bValid = false;
+	}
+
+	virtual void OnNewVertex(int VertexID, bool bInserted) override
+	{
+		bValid = false;
+	}
+
+	virtual void OnRemoveVertex(int VertexID) override
+	{
+		bValid = false;
+	}
+};
+
+
+void FDynamicMeshSourceTriangleIdAttributeChange::SaveInitialTriangle(const UE::Geometry::FDynamicMeshAttributeBase* Attribute, int TriangleID)
+{
+	const FDynamicMeshSourceTriangleIdAttribute* SourceAttribute = static_cast<const FDynamicMeshSourceTriangleIdAttribute*>(Attribute);
+	if (OldSourceTriangleIds.IsEmpty())
+	{
+		bOldInvalidState = SourceAttribute->IsValid();
+	}
+
+	OldSourceTriangleIds.Emplace(TriangleID, SourceAttribute->GetValue(TriangleID));
+}
+
+void FDynamicMeshSourceTriangleIdAttributeChange::StoreAllFinalTriangles(const UE::Geometry::FDynamicMeshAttributeBase* Attribute, const TArray<int>& TriangleIDs)
+{
+	NewSourceTriangleIds.Reserve(NewSourceTriangleIds.Num() + TriangleIDs.Num());
+	const FDynamicMeshSourceTriangleIdAttribute* SourceAttribute = static_cast<const FDynamicMeshSourceTriangleIdAttribute*>(Attribute);
+	for (int32 TriangleID : TriangleIDs)
+	{
+		NewSourceTriangleIds.Emplace(TriangleID, SourceAttribute->GetValue(TriangleID));
+	}
+
+	// Store the last known valid state.  
+	bNewInvalidState = SourceAttribute->IsValid();
+}
+
+bool FDynamicMeshSourceTriangleIdAttributeChange::Apply(UE::Geometry::FDynamicMeshAttributeBase* Attribute, bool bRevert) const
+{
+	const TArray<TPair<int32, int32>>& Changes = bRevert ? OldSourceTriangleIds : NewSourceTriangleIds;
+	FDynamicMeshSourceTriangleIdAttribute* SourceAttribute = static_cast<FDynamicMeshSourceTriangleIdAttribute*>(Attribute);
+	for (const TPair<int32, int32>& Item : Changes)
+	{
+		if (ensure(SourceAttribute->GetParent()->IsTriangle(Item.Key)))
+		{
+			SourceAttribute->SetValue(Item.Key, Item.Value);
+		}
+	}
+
+	// Restore the valid state as well.
+	SourceAttribute->bValid = bRevert ? bOldInvalidState : bNewInvalidState;
+
+	return true;
+}
+
+FClothPatternToDynamicMeshMappingSupport::FClothPatternToDynamicMeshMappingSupport(const UE::Geometry::FDynamicMesh3& Mesh)
+	: UE::Geometry::FNonManifoldMappingSupport(Mesh)
+{
+	Reset(Mesh);
+}
+
+const FName FClothPatternToDynamicMeshMappingSupport::ClothMeshTIDsAttrName = FName(TEXT("ClothMeshTIDsAttr"));
+
+void FClothPatternToDynamicMeshMappingSupport::Reset(const UE::Geometry::FDynamicMesh3& Mesh)
+{
+	UE::Geometry::FNonManifoldMappingSupport::Reset(Mesh);
+
+	if (const UE::Geometry::FDynamicMeshAttributeSet* Attributes = DynamicMesh->Attributes())
+	{
+		if (Attributes->HasAttachedAttribute(ClothMeshTIDsAttrName))
+		{
+			SrcTIDsAttribute = static_cast<const FDynamicMeshSourceTriangleIdAttribute*>(Attributes->GetAttachedAttribute(ClothMeshTIDsAttrName));
+		}
+	}
+}
+
+bool FClothPatternToDynamicMeshMappingSupport::IsMappedTriangleInSource() const
+{
+	return SrcTIDsAttribute != nullptr && SrcTIDsAttribute->IsValid();
+}
+
+int32 FClothPatternToDynamicMeshMappingSupport::GetOriginalTriangleID(const int32 Tid) const
+{
+	checkSlow(DynamicMesh->IsTriangle(Tid));
+
+	int32 Result = Tid;
+	if (SrcTIDsAttribute)
+	{
+		Result = SrcTIDsAttribute->GetValue(Tid);
+	}
+	return Result;
+}
+
+bool FClothPatternToDynamicMeshMappingSupport::AttachTriangleMappingData(const TArray<int32>& TriangleToOriginalVertexIDMap, UE::Geometry::FDynamicMesh3& InOutMesh)
+{
+	const int32 MaxTID = InOutMesh.MaxTriangleID();
+	UE::Geometry::FDynamicMeshAttributeSet* const Attributes = InOutMesh.Attributes();
+
+	if (TriangleToOriginalVertexIDMap.Num() < MaxTID || !Attributes)
+	{
+		return false;
+	}
+
+	// create and attach a triangle ID buffer, removing any pre-exising one.
+	FDynamicMeshSourceTriangleIdAttribute* TIDAttr = [&Attributes, &InOutMesh] {
+		// replace existing buffer.  we *should* be able to re-use it
+		if (Attributes->HasAttachedAttribute(ClothMeshTIDsAttrName))
+		{
+			Attributes->RemoveAttribute(ClothMeshTIDsAttrName);
+		}
+		FDynamicMeshSourceTriangleIdAttribute* SrcMeshTIDAttr = new FDynamicMeshSourceTriangleIdAttribute(&InOutMesh);
+		SrcMeshTIDAttr->SetName(ClothMeshTIDsAttrName);
+		InOutMesh.Attributes()->AttachAttribute(ClothMeshTIDsAttrName, SrcMeshTIDAttr);
+		return SrcMeshTIDAttr;
+		}();
+
+	TIDAttr->InitializeFromArray(TriangleToOriginalVertexIDMap);
+	return true;
+}
+
+void FClothPatternToDynamicMeshMappingSupport::RemoveTriangleMappingData(UE::Geometry::FDynamicMesh3& InOutMesh)
+{
+	UE::Geometry::FDynamicMeshAttributeSet* Attributes = InOutMesh.Attributes();
+
+	if (Attributes && Attributes->HasAttachedAttribute(ClothMeshTIDsAttrName))
+	{
+		Attributes->RemoveAttribute(ClothMeshTIDsAttrName);
+	}
+}
+
+
+}	// namespace UE::Chaos::ClothAsset
+
diff --git a/Engine/Plugins/ChaosClothAssetEditor/Source/ChaosClothAssetTools/Public/ChaosClothAsset/ClothPatternToDynamicMeshMappingSupport.h b/Engine/Plugins/ChaosClothAssetEditor/Source/ChaosClothAssetTools/Public/ChaosClothAsset/ClothPatternToDynamicMeshMappingSupport.h
new file mode 100644
--- /dev/null
+++ b/Engine/Plugins/ChaosClothAssetEditor/Source/ChaosClothAssetTools/Public/ChaosClothAsset/ClothPatternToDynamicMeshMappingSupport.h
@@ -0,0 +1,143 @@
+// Copyright Epic Games, Inc. All Rights Reserved.
+
+#pragma once
+
+#include "Templates/SharedPointer.h"
+#include "DynamicMesh/NonManifoldMappingSupport.h"
+
+
+namespace UE::Geometry { class FDynamicMesh3; }
+
+namespace UE::Chaos::ClothAsset
+{
+
+class FDynamicMeshSourceTriangleIdAttribute;
+
+// This uses the NonManifoldMapping attribute for vertices, and a new attribute for triangles.
+class FClothPatternToDynamicMeshMappingSupport : public UE::Geometry::FNonManifoldMappingSupport
+{
+public:
+	CHAOSCLOTHASSETTOOLS_API FClothPatternToDynamicMeshMappingSupport(const UE::Geometry::FDynamicMesh3& Mesh);
+
+	/**
+	* Update the support for a new DynamicMesh.
+	*/
+	CHAOSCLOTHASSETTOOLS_API void Reset(const UE::Geometry::FDynamicMesh3& Mesh);
+
+	/*
+	* Return true if attribute data indicates that the source data that was converted to this DynamicMesh contained mapped vertices.
+	*/
+	bool IsMappedVertexInSource() const
+	{
+		return IsNonManifoldVertexInSource();
+	}
+
+	/*
+	* Return true if attribute data indicates that the source data that was converted to this DynamicMesh contained mapped triangles.
+	*/
+	CHAOSCLOTHASSETTOOLS_API bool IsMappedTriangleInSource() const;
+
+	/*
+	* Return true if the provided DynamicMesh vertex id resulted from a remapped vertex in the source data.
+	* @param vid - the id of a vertex in the DynamicMesh.
+	*
+	* Note: the code assumes but does not check that vid is a valid vertex id
+	*/
+	bool IsMappedVertexID(const int32 vid) const
+	{
+		return !(vid == GetOriginalNonManifoldVertexID(vid));
+	}
+
+	/*
+	* Return true if the provided DynamicMesh triangle id resulted from a remapped triangle in the source data.
+	* @param Tid - the id of a triangle in the DynamicMesh.
+	*
+	* Note: the code assumes but does not check that Tid is a valid triangle id
+	*/
+	bool IsMappedTriangleID(const int32 Tid) const
+	{
+		return !(Tid == GetOriginalTriangleID(Tid));
+	}
+
+	/*
+	* Return the vertex ID in the original cloth data used to generate this DynamicMesh associated with the provided vertex id.
+	* In the case that the source data was actually manifold the returned vertex id will be identical to the DynamicMesh vertex id.
+	* @param vid - the id of a vertex in the DynamicMesh.
+	*
+	* Note: the code assumes but does not check that vid is a valid vertex vid.
+	*/
+	int32 GetOriginalVertexID(const int32 vid) const
+	{
+		return GetOriginalNonManifoldVertexID(vid);
+	}
+
+	/*
+	* Return the triangle ID in the original cloth data used to generate this DynamicMesh associated with the provided triangle id.
+	* In the case that the source data was actually manifold the returned triangle id will be identical to the DynamicMesh triangle id.
+	* @param Tid - the id of a triangle in the DynamicMesh.
+	*
+	* Note: the code assumes but does not check that Tid is a valid triangle id
+	*/
+	CHAOSCLOTHASSETTOOLS_API int32 GetOriginalTriangleID(const int32 Tid) const;
+
+
+	// --- helper functions.
+
+	/*
+	* Attaches or replaces vertex mapping data to the provided mesh.
+	* @param VertexToOriginalVertexIDMap - an array that maps each DynamicMesh vertex id to the associated original vertex id.
+	* @return false on failure (no attribute will be attached to the DynamicMesh in this case)
+	*
+	* Note: Failure occurs if the DynamicMesh does not have attributes enabled or if the provided array is not long enough to provide a mapping value for each DynamicMesh vertex id.
+	*/
+	static bool AttachVertexMappingData(const TArray<int32>& VertexToOriginalVertexIDMap, UE::Geometry::FDynamicMesh3& InOutMesh)
+	{
+		return AttachNonManifoldVertexMappingData(VertexToOriginalVertexIDMap, InOutMesh);
+	}
+
+	/*
+	* Attaches or replaces  triangle mapping data to the provided mesh.
+	* @param TriangleToOriginalVertexIDMap - an array that maps each DynamicMesh triangle id to the associated original triangle id.
+	* @return false on failure (no attribute will be attached to the DynamicMesh in this case)
+	*
+	* Note: Failure occurs if the DynamicMesh does not have attributes enabled or if the provided array is not long enough to provide a mapping value for each DynamicMesh triangle id.
+	*/
+	static CHAOSCLOTHASSETTOOLS_API bool AttachTriangleMappingData(const TArray<int32>& TriangleToOriginalVertexIDMap, UE::Geometry::FDynamicMesh3& InOutMesh);
+
+	/*
+	*  Removes triangle mapping data.
+	*
+	*  Note: this will invalidate any NonManifoldMappingSupport/ClothPatternToDynamicMeshMappingSupport object associated with this DynamicMesh,
+	*  and subsequent use of such object will produce unexpected results.
+	*/
+	static void RemoveVertexMappingData(UE::Geometry::FDynamicMesh3& InOutMesh)
+	{
+		RemoveNonManifoldVertexMappingData(InOutMesh);
+	}
+
+	static CHAOSCLOTHASSETTOOLS_API void RemoveTriangleMappingData(UE::Geometry::FDynamicMesh3& InOutMesh);
+
+	/*
+	*  Removes all mapping data.
+	*
+	*  Note, this will invalidate any  NonManifoldMappingSupport/ClothPatternToDynamicMeshMappingSupport object associated with this DynamicMesh,
+	*  and subsequent use of such object will produce unexpected results.
+	*/
+	static void RemoveAllMappingData(UE::Geometry::FDynamicMesh3& InOutMesh)
+	{
+		RemoveAllNonManifoldMappingData(InOutMesh);
+		RemoveTriangleMappingData(InOutMesh);
+	}
+
+	/*
+	* Name used to identify triangle attribute data generated during conversion to a DynamicMesh in the case that the source was non-manifold.
+	*/
+	static const CHAOSCLOTHASSETTOOLS_API FName ClothMeshTIDsAttrName;
+
+protected:
+
+	const FDynamicMeshSourceTriangleIdAttribute* SrcTIDsAttribute = nullptr;
+};
+
+}	// namespace UE::Chaos::ClothAsset
+
