Add input to ChaosClothAssetImport node

diff --git a/Engine/Plugins/ChaosClothAssetEditor/Source/ChaosClothAssetDataflowNodes/Private/ChaosClothAsset/ImportNode.cpp b/Engine/Plugins/ChaosClothAssetEditor/Source/ChaosClothAssetDataflowNodes/Private/ChaosClothAsset/ImportNode.cpp
--- a/Engine/Plugins/ChaosClothAssetEditor/Source/ChaosClothAssetDataflowNodes/Private/ChaosClothAsset/ImportNode.cpp
+++ b/Engine/Plugins/ChaosClothAssetEditor/Source/ChaosClothAssetDataflowNodes/Private/ChaosClothAsset/ImportNode.cpp
@@ -19,6 +19,7 @@
 FChaosClothAssetImportNode::FChaosClothAssetImportNode(const UE::Dataflow::FNodeParameters& InParam, FGuid InGuid)
 	: FDataflowNode(InParam, InGuid)
 {
+	RegisterInputConnection(&ClothAsset);
 	RegisterOutputConnection(&Collection);
 }
 
@@ -29,15 +30,20 @@ void FChaosClothAssetImportNode::Evaluate(UE::Dataflow::FContext& Context, const
 
 	if (Out->IsA<FManagedArrayCollection>(&Collection))
 	{
-		bool bSetValue = false;
-		if (ClothAsset)
+		// Create a new cloth collection with its LOD 0
+		const TSharedRef<FManagedArrayCollection> ClothCollection = MakeShared<FManagedArrayCollection>();
+		FCollectionClothFacade ClothFacade(ClothCollection);
+		ClothFacade.DefineSchema();
+
+		if (TObjectPtr<const UChaosClothAsset> InClothAsset = GetValue(Context, &ClothAsset))
 		{
+			bool bSetValue = false;
 			// Copy the main cloth asset details to this dataflow's owner if any
 			if (const UE::Dataflow::FEngineContext* EngineContext = Context.AsType<UE::Dataflow::FEngineContext>())
 			{
 				if (const UChaosClothAsset* const OwnerClothAsset = Cast<UChaosClothAsset>(EngineContext->Owner))
 				{
-					if (OwnerClothAsset == ClothAsset)
+					if (OwnerClothAsset == InClothAsset)
 					{
 						// Can't create a loop
 						bSetValue = false;
@@ -56,41 +62,36 @@ void FChaosClothAssetImportNode::Evaluate(UE::Dataflow::FContext& Context, const
 				// No terminal asset, this is a stray dataflow and it is safe to set the value without fear of a loop
 				bSetValue = true;
 			}
-		}
-
-		// Create a new cloth collection with its LOD 0
-		const TSharedRef<FManagedArrayCollection> ClothCollection = MakeShared<FManagedArrayCollection>();
-		FCollectionClothFacade ClothFacade(ClothCollection);
-		ClothFacade.DefineSchema();
 
-		// Copy the cloth asset to this node's output collection
-		if (bSetValue)
-		{
-			const TArray<TSharedRef<const FManagedArrayCollection>>& InClothCollections = ClothAsset->GetClothCollections();
-			if (ImportLod >= 0 && InClothCollections.Num() > ImportLod)
+			// Copy the cloth asset to this node's output collection
+			if (bSetValue)
 			{
-				// Copy cloth
-				const FCollectionClothConstFacade InClothFacade(InClothCollections[ImportLod]);
-				ClothFacade.Initialize(InClothFacade);
-
-				// Copy properties
-				const FCollectionPropertyConstFacade InPropertyFacade(InClothCollections[ImportLod]);
-				if (InPropertyFacade.IsValid())
+				const TArray<TSharedRef<const FManagedArrayCollection>>& InClothCollections = InClothAsset->GetClothCollections();
+				if (ImportLod >= 0 && InClothCollections.Num() > ImportLod)
 				{
-					FCollectionPropertyMutableFacade PropertyFacade(ClothCollection);
-					PropertyFacade.DefineSchema();
-					constexpr bool bUpdateExistingPropertiesFalse = false; // There are no existing properties.
-					PropertyFacade.Append(InClothCollections[ImportLod].ToSharedPtr(), bUpdateExistingPropertiesFalse);
-				}
+					// Copy cloth
+					const FCollectionClothConstFacade InClothFacade(InClothCollections[ImportLod]);
+					ClothFacade.Initialize(InClothFacade);
 
-				// Copy selections
-				const FCollectionClothSelectionConstFacade InSelectionFacade(InClothCollections[ImportLod]);
-				if (InSelectionFacade.IsValid())
-				{
-					FCollectionClothSelectionFacade SelectionFacade(ClothCollection);
-					SelectionFacade.DefineSchema();
-					constexpr bool bOverwriteExistingIfMismatchedTrue = true;
-					SelectionFacade.Append(InSelectionFacade, bOverwriteExistingIfMismatchedTrue);
+					// Copy properties
+					const FCollectionPropertyConstFacade InPropertyFacade(InClothCollections[ImportLod]);
+					if (InPropertyFacade.IsValid())
+					{
+						FCollectionPropertyMutableFacade PropertyFacade(ClothCollection);
+						PropertyFacade.DefineSchema();
+						constexpr bool bUpdateExistingPropertiesFalse = false; // There are no existing properties.
+						PropertyFacade.Append(InClothCollections[ImportLod].ToSharedPtr(), bUpdateExistingPropertiesFalse);
+					}
+
+					// Copy selections
+					const FCollectionClothSelectionConstFacade InSelectionFacade(InClothCollections[ImportLod]);
+					if (InSelectionFacade.IsValid())
+					{
+						FCollectionClothSelectionFacade SelectionFacade(ClothCollection);
+						SelectionFacade.DefineSchema();
+						constexpr bool bOverwriteExistingIfMismatchedTrue = true;
+						SelectionFacade.Append(InSelectionFacade, bOverwriteExistingIfMismatchedTrue);
+					}
 				}
 			}
 		}
diff --git a/Engine/Plugins/ChaosClothAssetEditor/Source/ChaosClothAssetDataflowNodes/Public/ChaosClothAsset/ImportNode.h b/Engine/Plugins/ChaosClothAssetEditor/Source/ChaosClothAssetDataflowNodes/Public/ChaosClothAsset/ImportNode.h
--- a/Engine/Plugins/ChaosClothAssetEditor/Source/ChaosClothAssetDataflowNodes/Public/ChaosClothAsset/ImportNode.h
+++ b/Engine/Plugins/ChaosClothAssetEditor/Source/ChaosClothAssetDataflowNodes/Public/ChaosClothAsset/ImportNode.h
@@ -40,7 +40,7 @@ public:
 	FManagedArrayCollection Collection;
 
 	/** The Cloth Asset to import into a collection. */
-	UPROPERTY(EditAnywhere, Category = "Cloth Asset Import")
+	UPROPERTY(EditAnywhere, Category = "Cloth Asset Import", Meta = (DataflowInput))
 	TObjectPtr<const UChaosClothAsset> ClothAsset;
 
 	/** The LOD to import into the collection. Only one LOD can be imported at a time. */
