Moving ClothAssetFactory.h from prviate to public folder in ChaosClothAssetTools for cross-module access

- Also appended export Macro onto the UChaosClothAssetFactory class

#rnx

diff --git a/Engine/Plugins/ChaosClothAssetEditor/Source/ChaosClothAssetTools/Private/ChaosClothAsset/ClothAssetFactory.h b/Engine/Plugins/ChaosClothAssetEditor/Source/ChaosClothAssetTools/Private/ChaosClothAsset/ClothAssetFactory.h
deleted file mode 100644
diff --git a/Engine/Plugins/ChaosClothAssetEditor/Source/ChaosClothAssetTools/Public/ChaosClothAsset/ClothAssetFactory.h b/Engine/Plugins/ChaosClothAssetEditor/Source/ChaosClothAssetTools/Public/ChaosClothAsset/ClothAssetFactory.h
new file mode 100644
--- /dev/null
+++ b/Engine/Plugins/ChaosClothAssetEditor/Source/ChaosClothAssetTools/Public/ChaosClothAsset/ClothAssetFactory.h
@@ -0,0 +1,28 @@
+// Copyright Epic Games, Inc. All Rights Reserved.
+
+#pragma once
+
+#include "CoreMinimal.h"
+#include "UObject/ObjectMacros.h"
+#include "Factories/Factory.h"
+
+#include "ClothAssetFactory.generated.h"
+
+/**
+ * Having a cloth factory allows the cloth asset to be created from the Editor's menus.
+ */
+UCLASS(Experimental)
+class CHAOSCLOTHASSETTOOLS_API UChaosClothAssetFactory : public UFactory
+{
+	GENERATED_BODY()
+public:
+	UChaosClothAssetFactory(const FObjectInitializer& ObjectInitializer);
+
+	/** UFactory Interface */
+	virtual bool CanCreateNew() const override { return true; }
+	virtual bool FactoryCanImport(const FString& Filename) override { return false; }
+	virtual bool ShouldShowInNewMenu() const override { return true; }
+	virtual UObject* FactoryCreateNew(UClass* InClass, UObject* InParent, FName InName, EObjectFlags Flags, UObject* Context, FFeedbackContext* Warn) override;
+	virtual FString GetDefaultNewAssetName() const override;
+	/** End UFactory Interface */
+};
