Cloth Paint Tool: have the brush radius adapt to the size of the mesh being painted, not a hardcoded value

diff --git a/Engine/Plugins/ChaosClothAssetEditor/Source/ChaosClothAssetEditorTools/Private/ChaosClothAsset/ClothWeightMapPaintTool.cpp b/Engine/Plugins/ChaosClothAssetEditor/Source/ChaosClothAssetEditorTools/Private/ChaosClothAsset/ClothWeightMapPaintTool.cpp
--- a/Engine/Plugins/ChaosClothAssetEditor/Source/ChaosClothAssetEditorTools/Private/ChaosClothAsset/ClothWeightMapPaintTool.cpp
+++ b/Engine/Plugins/ChaosClothAssetEditor/Source/ChaosClothAssetEditorTools/Private/ChaosClothAsset/ClothWeightMapPaintTool.cpp
@@ -675,7 +675,14 @@ void UClothEditorWeightMapPaintTool::PostSetupCheck() const
 
 void UClothEditorWeightMapPaintTool::InitializeBrushSizeRange(const UE::Geometry::FAxisAlignedBox3d& TargetBounds)
 {
-	BrushRelativeSizeRange = FInterval1d(5e-5, 50.0);	// brush diameter ranges from micrometer to meter scale
+	const double MaxDimension = FMath::Max(0.1, TargetBounds.MaxDim());
+
+	// Max brush size is the next power of 2 greater than MaxDimension. This allows the mesh size to change somewhat and still keep the same brush size, but still allow 
+	// the brush size to adapt to meshes of vastly different scales (e.g. switching from world coordinates to texture coordinates.)
+	constexpr double StepSize = 2.0;
+	const double MaxBrushSize = FMath::Pow(StepSize, FMath::CeilToDouble(FMath::LogX(StepSize, MaxDimension)));
+
+	BrushRelativeSizeRange = FInterval1d(5e-5 * MaxBrushSize, MaxBrushSize);
 	BrushProperties->BrushSize.InitializeWorldSizeRange(TInterval<float>((float)BrushRelativeSizeRange.Min, (float)BrushRelativeSizeRange.Max));
 	CalculateBrushRadius();
 }
