Added cloth-cloth constraint springs with a basic constraint creation UI.

UI exists for creating vertex-vertex and vertex-face constraints. Additionally, there is code for face-face constraints, but no UI; these correspond with the types of constraints imported via the new USD schema.

These constraints are considered experimental. The UI only allows creating constraints via selection sets. It would be nice to create/edit constraints with a picking tool as well.

Also includes changes to how GraphColoring works for coloring constraints for parallelism to be much simpler and faster to calculate. The final coloring is different from the original, but not in a way that meaningfully affects convergence/speed in tests.

diff --git a/Engine/Plugins/ChaosCloth/Source/ChaosCloth/Private/ChaosCloth/ChaosClothConstraints.cpp b/Engine/Plugins/ChaosCloth/Source/ChaosCloth/Private/ChaosCloth/ChaosClothConstraints.cpp
--- a/Engine/Plugins/ChaosCloth/Source/ChaosCloth/Private/ChaosCloth/ChaosClothConstraints.cpp
+++ b/Engine/Plugins/ChaosCloth/Source/ChaosCloth/Private/ChaosCloth/ChaosClothConstraints.cpp
@@ -21,6 +21,7 @@
 #include "Chaos/PBDCollisionSpringConstraints.h"
 #include "Chaos/PBDSelfCollisionSphereConstraints.h"
 #include "Chaos/PBDKinematicTriangleMeshCollisions.h"
+#include "Chaos/PBDRepulsionConstraints.h"
 #include "Chaos/PBDTriangleMeshCollisions.h"
 #include "Chaos/PBDTriangleMeshIntersections.h"
 #include "Chaos/PBDEvolution.h"
@@ -31,6 +32,8 @@
 #include "Chaos/CollectionPropertyFacade.h"
 #include "Chaos/Deformable/GaussSeidelMainConstraint.h"
 #include "Chaos/Deformable/GaussSeidelCorotatedCodimensionalConstraints.h"
+#include "Chaos/XPBDEmbeddedSpringConstraints.h"
+#include "Chaos/CollectionEmbeddedSpringConstraintFacade.h"
 #include "PhysicsProxy/PerSolverFieldSystem.h"
 #include "Utils/ClothingMeshUtils.h"
 #include "HAL/IConsoleManager.h"
@@ -464,18 +467,6 @@ void FClothConstraints::OnCollisionRangeRemoved(int32 CollisionRangeId)
 	}
 }
 
-void FClothConstraints::AddRules(
-	const Softs::FCollectionPropertyConstFacade& ConfigProperties,
-	const FTriangleMesh& TriangleMesh,
-	const FClothingPatternData* PatternData,
-	const TMap<FString, TConstArrayView<FRealSingle>>& WeightMaps,
-	const TArray<TConstArrayView<TTuple<int32, int32, FRealSingle>>>& Tethers,
-	Softs::FSolverReal MeshScale, bool bEnabled)
-{
-	AddRules(ConfigProperties, TriangleMesh, PatternData, WeightMaps, TMap<FString, const TSet<int32>*>(), TMap<FString, const TSet<int32>*>(),
-		TMap<FString, TConstArrayView<int32>>(), Tethers, MeshScale, bEnabled);
-}
-
 //only counts the number of rules that GS/CG currently cupports
 void FClothConstraints::GetGSNumRules()
 {
@@ -795,7 +786,8 @@ void FClothConstraints::AddRules(
 	const FTriangleMesh* MultiResCoarseLODMesh,
 	const int32 MultiResCoarseLODParticleRangeId,
 	const TSharedPtr<Softs::FMultiResConstraints>& FineLODMultiResConstraint,
-	const TArray<const FClothComplexColliders*>& ComplexColliders)
+	const TArray<const FClothComplexColliders*>& ComplexColliders,
+	const TSharedPtr<const FManagedArrayCollection>& ManagedArrayCollection)
 {
 	// Self collisions
 	CreateSelfCollisionConstraints(ConfigProperties, WeightMaps, VertexSets, FaceSets, FaceIntMaps, TriangleMesh);
@@ -840,6 +832,9 @@ void FClothConstraints::AddRules(
 
 		//Multires Springs
 		CreateMultiresConstraint(ConfigProperties, WeightMaps, TriangleMesh, MultiResCoarseLODMesh, MultiResCoarseLODParticleRangeId);
+
+		// ClothCloth Springs
+		CreateClothClothConstraints(ConfigProperties, ManagedArrayCollection);
 	}
 
 	// Commit rules to solver
@@ -1536,6 +1531,100 @@ void FClothConstraints::CreateMultiresConstraint(
 	}
 }
 
+void FClothConstraints::CreateClothClothConstraints(
+	const Softs::FCollectionPropertyConstFacade& ConfigProperties,
+	const TSharedPtr<const FManagedArrayCollection>& ManagedArrayCollection)
+{
+	if (ManagedArrayCollection)
+	{
+		const Softs::FEmbeddedSpringFacade SpringFacade(*ManagedArrayCollection, FName(TEXT("SimVertices3D")));
+		for (int32 ConstraintIndex = 0; ConstraintIndex < SpringFacade.GetNumSpringConstraints(); ++ConstraintIndex)
+		{
+			const Softs::FEmbeddedSpringConstraintFacade SpringConstraintFacade = SpringFacade.GetSpringConstraintConst(ConstraintIndex);
+			if (SpringConstraintFacade.IsValid())
+			{
+				const FUintVector2 EndPoints = SpringConstraintFacade.GetConstraintEndPointNumIndices();
+				const FString& ConstraintName = SpringConstraintFacade.GetConstraintName();
+				if (EndPoints == FUintVector2(1, 1))
+				{
+					if (ConstraintName == TEXT("VertexSpringConstraint"))
+					{
+						if (Softs::FXPBDVertexConstraints::IsEnabled(ConfigProperties))
+						{
+							if (ensure(!ClothVertexSpringConstraints))
+							{
+
+								ClothVertexSpringConstraints = MakeShared<Softs::FXPBDVertexConstraints>(
+									Evolution->GetSoftBodyParticles(ParticleRangeId),
+									ConfigProperties,
+									SpringConstraintFacade
+								);
+								++NumConstraintInits;
+								++NumConstraintRules;
+							}
+						}
+					}
+				}
+				else if (EndPoints == FUintVector2(1, 3))
+				{
+					if (ConstraintName == TEXT("VertexFaceSpringConstraint"))
+					{
+						if (Softs::FXPBDVertexFaceConstraints::IsEnabled(ConfigProperties))
+						{
+							if (ensure(!ClothVertexFaceSpringConstraints))
+							{
+
+								ClothVertexFaceSpringConstraints = MakeShared<Softs::FXPBDVertexFaceConstraints>(
+									Evolution->GetSoftBodyParticles(ParticleRangeId),
+									ConfigProperties,
+									SpringConstraintFacade
+								);
+								++NumConstraintInits;
+								++NumConstraintRules;
+							}
+						}
+					}
+					else if (ConstraintName == TEXT("VertexFaceRepulsionConstraint"))
+					{
+						if (Softs::FPBDRepulsionConstraints::IsEnabled(ConfigProperties))
+						{
+							if (ensure(!RepulsionConstraints))
+							{
+								RepulsionConstraints = MakeShared<Softs::FPBDRepulsionConstraints>(
+									Evolution->GetSoftBodyParticles(ParticleRangeId),
+									ConfigProperties,
+									SpringConstraintFacade
+								);
+								++NumPostCollisionConstraintRules;
+							}
+						}
+					}
+				}
+				else if (EndPoints == FUintVector2(3, 3))
+				{
+					if (ConstraintName == TEXT("FaceSpringConstraint"))
+					{
+						if (Softs::FXPBDFaceConstraints::IsEnabled(ConfigProperties))
+						{
+							if (ensure(!ClothFaceSpringConstraints))
+							{
+
+								ClothFaceSpringConstraints = MakeShared<Softs::FXPBDFaceConstraints>(
+									Evolution->GetSoftBodyParticles(ParticleRangeId),
+									ConfigProperties,
+									SpringConstraintFacade
+								);
+								++NumConstraintInits;
+								++NumConstraintRules;
+							}
+						}
+					}
+				}
+			}
+		}
+	}
+}
+
 void FClothConstraints::CreateForceBasedRules(const TSharedPtr<Softs::FMultiResConstraints>& FineLODMultiResConstraint)
 {
 	if (FineLODMultiResConstraint)
@@ -1719,6 +1808,27 @@ void FClothConstraints::CreateForceBasedRules(const TSharedPtr<Softs::FMultiResC
 		constexpr bool bPostCollisions = false;
 		RuleCreator.AddPerIterationPBDConstraintRule_Apply<bPostCollisions>(AnimDriveConstraints.Get());
 	}
+	if (ClothVertexSpringConstraints)
+	{
+		constexpr bool bInitParticles = false;
+		RuleCreator.AddXPBDParallelInitRule<bInitParticles>(ClothVertexSpringConstraints.Get());
+		constexpr bool bPostCollisions = false;
+		RuleCreator.AddPerIterationPBDConstraintRule_Apply<bPostCollisions>(ClothVertexSpringConstraints.Get());
+	}
+	if (ClothVertexFaceSpringConstraints)
+	{
+		constexpr bool bInitParticles = false;
+		RuleCreator.AddXPBDParallelInitRule<bInitParticles>(ClothVertexFaceSpringConstraints.Get());
+		constexpr bool bPostCollisions = false;
+		RuleCreator.AddPerIterationPBDConstraintRule_Apply<bPostCollisions>(ClothVertexFaceSpringConstraints.Get());
+	}
+	if (ClothFaceSpringConstraints)
+	{
+		constexpr bool bInitParticles = false;
+		RuleCreator.AddXPBDParallelInitRule<bInitParticles>(ClothFaceSpringConstraints.Get());
+		constexpr bool bPostCollisions = false;
+		RuleCreator.AddPerIterationPBDConstraintRule_Apply<bPostCollisions>(ClothFaceSpringConstraints.Get());
+	}
 
 	if (CollisionConstraint)
 	{
@@ -1773,6 +1883,12 @@ void FClothConstraints::CreateForceBasedRules(const TSharedPtr<Softs::FMultiResC
 		RuleCreator.AddUpdateLinearSystemRule_UpdateLinearSystem(SelfCollisionConstraints.Get());
 	}
 
+	if (RepulsionConstraints)
+	{
+		constexpr bool bPostCollisions = true;
+		RuleCreator.AddPerIterationPBDConstraintRule_Apply<bPostCollisions>(RepulsionConstraints.Get());
+	}
+
 	if (SelfCollisionSphereConstraints)
 	{
 		RuleCreator.AddPostInitialGuessParallelInitRule(
@@ -2152,14 +2268,6 @@ void FClothConstraints::CreatePBDRules()
 				AnimDriveConstraints->Apply(Particles, Dt);
 			};
 	}
-	if (ShapeConstraints_Deprecated) // TODO: Remove 5.6
-	{
-		ConstraintRules[ConstraintRuleIndex++] =
-			[this](Softs::FSolverParticles& Particles, const Softs::FSolverReal Dt)
-			{
-				ShapeConstraints_Deprecated->Apply(Particles, Dt);
-			};
-	}
 
 	if (SelfCollisionInit && SelfCollisionConstraints)
 	{
@@ -2337,6 +2445,18 @@ void FClothConstraints::Update(
 	{
 		AnimDriveConstraints->SetProperties(ConfigProperties, WeightMaps);
 	}
+	if (ClothVertexSpringConstraints)
+	{
+		ClothVertexSpringConstraints->SetProperties(ConfigProperties);
+	}
+	if (ClothVertexFaceSpringConstraints)
+	{
+		ClothVertexFaceSpringConstraints->SetProperties(ConfigProperties);
+	}
+	if (ClothFaceSpringConstraints)
+	{
+		ClothFaceSpringConstraints->SetProperties(ConfigProperties);
+	}
 	bool bSelfCollisionInitCollidableSubMeshWillRebuild = false;
 	if (SelfCollisionInit)
 	{
@@ -2356,6 +2476,10 @@ void FClothConstraints::Update(
 	{
 		SelfCollisionSphereConstraints->SetProperties(ConfigProperties, VertexSets);
 	}
+	if (RepulsionConstraints)
+	{
+		RepulsionConstraints->SetProperties(ConfigProperties);
+	}
 	if (MultiResConstraints)
 	{
 		MultiResConstraints->SetProperties(ConfigProperties, WeightMaps);
@@ -2394,15 +2518,4 @@ void FClothConstraints::Update(
 		SkinnedTriangleCollisionsConstraint->SetProperties(ConfigProperties, WeightMaps);
 	}
 }
-
-// Deprecated 5.4
-void FClothConstraints::Update(
-	const Softs::FCollectionPropertyConstFacade& ConfigProperties,
-	const TMap<FString, TConstArrayView<FRealSingle>>& WeightMaps,
-	Softs::FSolverReal MeshScale,
-	Softs::FSolverReal MaxDistancesScale)
-{
-	Update(ConfigProperties, WeightMaps, TMap<FString, const TSet<int32>*>(), TMap<FString, const TSet<int32>*>(), TMap<FString, TConstArrayView<int32>>(), MeshScale, MaxDistancesScale);
-}
-
 }  // End namespace Chaos
diff --git a/Engine/Plugins/ChaosCloth/Source/ChaosCloth/Private/ChaosCloth/ChaosClothVisualization.cpp b/Engine/Plugins/ChaosCloth/Source/ChaosCloth/Private/ChaosCloth/ChaosClothVisualization.cpp
--- a/Engine/Plugins/ChaosCloth/Source/ChaosCloth/Private/ChaosCloth/ChaosClothVisualization.cpp
+++ b/Engine/Plugins/ChaosCloth/Source/ChaosCloth/Private/ChaosCloth/ChaosClothVisualization.cpp
@@ -21,6 +21,7 @@
 #include "Chaos/PBDCollisionSpringConstraints.h"
 #include "Chaos/PBDFlatWeightMap.h"
 #include "Chaos/PBDLongRangeConstraints.h"
+#include "Chaos/PBDRepulsionConstraints.h"
 #include "Chaos/PBDSelfCollisionSphereConstraints.h"
 #include "Chaos/PBDSphericalConstraint.h"
 #include "Chaos/PBDSpringConstraints.h"
@@ -32,6 +33,7 @@
 #include "Chaos/XPBDSpringConstraints.h"
 #include "Chaos/XPBDAnisotropicBendingConstraints.h"
 #include "Chaos/XPBDAnisotropicSpringConstraints.h"
+#include "Chaos/XPBDEmbeddedSpringConstraints.h"
 #include "Chaos/XPBDStretchBiasElementConstraints.h"
 #include "Chaos/WeightedLatticeImplicitObject.h"
 #include "DynamicMeshBuilder.h"
@@ -2940,6 +2942,127 @@ namespace Chaos
 		}
 	}
 
+	void FClothVisualizationNoGC::DrawClothClothConstraints(FPrimitiveDrawInterface* PDI) const
+	{
+		if (!Solver)
+		{
+			return;
+		}
+		static const FLinearColor ActiveConstraintColor = FColor::Turquoise;
+		static const FLinearColor InactiveConstraintColor = FColor::Silver;
+		for (const FClothingSimulationCloth* const Cloth : Solver->GetCloths())
+		{
+			const int32 ParticleRangeId = Cloth->GetParticleRangeId(Solver);
+			if (ParticleRangeId == INDEX_NONE)
+			{
+				continue;
+			}
+
+			const FClothConstraints& ClothConstraints = Solver->GetClothConstraints(ParticleRangeId);
+			if (const Softs::FXPBDVertexConstraints* const ClothVertexConstraints = ClothConstraints.GetClothVertexSpringConstraints().Get())
+			{
+				const TConstArrayView<Softs::FSolverVec3> Positions = Cloth->GetParticlePositions(Solver);
+				const TArray<TVec2<int32>>& Constraints = ClothVertexConstraints->GetConstraints();
+				const TArray<Softs::FSolverReal>& SpringLengths = ClothVertexConstraints->GetSpringLengths();
+				for (int32 ConstraintIndex = 0; ConstraintIndex < Constraints.Num(); ++ConstraintIndex)
+				{
+					const Softs::FSolverVec3& P1 = Positions[Constraints[ConstraintIndex][0]];
+					const Softs::FSolverVec3& P2 = Positions[Constraints[ConstraintIndex][1]];
+
+					const Softs::FSolverReal DistSq = Softs::FSolverVec3::DistSquared(P1, P2);
+					const Softs::FSolverReal Stiffness = DistSq >= FMath::Square(SpringLengths[ConstraintIndex]) ? ClothVertexConstraints->GetExtensionStiffness(ConstraintIndex) : ClothVertexConstraints->GetCompressionStiffness(ConstraintIndex);
+					const bool bIsActive = Stiffness > (Softs::FSolverReal)0.f;
+					const FVec3 Pos0 = GetWorldPosition(*Solver, P1);
+					const FVec3 Pos1 = GetWorldPosition(*Solver, P2);
+					const FLinearColor& ConstraintColor = bIsActive ? ActiveConstraintColor : InactiveConstraintColor;
+					DrawLine(PDI, Pos0, Pos1, ConstraintColor);
+					DrawPoint(PDI, Pos0, ConstraintColor, nullptr);
+					DrawPoint(PDI, Pos1, ConstraintColor, nullptr);
+				}
+			}
+			if (const Softs::FXPBDVertexFaceConstraints* const ClothVertexFaceConstraints = ClothConstraints.GetClothVertexFaceSpringConstraints().Get())
+			{
+				const TConstArrayView<Softs::FSolverVec3> Positions = Cloth->GetParticlePositions(Solver);
+				const TArray<TVec4<int32>>& Constraints = ClothVertexFaceConstraints->GetConstraints();
+				const TArray<TVec4<Softs::FSolverReal>>& Weights = ClothVertexFaceConstraints->GetWeights();
+				const TArray<Softs::FSolverReal>& SpringLengths = ClothVertexFaceConstraints->GetSpringLengths();
+				for (int32 ConstraintIndex = 0; ConstraintIndex < Constraints.Num(); ++ConstraintIndex)
+				{
+					const Softs::FSolverVec3& P1 = Positions[Constraints[ConstraintIndex][0]];
+					const Softs::FSolverVec3& P2 = Positions[Constraints[ConstraintIndex][1]];
+					const Softs::FSolverVec3& P3 = Positions[Constraints[ConstraintIndex][2]];
+					const Softs::FSolverVec3& P4 = Positions[Constraints[ConstraintIndex][3]];
+
+					const Softs::FSolverVec3 P = -Weights[ConstraintIndex][1] * P2 - Weights[ConstraintIndex][2] * P3 - Weights[ConstraintIndex][3] * P4;
+
+					const Softs::FSolverReal DistSq = Softs::FSolverVec3::DistSquared(P1, P);
+					const Softs::FSolverReal Stiffness = DistSq >= FMath::Square(SpringLengths[ConstraintIndex]) ? ClothVertexFaceConstraints->GetExtensionStiffness(ConstraintIndex) : ClothVertexFaceConstraints->GetCompressionStiffness(ConstraintIndex);
+					const bool bIsActive = Stiffness > (Softs::FSolverReal)0.f;
+					const FVec3 Pos0 = GetWorldPosition(*Solver, P1);
+					const FVec3 Pos1 = GetWorldPosition(*Solver, P);
+					const FLinearColor& ConstraintColor = bIsActive ? ActiveConstraintColor : InactiveConstraintColor;
+					DrawLine(PDI, Pos0, Pos1, ConstraintColor);
+					DrawPoint(PDI, Pos0, ConstraintColor, nullptr);
+					DrawPoint(PDI, Pos1, ConstraintColor, nullptr);
+				}
+			}
+			if (const Softs::FXPBDFaceConstraints* const ClothFaceConstraints = ClothConstraints.GetClothFaceSpringConstraints().Get())
+			{
+				const TConstArrayView<Softs::FSolverVec3> Positions = Cloth->GetParticlePositions(Solver);
+				const TArray<TVector<int32,6>>& Constraints = ClothFaceConstraints->GetConstraints();
+				const TArray<TVector<Softs::FSolverReal,6>>& Weights = ClothFaceConstraints->GetWeights();
+				const TArray<Softs::FSolverReal>& SpringLengths = ClothFaceConstraints->GetSpringLengths();
+				for (int32 ConstraintIndex = 0; ConstraintIndex < Constraints.Num(); ++ConstraintIndex)
+				{
+					const Softs::FSolverVec3& P1 = Positions[Constraints[ConstraintIndex][0]];
+					const Softs::FSolverVec3& P2 = Positions[Constraints[ConstraintIndex][1]];
+					const Softs::FSolverVec3& P3 = Positions[Constraints[ConstraintIndex][2]];
+					const Softs::FSolverVec3& P4 = Positions[Constraints[ConstraintIndex][3]];
+					const Softs::FSolverVec3& P5 = Positions[Constraints[ConstraintIndex][4]];
+					const Softs::FSolverVec3& P6 = Positions[Constraints[ConstraintIndex][5]];
+
+					const Softs::FSolverVec3 Ps =  Weights[ConstraintIndex][0] * P1 + Weights[ConstraintIndex][1] * P2 + Weights[ConstraintIndex][2] * P3;
+					const Softs::FSolverVec3 Pt = -Weights[ConstraintIndex][3] * P4 - Weights[ConstraintIndex][4] * P5 - Weights[ConstraintIndex][5] * P6;
+
+					const Softs::FSolverReal DistSq = Softs::FSolverVec3::DistSquared(Ps, Pt);
+					const Softs::FSolverReal Stiffness = DistSq >= FMath::Square(SpringLengths[ConstraintIndex]) ? ClothFaceConstraints->GetExtensionStiffness(ConstraintIndex) : ClothFaceConstraints->GetCompressionStiffness(ConstraintIndex);
+					const bool bIsActive = Stiffness > (Softs::FSolverReal)0.f;
+					const FVec3 Pos0 = GetWorldPosition(*Solver, Ps);
+					const FVec3 Pos1 = GetWorldPosition(*Solver, Pt);
+					const FLinearColor& ConstraintColor = bIsActive ? ActiveConstraintColor : InactiveConstraintColor;
+					DrawLine(PDI, Pos0, Pos1, ConstraintColor);
+					DrawPoint(PDI, Pos0, ConstraintColor, nullptr);
+					DrawPoint(PDI, Pos1, ConstraintColor, nullptr);
+				}
+			}
+			if (const Softs::FPBDRepulsionConstraints* const RepulsionConstraints = ClothConstraints.GetRepulsionConstraints().Get())
+			{
+				const TConstArrayView<Softs::FSolverVec3> Positions = Cloth->GetParticlePositions(Solver);
+				const TArray<TVec4<int32>>& Constraints = RepulsionConstraints->GetConstraints();
+				const TArray<Softs::FSolverReal>& RestLengths = RepulsionConstraints->GetRestLengths();
+				const TArray<bool>& IsActive = RepulsionConstraints->GetConstraintIsActive();
+				for (int32 ConstraintIndex = 0; ConstraintIndex < Constraints.Num(); ++ConstraintIndex)
+				{
+					const Softs::FSolverVec3& P1 = Positions[Constraints[ConstraintIndex][0]];
+					const Softs::FSolverVec3& P2 = Positions[Constraints[ConstraintIndex][1]];
+					const Softs::FSolverVec3& P3 = Positions[Constraints[ConstraintIndex][2]];
+					const Softs::FSolverVec3& P4 = Positions[Constraints[ConstraintIndex][3]];
+
+					FSolverVec3 Bary;
+					const Softs::FSolverVec3 P = FindClosestPointAndBaryOnTriangle(P2, P3, P4, P1, Bary);
+
+					const Softs::FSolverReal DistSq = Softs::FSolverVec3::DistSquared(P1, P);
+					const FVec3 Pos0 = GetWorldPosition(*Solver, P1);
+					const FVec3 Pos1 = GetWorldPosition(*Solver, P);
+					const FLinearColor& ConstraintColor = IsActive[ConstraintIndex] ? ActiveConstraintColor : InactiveConstraintColor;
+					DrawLine(PDI, Pos0, Pos1, ConstraintColor);
+					DrawPoint(PDI, Pos0, ConstraintColor, nullptr);
+					DrawPoint(PDI, Pos1, ConstraintColor, nullptr);
+				}
+			}
+		}
+	}
+
 #if WITH_EDITOR
 	TArray<FString> FClothVisualizationNoGC::GetAllWeightMapNames() const
 	{
diff --git a/Engine/Plugins/ChaosCloth/Source/ChaosCloth/Private/ChaosCloth/ChaosClothingSimulation.cpp b/Engine/Plugins/ChaosCloth/Source/ChaosCloth/Private/ChaosCloth/ChaosClothingSimulation.cpp
--- a/Engine/Plugins/ChaosCloth/Source/ChaosCloth/Private/ChaosCloth/ChaosClothingSimulation.cpp
+++ b/Engine/Plugins/ChaosCloth/Source/ChaosCloth/Private/ChaosCloth/ChaosClothingSimulation.cpp
@@ -95,6 +95,7 @@ namespace ClothingSimulationCVar
 	TAutoConsoleVariable<bool> DebugDrawSelfIntersection    (TEXT("p.ChaosCloth.DebugDrawSelfIntersection"    ), false, TEXT("Whether to debug draw the Chaos Cloth self intersection information"), ECVF_Cheat);
 	TAutoConsoleVariable<bool> DebugDrawParticleIndices     (TEXT("p.ChaosCloth.DebugDrawParticleIndices"     ), false, TEXT("Whether to debug draw the Chaos Cloth particle indices"), ECVF_Cheat);
 	TAutoConsoleVariable<bool> DebugDrawElementIndices      (TEXT("p.ChaosCloth.DebugDrawElementIndices"      ), false, TEXT("Whether to debug draw the Chaos Cloth element indices"), ECVF_Cheat);
+	TAutoConsoleVariable<bool> DebugDrawClothClothConstraints    (TEXT("p.ChaosCloth.DebugDrawClothClothConstraints"    ), false, TEXT("Whether to debug draw the Chaos Cloth-Cloth constraints"), ECVF_Cheat);
 }
 #endif  // #if CHAOS_DEBUG_DRAW
 
diff --git a/Engine/Plugins/ChaosCloth/Source/ChaosCloth/Private/ChaosCloth/ChaosClothingSimulationCloth.cpp b/Engine/Plugins/ChaosCloth/Source/ChaosCloth/Private/ChaosCloth/ChaosClothingSimulationCloth.cpp
--- a/Engine/Plugins/ChaosCloth/Source/ChaosCloth/Private/ChaosCloth/ChaosClothingSimulationCloth.cpp
+++ b/Engine/Plugins/ChaosCloth/Source/ChaosCloth/Private/ChaosCloth/ChaosClothingSimulationCloth.cpp
@@ -68,6 +68,7 @@ struct FClothingSimulationCloth::FLODData
 	const TMap<FString, const TSet<int32>*> FaceSets;
 	const TMap<FString, TConstArrayView<int32>> FaceIntMaps;
 	const TArray<TConstArrayView<TTuple<int32, int32, FRealSingle>>> Tethers;
+	const TSharedPtr<const FManagedArrayCollection> ManagedArrayCollection;
 
 	const FClothingPatternData PatternData;
 	const FTriangleMesh NoOffsetTriangleMesh;
@@ -96,7 +97,8 @@ struct FClothingSimulationCloth::FLODData
 		TMap<FString, const TSet<int32>*>&& InVertexSets,
 		TMap<FString, const TSet<int32>*>&& InFaceSets,
 		TMap<FString, TConstArrayView<int32>>&& InFaceIntMaps,
-		const TArray<TConstArrayView<TTuple<int32, int32, FRealSingle>>>& InTethers);
+		const TArray<TConstArrayView<TTuple<int32, int32, FRealSingle>>>& InTethers,
+		const TSharedPtr<const FManagedArrayCollection>& InManagedArrayCollection);
 
 	static FTriangleMesh BuildTriangleMesh(const TConstArrayView<uint32>& Indices, const int32 NumParticles);
 
@@ -130,13 +132,15 @@ FClothingSimulationCloth::FLODData::FLODData(
 	TMap<FString, const TSet<int32>*>&& InVertexSets,
 	TMap<FString, const TSet<int32>*>&& InFaceSets,
 	TMap<FString, TConstArrayView<int32>>&& InFaceIntMaps,
-	const TArray<TConstArrayView<TTuple<int32, int32, FRealSingle>>>& InTethers)
+	const TArray<TConstArrayView<TTuple<int32, int32, FRealSingle>>>& InTethers,
+	const TSharedPtr<const FManagedArrayCollection>& InManagedArrayCollection)
 	: NumParticles(InNumParticles)
 	, WeightMaps(MoveTemp(InWeightMaps))
 	, VertexSets(MoveTemp(InVertexSets))
 	, FaceSets(MoveTemp(InFaceSets))
 	, FaceIntMaps(MoveTemp(InFaceIntMaps))
 	, Tethers(InTethers)
+	, ManagedArrayCollection(InManagedArrayCollection)
 	, PatternData(NumParticles, InIndices, InPatternPositions, InPatternIndices, InPatternToWeldedIndices)
 	, NoOffsetTriangleMesh(BuildTriangleMesh(InIndices, NumParticles))
 {
@@ -347,7 +351,7 @@ void FClothingSimulationCloth::FLODData::Add(FClothingSimulationSolver* Solver,
 
 	// Create constraints
 	const bool bEnabled = false;  // Set constraint disabled by default
-	ClothConstraints.AddRules(ConfigProperties, TriangleMesh, &PatternData, WeightMaps, VertexSets, FaceSets, FaceIntMaps, Tethers, MeshScale, bEnabled, CoarseLODTriangleMesh, CoarseLODParticleRangeId, FineLODMultiResConstraint, ComplexColliders);
+	ClothConstraints.AddRules(ConfigProperties, TriangleMesh, &PatternData, WeightMaps, VertexSets, FaceSets, FaceIntMaps, Tethers, MeshScale, bEnabled, CoarseLODTriangleMesh, CoarseLODParticleRangeId, FineLODMultiResConstraint, ComplexColliders, ManagedArrayCollection);
 
 	// Update LOD stats
 	const TConstArrayView<Softs::FSolverReal> InvMasses(Solver->GetParticleInvMasses(ParticleRangeId), NumParticles);
@@ -491,7 +495,8 @@ void FClothingSimulationCloth::SetMesh(FClothingSimulationMesh* InMesh)
 			MoveTemp(VertexSets),
 			MoveTemp(FaceSets),
 			MoveTemp(FaceIntMaps),
-			Mesh->GetTethers(LODIndex, bUseGeodesicTethers)));
+			Mesh->GetTethers(LODIndex, bUseGeodesicTethers),
+			Mesh->GetManagedArrayCollection(LODIndex)));
 	}
 
 	// Iterate all known solvers
diff --git a/Engine/Plugins/ChaosCloth/Source/ChaosCloth/Public/ChaosCloth/ChaosClothConstraints.h b/Engine/Plugins/ChaosCloth/Source/ChaosCloth/Public/ChaosCloth/ChaosClothConstraints.h
--- a/Engine/Plugins/ChaosCloth/Source/ChaosCloth/Public/ChaosCloth/ChaosClothConstraints.h
+++ b/Engine/Plugins/ChaosCloth/Source/ChaosCloth/Public/ChaosCloth/ChaosClothConstraints.h
@@ -11,6 +11,8 @@
 #include "Chaos/Deformable/GaussSeidelMainConstraint.h"
 #include "Chaos/Deformable/GaussSeidelCorotatedCodimensionalConstraints.h"
 
+struct FManagedArrayCollection;
+
 namespace Chaos
 {
 
@@ -74,17 +76,8 @@ namespace Chaos
 			const FTriangleMesh* MultiResCoarseLODMesh = nullptr,
 			const int32 MultiResCoarseLODParticleRangeId = INDEX_NONE,
 			const TSharedPtr<Softs::FMultiResConstraints>& FineLODMultiResConstraint = TSharedPtr<Softs::FMultiResConstraints>(nullptr),
-			const TArray<const FClothComplexColliders*>& ComplexColliders = TArray<const FClothComplexColliders*>());
-
-		UE_DEPRECATED(5.4, "Use AddRules() with WeightMaps, VertexSets, FaceSets, FaceIntMaps, and optional PatternData instead.")
-		void AddRules(
-			const Softs::FCollectionPropertyConstFacade& ConfigProperties,
-			const FTriangleMesh& TriangleMesh,
-			const FClothingPatternData* PatternData,
-			const TMap<FString, TConstArrayView<FRealSingle>>& WeightMaps,
-			const TArray<TConstArrayView<TTuple<int32, int32, FRealSingle>>>& Tethers,
-			Softs::FSolverReal MeshScale,
-			bool bEnabled);
+			const TArray<const FClothComplexColliders*>& ComplexColliders = TArray<const FClothComplexColliders*>(),
+			const TSharedPtr<const FManagedArrayCollection>& ManagedArrayCollection = TSharedPtr<const FManagedArrayCollection>());
 
 		void Update(
 			const Softs::FCollectionPropertyConstFacade& ConfigProperties,
@@ -97,22 +90,12 @@ namespace Chaos
 			const FRotation3& LocalSpaceRotation = FRotation3::Identity,
 			const FRotation3& ReferenceSpaceRotation = FRotation3::Identity);
 
-		UE_DEPRECATED(5.4, "Use Update() with WeightMaps, VertexSets, FaceSets, and FaceIntMaps instead.")
-		void Update(
-			const Softs::FCollectionPropertyConstFacade& ConfigProperties,
-			const TMap<FString, TConstArrayView<FRealSingle>>& WeightMaps,
-			Softs::FSolverReal MeshScale,
-			Softs::FSolverReal MaxDistancesScale = (Softs::FSolverReal)1.);
-
 		// NOTE: this only does something if using the PBDSolver. Force-based solver constraints are activated automatically
 		// when activating a particle range.
 		void Enable(bool bEnable);
 		// ---- End of Cloth interface ----
 
 		// ---- Debug functions ----
-		UE_DEPRECATED(5.4, "Shape constraints are no longer supported.")
-		const TSharedPtr<Softs::FPBDShapeConstraints>& GetShapeConstraints() const { return ShapeConstraints_Deprecated; }
-
 		const TSharedPtr<Softs::FPBDEdgeSpringConstraints>& GetEdgeSpringConstraints() const { return EdgeConstraints; }
 		const TSharedPtr<Softs::FXPBDEdgeSpringConstraints>& GetXEdgeSpringConstraints() const { return XEdgeConstraints; }
 		const TSharedPtr<Softs::FXPBDStretchBiasElementConstraints>& GetXStretchBiasConstraints() const { return XStretchBiasConstraints; }
@@ -138,6 +121,10 @@ namespace Chaos
 		const TSharedPtr<Softs::FPBDSoftBodyCollisionConstraint>& GetCollisionConstraint() const { return CollisionConstraint; }
 		const TSharedPtr<Softs::FPBDSkinnedTriangleMeshCollisionConstraints>& GetSkinnedTriangleCollisionsConstraint() const { return SkinnedTriangleCollisionsConstraint; }
 		const TSharedPtr<Softs::FMultiResConstraints>& GetMultiResConstraints() const { return MultiResConstraints; }
+		const TSharedPtr<Softs::FXPBDVertexConstraints>& GetClothVertexSpringConstraints() const { return ClothVertexSpringConstraints; };
+		const TSharedPtr<Softs::FXPBDVertexFaceConstraints>& GetClothVertexFaceSpringConstraints() const { return ClothVertexFaceSpringConstraints; }
+		const TSharedPtr<Softs::FXPBDFaceConstraints>& GetClothFaceSpringConstraints() const { return ClothFaceSpringConstraints; }
+		const TSharedPtr<Softs::FPBDRepulsionConstraints>& GetRepulsionConstraints() const { return RepulsionConstraints; }
 		// ---- End of debug functions ----
 
 	private:
@@ -198,17 +185,15 @@ namespace Chaos
 			const FTriangleMesh& TriangleMesh,
 			const FTriangleMesh* MultiResCoarseLODMesh,
 			const int32 MultiResCoarseLODParticleRangeId);
+		void CreateClothClothConstraints(
+			const Softs::FCollectionPropertyConstFacade& ConfigProperties,
+			const TSharedPtr<const FManagedArrayCollection>& ManagedArrayCollection);
 
 		void CreateForceBasedRules(const TSharedPtr<Softs::FMultiResConstraints>& FineLODMultiResConstraint);
 		void CreatePBDRules();
 		void CreateGSRules();
 		void GetGSNumRules();
 
-
-		//~ Begin deprecated constraints
-		TSharedPtr<Softs::FPBDShapeConstraints> ShapeConstraints_Deprecated;
-		//~ End deprecated constraints
-
 		TSharedPtr<Softs::FPBDEdgeSpringConstraints> EdgeConstraints;
 		TSharedPtr<Softs::FXPBDEdgeSpringConstraints> XEdgeConstraints;
 		TSharedPtr<Softs::FXPBDStretchBiasElementConstraints> XStretchBiasConstraints;
@@ -230,7 +215,6 @@ namespace Chaos
 		TSharedPtr<Softs::FPBDSelfCollisionSphereConstraints> SelfCollisionSphereConstraints;
 		TSharedPtr<Softs::FGaussSeidelMainConstraint<Softs::FSolverReal, Softs::FSolverParticles>> GSMainConstraint;
 		TSharedPtr<Softs::FGaussSeidelCorotatedCodimensionalConstraints<Softs::FSolverReal, Softs::FSolverParticles>> GSCorotatedCodimensionalConstraint;
-		TSharedPtr<Softs::TXPBDPlanarConstraints<Softs::FSolverReal>> XPlanarConstraints;
 		TSharedPtr<Softs::FMultiResConstraints> MultiResConstraints;
 		//~ Begin Force-based solver only constraints
 		Softs::FSolverVec3 SolverWindVelocity; // Set from solver and added to wind from the config
@@ -238,6 +222,10 @@ namespace Chaos
 		TSharedPtr<Softs::FExternalForces> ExternalForces;
 		TSharedPtr<Softs::FPBDSoftBodyCollisionConstraint> CollisionConstraint;
 		TSharedPtr<Softs::FPBDSkinnedTriangleMeshCollisionConstraints> SkinnedTriangleCollisionsConstraint;
+		TSharedPtr<Softs::FXPBDVertexConstraints> ClothVertexSpringConstraints;
+		TSharedPtr<Softs::FXPBDVertexFaceConstraints> ClothVertexFaceSpringConstraints;
+		TSharedPtr<Softs::FXPBDFaceConstraints> ClothFaceSpringConstraints;
+		TSharedPtr<Softs::FPBDRepulsionConstraints> RepulsionConstraints;
 		//~ End Force-based solver only constraints
 
 		// Exactly one of these should be non-null
diff --git a/Engine/Plugins/ChaosCloth/Source/ChaosCloth/Public/ChaosCloth/ChaosClothVisualization.h b/Engine/Plugins/ChaosCloth/Source/ChaosCloth/Public/ChaosCloth/ChaosClothVisualization.h
--- a/Engine/Plugins/ChaosCloth/Source/ChaosCloth/Public/ChaosCloth/ChaosClothVisualization.h
+++ b/Engine/Plugins/ChaosCloth/Source/ChaosCloth/Public/ChaosCloth/ChaosClothVisualization.h
@@ -58,6 +58,7 @@ namespace Chaos
 		CHAOSCLOTH_API void DrawGravity(FPrimitiveDrawInterface* PDI = nullptr) const;
 		CHAOSCLOTH_API void DrawFictitiousAngularForces(FPrimitiveDrawInterface* PDI = nullptr) const;
 		CHAOSCLOTH_API void DrawMultiResConstraint(FPrimitiveDrawInterface* PDI = nullptr) const;
+		CHAOSCLOTH_API void DrawClothClothConstraints(FPrimitiveDrawInterface* PDI = nullptr) const;
 
 #else  // #if CHAOS_DEBUG_DRAW
 		void SetSolver(const ::Chaos::FClothingSimulationSolver* /*InSolver*/) {}
@@ -89,6 +90,7 @@ namespace Chaos
 		void DrawGravity(FPrimitiveDrawInterface* /*PDI*/ = nullptr) const {}
 		void DrawFictitiousAngularForces(FPrimitiveDrawInterface* /*PDI*/ = nullptr) const {}
 		void DrawMultiResConstraint(FPrimitiveDrawInterface* /*PDI*/ = nullptr) const {}
+		void DrawClothClothConstraints(FPrimitiveDrawInterface* /*PDI*/ = nullptr) const {}
 
 #endif  // #if CHAOS_DEBUG_DRAW
 
diff --git a/Engine/Plugins/ChaosCloth/Source/ChaosCloth/Public/ChaosCloth/ChaosClothingSimulationMesh.h b/Engine/Plugins/ChaosCloth/Source/ChaosCloth/Public/ChaosCloth/ChaosClothingSimulationMesh.h
--- a/Engine/Plugins/ChaosCloth/Source/ChaosCloth/Public/ChaosCloth/ChaosClothingSimulationMesh.h
+++ b/Engine/Plugins/ChaosCloth/Source/ChaosCloth/Public/ChaosCloth/ChaosClothingSimulationMesh.h
@@ -9,6 +9,7 @@
 
 struct FMeshToMeshVertData;
 struct FClothVertBoneData;
+struct FManagedArrayCollection;
 
 #if UE_ENABLE_INCLUDE_ORDER_DEPRECATED_IN_5_5
 class USkeletalMeshComponent;
@@ -134,6 +135,12 @@ namespace Chaos
 		/** Return this mesh uniform scale as the maximum of the three axis scale value. */
 		virtual Softs::FSolverReal GetScale() const;
 
+		/** Return the associated ManagedArrayCollection for this LOD, if it exists. */
+		virtual TSharedPtr<const FManagedArrayCollection> GetManagedArrayCollection(int32 LODIndex) const
+		{
+			return TSharedPtr<const FManagedArrayCollection>();
+		}
+
 		/** Deform the specified positions to match the shape of the previous LOD. */
 		bool WrapDeformLOD(
 			int32 PrevLODIndex,
diff --git a/Engine/Plugins/ChaosClothAsset/Source/ChaosClothAsset/Private/ChaosClothAsset/ClothCollection.cpp b/Engine/Plugins/ChaosClothAsset/Source/ChaosClothAsset/Private/ChaosClothAsset/ClothCollection.cpp
--- a/Engine/Plugins/ChaosClothAsset/Source/ChaosClothAsset/Private/ChaosClothAsset/ClothCollection.cpp
+++ b/Engine/Plugins/ChaosClothAsset/Source/ChaosClothAsset/Private/ChaosClothAsset/ClothCollection.cpp
@@ -703,15 +703,15 @@ namespace UE::Chaos::ClothAsset
 	template CHAOSCLOTHASSET_API TArray<FName> FClothCollection::GetUserDefinedAttributeNames<FVector3f>(const FName& GroupName) const;
 
 	template<typename T UE_REQUIRES_DEFINITION(TIsUserAttributeType<T>::Value)>
-	TManagedArray<T>* FClothCollection::FindOrAddUserDefinedAttribute(const FName& Name, const FName& GroupName)
+	TManagedArray<T>* FClothCollection::FindOrAddUserDefinedAttribute(const FName& Name, const FName& GroupName, const FName& GroupDependency)
 	{
 		using namespace UE::Chaos::ClothAsset::Private;
 		if (const TArray<FName>* const FixedAttributeNames = FixedAttributeNamesMap.Find(GroupName))
 		{
-			// Only allow adding to know groups. Do not allow adding a name reserved as a fixed attribute.
+			// Only allow adding to known groups. Do not allow adding a name reserved as a fixed attribute.
 			if (!FixedAttributeNames->Contains(Name))
 			{
-				return ManagedArrayCollection->FindOrAddAttributeTyped<T>(Name, GroupName);
+				return ManagedArrayCollection->FindOrAddAttributeTyped<T>(Name, GroupName, FManagedArrayCollection::FConstructionParameters(GroupDependency));
 			}
 		}
 		if (const TArray<FName>* const UserAccessibleAttributeNames = UserAccessibleAttributeNamesMap.Find(GroupName))
@@ -724,10 +724,10 @@ namespace UE::Chaos::ClothAsset
 		}
 		return nullptr;
 	}
-	template CHAOSCLOTHASSET_API TManagedArray<bool>* FClothCollection::FindOrAddUserDefinedAttribute<bool>(const FName& Name, const FName& GroupName);
-	template CHAOSCLOTHASSET_API TManagedArray<int32>* FClothCollection::FindOrAddUserDefinedAttribute<int32>(const FName& Name, const FName& GroupName);
-	template CHAOSCLOTHASSET_API TManagedArray<float>* FClothCollection::FindOrAddUserDefinedAttribute<float>(const FName& Name, const FName& GroupName);
-	template CHAOSCLOTHASSET_API TManagedArray<FVector3f>* FClothCollection::FindOrAddUserDefinedAttribute<FVector3f>(const FName& Name, const FName& GroupName);
+	template CHAOSCLOTHASSET_API TManagedArray<bool>* FClothCollection::FindOrAddUserDefinedAttribute<bool>(const FName& Name, const FName& GroupName, const FName& GroupDependency);
+	template CHAOSCLOTHASSET_API TManagedArray<int32>* FClothCollection::FindOrAddUserDefinedAttribute<int32>(const FName& Name, const FName& GroupName, const FName& GroupDependency);
+	template CHAOSCLOTHASSET_API TManagedArray<float>* FClothCollection::FindOrAddUserDefinedAttribute<float>(const FName& Name, const FName& GroupName, const FName& GroupDependency);
+	template CHAOSCLOTHASSET_API TManagedArray<FVector3f>* FClothCollection::FindOrAddUserDefinedAttribute<FVector3f>(const FName& Name, const FName& GroupName, const FName& GroupDependency);
 
 	void FClothCollection::RemoveUserDefinedAttribute(const FName& Name, const FName& GroupName)
 	{
diff --git a/Engine/Plugins/ChaosClothAsset/Source/ChaosClothAsset/Private/ChaosClothAsset/ClothCollection.h b/Engine/Plugins/ChaosClothAsset/Source/ChaosClothAsset/Private/ChaosClothAsset/ClothCollection.h
--- a/Engine/Plugins/ChaosClothAsset/Source/ChaosClothAsset/Private/ChaosClothAsset/ClothCollection.h
+++ b/Engine/Plugins/ChaosClothAsset/Source/ChaosClothAsset/Private/ChaosClothAsset/ClothCollection.h
@@ -132,12 +132,12 @@ namespace UE::Chaos::ClothAsset
 		TArray<FName> GetUserDefinedAttributeNames(const FName& GroupName) const;
 
 		template<typename T UE_REQUIRES(TIsUserAttributeType<T>::Value)>
-		TManagedArray<T>* FindOrAddUserDefinedAttribute(const FName& Name, const FName& GroupName);
+		TManagedArray<T>* FindOrAddUserDefinedAttribute(const FName& Name, const FName& GroupName, const FName& GroupDependency = NAME_None);
 
 		template<typename T UE_REQUIRES(TIsUserAttributeType<T>::Value)>
-		void AddUserDefinedAttribute(const FName& Name, const FName& GroupName)
+		void AddUserDefinedAttribute(const FName& Name, const FName& GroupName, const FName& GroupDependency = NAME_None)
 		{
-			FindOrAddUserDefinedAttribute<T>(Name, GroupName);
+			FindOrAddUserDefinedAttribute<T>(Name, GroupName, GroupDependency);
 		}
 
 		void RemoveUserDefinedAttribute(const FName& Name, const FName& GroupName);
diff --git a/Engine/Plugins/ChaosClothAsset/Source/ChaosClothAsset/Private/ChaosClothAsset/ClothGeometryTools.cpp b/Engine/Plugins/ChaosClothAsset/Source/ChaosClothAsset/Private/ChaosClothAsset/ClothGeometryTools.cpp
--- a/Engine/Plugins/ChaosClothAsset/Source/ChaosClothAsset/Private/ChaosClothAsset/ClothGeometryTools.cpp
+++ b/Engine/Plugins/ChaosClothAsset/Source/ChaosClothAsset/Private/ChaosClothAsset/ClothGeometryTools.cpp
@@ -4,6 +4,7 @@
 #include "ChaosClothAsset/ClothCollectionGroup.h"
 #include "ChaosClothAsset/CollectionClothFacade.h"
 #include "ChaosClothAsset/CollectionClothSelectionFacade.h"
+#include "Chaos/CollectionEmbeddedSpringConstraintFacade.h"
 #include "Containers/Queue.h"
 #include "DynamicMesh/DynamicMesh3.h"
 #include "DynamicMesh/DynamicMeshAttributeSet.h"
@@ -1171,6 +1172,8 @@ namespace UE::Chaos::ClothAsset
 			Cloth.GetSeam(SeamIndex).ValidateSeam();
 		}
 #endif
+		::Chaos::Softs::FEmbeddedSpringFacade SpringFacade(ClothCollection.Get(), ClothCollectionGroup::SimVertices3D);
+		SpringFacade.CleanupAndCompactInvalidSprings();
 	}
 
 
diff --git a/Engine/Plugins/ChaosClothAsset/Source/ChaosClothAsset/Private/ChaosClothAsset/CollectionClothFacade.cpp b/Engine/Plugins/ChaosClothAsset/Source/ChaosClothAsset/Private/ChaosClothAsset/CollectionClothFacade.cpp
--- a/Engine/Plugins/ChaosClothAsset/Source/ChaosClothAsset/Private/ChaosClothAsset/CollectionClothFacade.cpp
+++ b/Engine/Plugins/ChaosClothAsset/Source/ChaosClothAsset/Private/ChaosClothAsset/CollectionClothFacade.cpp
@@ -1127,14 +1127,14 @@ namespace UE::Chaos::ClothAsset
 	}
 
 	template<typename T UE_REQUIRES_DEFINITION(TIsUserAttributeType<T>::Value)>
-	bool FCollectionClothFacade::AddUserDefinedAttribute(const FName& Name, const FName& GroupName)
+	bool FCollectionClothFacade::AddUserDefinedAttribute(const FName& Name, const FName& GroupName, const FName& GroupDependency)
 	{
 		check(IsValid());
-		return GetClothCollection()->FindOrAddUserDefinedAttribute<T>(Name, GroupName) != nullptr;
+		return GetClothCollection()->FindOrAddUserDefinedAttribute<T>(Name, GroupName, GroupDependency) != nullptr;
 	}
-	template CHAOSCLOTHASSET_API bool FCollectionClothFacade::AddUserDefinedAttribute<int32>(const FName& Name, const FName& GroupName);
-	template CHAOSCLOTHASSET_API bool FCollectionClothFacade::AddUserDefinedAttribute<float>(const FName& Name, const FName& GroupName);
-	template CHAOSCLOTHASSET_API bool FCollectionClothFacade::AddUserDefinedAttribute<FVector3f>(const FName& Name, const FName& GroupName);
+	template CHAOSCLOTHASSET_API bool FCollectionClothFacade::AddUserDefinedAttribute<int32>(const FName& Name, const FName& GroupName, const FName& GroupDependency);
+	template CHAOSCLOTHASSET_API bool FCollectionClothFacade::AddUserDefinedAttribute<float>(const FName& Name, const FName& GroupName, const FName& GroupDependency);
+	template CHAOSCLOTHASSET_API bool FCollectionClothFacade::AddUserDefinedAttribute<FVector3f>(const FName& Name, const FName& GroupName, const FName& GroupDependency);
 
 	void FCollectionClothFacade::RemoveUserDefinedAttribute(const FName& Name, const FName& GroupName)
 	{
diff --git a/Engine/Plugins/ChaosClothAsset/Source/ChaosClothAsset/Public/ChaosClothAsset/CollectionClothFacade.h b/Engine/Plugins/ChaosClothAsset/Source/ChaosClothAsset/Public/ChaosClothAsset/CollectionClothFacade.h
--- a/Engine/Plugins/ChaosClothAsset/Source/ChaosClothAsset/Public/ChaosClothAsset/CollectionClothFacade.h
+++ b/Engine/Plugins/ChaosClothAsset/Source/ChaosClothAsset/Public/ChaosClothAsset/CollectionClothFacade.h
@@ -353,7 +353,7 @@ namespace UE::Chaos::ClothAsset
 		//~ Other User-Defined Attributes (not instantiated for bools)
 		/** GroupName must be an existing group as defined in ClothCollectionGroup. Returns success */
 		template<typename T UE_REQUIRES(TIsUserAttributeType<T>::Value)>
-		bool AddUserDefinedAttribute(const FName& Name, const FName& GroupName);
+		bool AddUserDefinedAttribute(const FName& Name, const FName& GroupName, const FName& GroupDependency = NAME_None);
 		void RemoveUserDefinedAttribute(const FName& Name, const FName& GroupName);
 		template<typename T UE_REQUIRES(TIsUserAttributeType<T>::Value)>
 		TArrayView<T> GetUserDefinedAttribute(const FName& Name, const FName& GroupName);
diff --git a/Engine/Plugins/ChaosClothAsset/Source/ChaosClothAssetEngine/Private/ChaosClothAsset/ClothAsset.cpp b/Engine/Plugins/ChaosClothAsset/Source/ChaosClothAssetEngine/Private/ChaosClothAsset/ClothAsset.cpp
--- a/Engine/Plugins/ChaosClothAsset/Source/ChaosClothAssetEngine/Private/ChaosClothAsset/ClothAsset.cpp
+++ b/Engine/Plugins/ChaosClothAsset/Source/ChaosClothAssetEngine/Private/ChaosClothAsset/ClothAsset.cpp
@@ -2,11 +2,13 @@
 
 #include "ChaosClothAsset/ClothAsset.h"
 #include "ChaosClothAsset/ClothAssetBuilder.h"
-#include "ChaosClothAsset/CollectionClothFacade.h"
+#include "ChaosClothAsset/ClothAssetPrivate.h"
+#include "ChaosClothAsset/ClothCollectionGroup.h"
 #include "ChaosClothAsset/ClothComponent.h"
 #include "ChaosClothAsset/ClothGeometryTools.h"
-#include "ChaosClothAsset/ClothAssetPrivate.h"
 #include "ChaosClothAsset/ClothSimulationModel.h"
+#include "ChaosClothAsset/CollectionClothFacade.h"
+#include "Chaos/CollectionEmbeddedSpringConstraintFacade.h"
 #include "Chaos/CollectionPropertyFacade.h"
 #include "Animation/Skeleton.h"
 #if WITH_EDITORONLY_DATA
@@ -37,11 +39,11 @@
 
 namespace UE::Chaos::ClothAsset::Private
 {
-bool bClothCollectionOnlyCookPropertyFacade = true;
-FAutoConsoleVariableRef CVarClothCollectionOnlyCookPropertyFacade(
-	TEXT("p.ClothCollectionOnlyCookPropertyFacade"),
-	bClothCollectionOnlyCookPropertyFacade,
-	TEXT("Default setting for culling propertys on the cloth collection during the cook. Default[false]"));
+bool bClothCollectionOnlyCookRequiredFacades = true;
+FAutoConsoleVariableRef CVarClothCollectionOnlyCookRequiredFacades(
+	TEXT("p.ClothCollectionOnlyCookRequiredFacades"),
+	bClothCollectionOnlyCookRequiredFacades,
+	TEXT("Default setting for culling managed arrays on the cloth collection during the cook. Default[true]"));
 
 
 const TCHAR* MinLodQualityLevelCVarName = TEXT("p.ClothAsset.MinLodQualityLevel");
@@ -123,14 +125,26 @@ TArray<TSharedRef<FManagedArrayCollection>> TrimOnCook(const FString InAssetName
 {
 	int32 Index = 0;
 #if WITH_EDITORONLY_DATA
-	if (bClothCollectionOnlyCookPropertyFacade)
+	if (bClothCollectionOnlyCookRequiredFacades)
 	{
 		TArray<TSharedRef<FManagedArrayCollection>> OutputCollections;
 		for (TSharedRef<FManagedArrayCollection> ClothCollection : InClothCollections)
 		{
+			// Properties
 			TSharedPtr<FManagedArrayCollection> PropertyCollection(new FManagedArrayCollection());
 			::Chaos::Softs::FCollectionPropertyMutableFacade CollectionPropertyMutableFacade(PropertyCollection);
 			CollectionPropertyMutableFacade.Copy(*ClothCollection);
+
+			// Springs
+			const ::Chaos::Softs::FEmbeddedSpringFacade InEmbeddedSpringFacade(ClothCollection.Get(), UE::Chaos::ClothAsset::ClothCollectionGroup::SimVertices3D);
+			if (InEmbeddedSpringFacade.IsValid())
+			{
+				::Chaos::Softs::FEmbeddedSpringFacade EmbeddedSpringFacade(*PropertyCollection, ClothCollectionGroup::SimVertices3D);
+				EmbeddedSpringFacade.DefineSchema();
+				constexpr int32 VertexOffset = 0;
+				EmbeddedSpringFacade.Append(InEmbeddedSpringFacade, VertexOffset);
+			}
+
 			OutputCollections.Add(TSharedRef<FManagedArrayCollection>(PropertyCollection.ToSharedRef()));
 			UE_LOG(LogChaosClothAsset, Display, TEXT("TrimOnCook[ON] %s:[%d] [size:%d]"),
 				*InAssetName, Index++, PropertyCollection->GetAllocatedSize());
diff --git a/Engine/Plugins/ChaosClothAsset/Source/ChaosClothAssetEngine/Private/ChaosClothAsset/ClothSimulationMesh.cpp b/Engine/Plugins/ChaosClothAsset/Source/ChaosClothAssetEngine/Private/ChaosClothAsset/ClothSimulationMesh.cpp
--- a/Engine/Plugins/ChaosClothAsset/Source/ChaosClothAssetEngine/Private/ChaosClothAsset/ClothSimulationMesh.cpp
+++ b/Engine/Plugins/ChaosClothAsset/Source/ChaosClothAssetEngine/Private/ChaosClothAsset/ClothSimulationMesh.cpp
@@ -9,10 +9,12 @@
 
 namespace UE::Chaos::ClothAsset
 {
-	FClothSimulationMesh::FClothSimulationMesh(const FChaosClothSimulationModel& InClothSimulationModel, const FClothSimulationContext& InClothSimulationContext, const FString& DebugName)
+	FClothSimulationMesh::FClothSimulationMesh(const FChaosClothSimulationModel& InClothSimulationModel, const FClothSimulationContext& InClothSimulationContext,
+		const TArray<TSharedRef<const FManagedArrayCollection>>& InManagedArrayCollections, const FString& DebugName)
 		: ::Chaos::FClothingSimulationMesh(DebugName)
 		, ClothSimulationModel(InClothSimulationModel)
 		, ClothSimulationContext(InClothSimulationContext)
+		, ManagedArrayCollections(InManagedArrayCollections)
 	{
 #if CHAOS_DEBUG_DRAW
 		const int32 ReferenceBoneIndex = GetReferenceBoneIndex();
@@ -229,4 +231,9 @@ namespace UE::Chaos::ClothAsset
 	{
 		return ClothSimulationModel.ClothSimulationLodModels[LODIndex].LODTransitionDownData;
 	}
+
+	TSharedPtr<const FManagedArrayCollection> FClothSimulationMesh::GetManagedArrayCollection(int32 LODIndex) const
+	{
+		return ManagedArrayCollections.IsValidIndex(LODIndex) ? ManagedArrayCollections[LODIndex] : TSharedPtr<const FManagedArrayCollection>();
+	}
 }
diff --git a/Engine/Plugins/ChaosClothAsset/Source/ChaosClothAssetEngine/Private/ChaosClothAsset/ClothSimulationMesh.h b/Engine/Plugins/ChaosClothAsset/Source/ChaosClothAssetEngine/Private/ChaosClothAsset/ClothSimulationMesh.h
--- a/Engine/Plugins/ChaosClothAsset/Source/ChaosClothAssetEngine/Private/ChaosClothAsset/ClothSimulationMesh.h
+++ b/Engine/Plugins/ChaosClothAsset/Source/ChaosClothAssetEngine/Private/ChaosClothAsset/ClothSimulationMesh.h
@@ -13,7 +13,8 @@ namespace UE::Chaos::ClothAsset
 	class FClothSimulationMesh : public ::Chaos::FClothingSimulationMesh
 	{
 	public:
-		FClothSimulationMesh(const FChaosClothSimulationModel& InClothSimulationModel, const FClothSimulationContext& InClothSimulationContext, const FString& DebugName);
+		FClothSimulationMesh(const FChaosClothSimulationModel& InClothSimulationModel, const FClothSimulationContext& InClothSimulationContext, const TArray<TSharedRef<const FManagedArrayCollection>>& InManagedArrayCollections, const FString& DebugName);
+
 		virtual ~FClothSimulationMesh() override = default;
 
 		FClothSimulationMesh(const FClothSimulationMesh&) = delete;
@@ -51,10 +52,12 @@ namespace UE::Chaos::ClothAsset
 		virtual TConstArrayView<FClothVertBoneData> GetBoneData(int32 LODIndex) const override;
 		virtual TConstArrayView<FMeshToMeshVertData> GetTransitionUpSkinData(int32 LODIndex) const override;
 		virtual TConstArrayView<FMeshToMeshVertData> GetTransitionDownSkinData(int32 LODIndex) const override;
+		virtual TSharedPtr<const FManagedArrayCollection> GetManagedArrayCollection(int32 LODIndex) const override;
 		//~ End FClothingSimulationMesh Interface
 
 	private:
 		const FChaosClothSimulationModel& ClothSimulationModel;
 		const FClothSimulationContext& ClothSimulationContext;
+		const TArray<TSharedRef<const FManagedArrayCollection>>& ManagedArrayCollections;
 	};
 }
diff --git a/Engine/Plugins/ChaosClothAsset/Source/ChaosClothAssetEngine/Private/ChaosClothAsset/ClothSimulationProxy.cpp b/Engine/Plugins/ChaosClothAsset/Source/ChaosClothAssetEngine/Private/ChaosClothAsset/ClothSimulationProxy.cpp
--- a/Engine/Plugins/ChaosClothAsset/Source/ChaosClothAssetEngine/Private/ChaosClothAsset/ClothSimulationProxy.cpp
+++ b/Engine/Plugins/ChaosClothAsset/Source/ChaosClothAssetEngine/Private/ChaosClothAsset/ClothSimulationProxy.cpp
@@ -193,7 +193,7 @@ PRAGMA_ENABLE_DEPRECATION_WARNINGS
 			FString::Format(TEXT("{0}|{1}"), { ClothComponent.GetOwner()->GetActorNameOrLabel(), ClothComponent.GetName() }) :
 			ClothComponent.GetName();
 #endif
-		const int32 MeshIndex = Meshes.Emplace(MakeUnique<FClothSimulationMesh>(*ClothSimulationModel, *ClothSimulationContext, DebugName));
+		const int32 MeshIndex = Meshes.Emplace(MakeUnique<FClothSimulationMesh>(*ClothSimulationModel, *ClothSimulationContext, ClothAsset->GetClothCollections(), DebugName));
 
 		// Create collider simulation thread object
 		const FReferenceSkeleton* const ReferenceSkeleton = &ClothAsset->GetRefSkeleton();
@@ -545,6 +545,7 @@ PRAGMA_ENABLE_DEPRECATION_WARNINGS
 		static const TConsoleVariableData<bool>* const DebugDrawSelfIntersectionCVar = IConsoleManager::Get().FindTConsoleVariableDataBool(TEXT("p.ChaosCloth.DebugDrawSelfIntersection"));
 		static const TConsoleVariableData<bool>* const DebugDrawParticleIndicesCVar = IConsoleManager::Get().FindTConsoleVariableDataBool(TEXT("p.ChaosCloth.DebugDrawParticleIndices"));
 		static const TConsoleVariableData<bool>* const DebugDrawElementIndicesCVar = IConsoleManager::Get().FindTConsoleVariableDataBool(TEXT("p.ChaosCloth.DebugDrawElementIndices"));
+		static const TConsoleVariableData<bool>* const DebugDrawClothClothConstraintsCVar = IConsoleManager::Get().FindTConsoleVariableDataBool(TEXT("p.ChaosCloth.DebugDrawClothClothConstraints"));
 
 		if (DebugDrawLocalSpaceCVar && DebugDrawLocalSpaceCVar->GetValueOnAnyThread()) { Visualization->DrawLocalSpace(); }
 		if (DebugDrawBoundsCVar && DebugDrawBoundsCVar->GetValueOnAnyThread()) { Visualization->DrawBounds(); }
@@ -568,6 +569,10 @@ PRAGMA_ENABLE_DEPRECATION_WARNINGS
 		if (DebugDrawSelfIntersectionCVar && DebugDrawSelfIntersectionCVar->GetValueOnAnyThread()) { Visualization->DrawSelfIntersection(); }
 		if (DebugDrawParticleIndicesCVar && DebugDrawParticleIndicesCVar->GetValueOnAnyThread()) { Visualization->DrawParticleIndices(); }
 		if (DebugDrawElementIndicesCVar && DebugDrawElementIndicesCVar->GetValueOnAnyThread()) { Visualization->DrawElementIndices(); }
+		if (DebugDrawClothClothConstraintsCVar && DebugDrawClothClothConstraintsCVar->GetValueOnAnyThread())
+		{
+			Visualization->DrawClothClothConstraints();
+		}
 #endif  // #if CHAOS_DEBUG_DRAW
 	}
 
diff --git a/Engine/Plugins/ChaosClothAssetEditor/Source/ChaosClothAssetDataflowNodes/Private/ChaosClothAsset/ChaosClothAssetDataflowNodesModule.cpp b/Engine/Plugins/ChaosClothAssetEditor/Source/ChaosClothAssetDataflowNodes/Private/ChaosClothAsset/ChaosClothAssetDataflowNodesModule.cpp
--- a/Engine/Plugins/ChaosClothAssetEditor/Source/ChaosClothAssetDataflowNodes/Private/ChaosClothAsset/ChaosClothAssetDataflowNodesModule.cpp
+++ b/Engine/Plugins/ChaosClothAssetEditor/Source/ChaosClothAssetDataflowNodes/Private/ChaosClothAsset/ChaosClothAssetDataflowNodesModule.cpp
@@ -26,6 +26,8 @@
 #include "ChaosClothAsset/SimulationBackstopConfigNode.h"
 #include "ChaosClothAsset/SimulationBendingConfigNode.h"
 #include "ChaosClothAsset/SimulationBendingOverrideConfigNode.h"
+#include "ChaosClothAsset/SimulationClothVertexSpringConfigNode.h"
+#include "ChaosClothAsset/SimulationClothVertexFaceSpringConfigNode.h"
 #include "ChaosClothAsset/SimulationCollisionConfigNode.h"
 #include "ChaosClothAsset/SimulationDampingConfigNode.h"
 #include "ChaosClothAsset/SimulationDefaultConfigNode.h"
@@ -104,6 +106,8 @@ namespace UE::Chaos::ClothAsset
 			DATAFLOW_NODE_REGISTER_CREATION_FACTORY(FChaosClothAssetSimulationAerodynamicsConfigNode);
 			DATAFLOW_NODE_REGISTER_CREATION_FACTORY(FChaosClothAssetSimulationAnimDriveConfigNode);
 			DATAFLOW_NODE_REGISTER_CREATION_FACTORY(FChaosClothAssetSimulationBackstopConfigNode);
+			DATAFLOW_NODE_REGISTER_CREATION_FACTORY(FChaosClothAssetSimulationClothVertexSpringConfigNode);
+			DATAFLOW_NODE_REGISTER_CREATION_FACTORY(FChaosClothAssetSimulationClothVertexFaceSpringConfigNode);
 			DATAFLOW_NODE_REGISTER_CREATION_FACTORY(FChaosClothAssetSimulationCollisionConfigNode);
 			DATAFLOW_NODE_REGISTER_CREATION_FACTORY(FChaosClothAssetSimulationBendingConfigNode);
 			DATAFLOW_NODE_REGISTER_CREATION_FACTORY(FChaosClothAssetSimulationBendingOverrideConfigNode);
diff --git a/Engine/Plugins/ChaosClothAssetEditor/Source/ChaosClothAssetDataflowNodes/Private/ChaosClothAsset/MergeClothCollectionsNode.cpp b/Engine/Plugins/ChaosClothAssetEditor/Source/ChaosClothAssetDataflowNodes/Private/ChaosClothAsset/MergeClothCollectionsNode.cpp
--- a/Engine/Plugins/ChaosClothAssetEditor/Source/ChaosClothAssetDataflowNodes/Private/ChaosClothAsset/MergeClothCollectionsNode.cpp
+++ b/Engine/Plugins/ChaosClothAssetEditor/Source/ChaosClothAssetDataflowNodes/Private/ChaosClothAsset/MergeClothCollectionsNode.cpp
@@ -1,12 +1,14 @@
 // Copyright Epic Games, Inc. All Rights Reserved.
 
 #include "ChaosClothAsset/MergeClothCollectionsNode.h"
+#include "ChaosClothAsset/ClothCollectionGroup.h"
 #include "ChaosClothAsset/CollectionClothFacade.h"
 #include "ChaosClothAsset/CollectionClothSelectionFacade.h"
 #include "ChaosClothAsset/ClothDataflowTools.h"
 #include "Chaos/CollectionPropertyFacade.h"
 #include "Dataflow/DataflowInputOutput.h"
 #include "Engine/SkeletalMesh.h"
+#include "Chaos/CollectionEmbeddedSpringConstraintFacade.h"
 
 #include UE_INLINE_GENERATED_CPP_BY_NAME(MergeClothCollectionsNode)
 
@@ -38,6 +40,39 @@ namespace UE::Chaos::ClothAsset::Private
 		FVector4f PropertyBounds;
 	};
 
+	struct FConstraintMergeData
+	{
+		int32 VertexSpringConstraintIndex = INDEX_NONE;
+		int32 VertexFaceSpringConstraintIndex = INDEX_NONE;
+		int32 VertexFaceRepulsionConstraintIndex = INDEX_NONE;
+		int32 FaceSpringConstraintIndex = INDEX_NONE;
+
+		int32 OtherVertexSpringConstraintIndex = INDEX_NONE;
+		int32 OtherVertexFaceSpringConstraintIndex = INDEX_NONE;
+		int32 OtherVertexFaceRepulsionConstraintIndex = INDEX_NONE;
+		int32 OtherFaceSpringConstraintIndex = INDEX_NONE;
+
+		void ResetOtherData()
+		{
+			OtherVertexSpringConstraintIndex = INDEX_NONE;
+			OtherVertexFaceSpringConstraintIndex = INDEX_NONE; 
+			OtherVertexFaceRepulsionConstraintIndex = INDEX_NONE;
+			OtherFaceSpringConstraintIndex = INDEX_NONE;
+		}
+	};
+
+	static void FillWeightMap(const TArrayView<float> WeightMap, const TConstArrayView<float> InWeightMap, const FVector2f& PropertyBounds, const FVector2f& InPropertyBounds)
+	{
+		const bool bHasAlreadyValues = (InWeightMap.Num() > 0);
+		for (int32 VertexIndex = 0; VertexIndex < WeightMap.Num(); ++VertexIndex)
+		{
+			// If no values in the weight map we are using the low value
+			const float WeightMapValue = bHasAlreadyValues ? (InWeightMap[VertexIndex] * (InPropertyBounds[1] - InPropertyBounds[0]) +
+				InPropertyBounds[0]) : InPropertyBounds[0];
+			WeightMap[VertexIndex] = (WeightMapValue - PropertyBounds[0]) / (PropertyBounds[1] - PropertyBounds[0]);
+		}
+	}
+
 	/** Build weight maps for each properties if necessary */
 	static FString BuildWeightMaps(const FDataflowNode& DataflowNode,
 		const FCollectionClothConstFacade& InClothFacade, FCollectionClothFacade& OutClothFacade,
@@ -82,23 +117,12 @@ namespace UE::Chaos::ClothAsset::Private
 			OutClothFacade.AddWeightMap(FName(WeightMapName));
 			TArrayView<float> WeightMap = OutClothFacade.GetWeightMap(FName(WeightMapName));
 
-			auto FillWeightMap = [&PropertyBounds, &WeightMap](const TConstArrayView<float> InWeightMap, const FVector2f& InPropertyBounds,
-				const int32 InNumVertices, const int32 OutVertexOffset)
-			{
-				const bool bHasAlreadyValues = (InWeightMap.Num() > 0);
-				for(int32 VertexIndex = 0; VertexIndex < InNumVertices; ++VertexIndex)
-                {
-                	// If no values in the weight map we are using the low value
-					const float WeightMapValue = bHasAlreadyValues ? (InWeightMap[VertexIndex] * (InPropertyBounds[1] - InPropertyBounds[0]) +
-						InPropertyBounds[0]) : InPropertyBounds[0];
-					WeightMap[OutVertexOffset+VertexIndex] = (WeightMapValue - PropertyBounds[0]) / (PropertyBounds[1] - PropertyBounds[0]);
-                }
-			};
 			const int32 InNumVertices = InClothFacade.GetNumSimVertices3D();
 			const int32 OutNumVertices = OutClothFacade.GetNumSimVertices3D() - InNumVertices;
+			check(OutNumVertices >= 0);
 			
-			FillWeightMap(OutClothFacade.GetWeightMap(FName(OutWeightMapName)), OutPropertyBounds, OutNumVertices, 0);
-			FillWeightMap(InClothFacade.GetWeightMap(FName(InWeightMapName)), InPropertyBounds, InNumVertices, OutNumVertices);
+			FillWeightMap(WeightMap.Left(OutNumVertices), OutClothFacade.GetWeightMap(FName(OutWeightMapName)).Left(OutNumVertices), PropertyBounds, OutPropertyBounds);
+			FillWeightMap(WeightMap.Right(InNumVertices), InClothFacade.GetWeightMap(FName(InWeightMapName)), PropertyBounds, InPropertyBounds);
 		}
 		return WeightMapName;
 	}
@@ -140,6 +164,136 @@ namespace UE::Chaos::ClothAsset::Private
 		return PropertyBounds;
 	}
 
+	static const TArray<FName> VertexSpringConstraintPropertyNames =
+	{
+		TEXT("VertexSpringExtensionStiffness"),
+		TEXT("VertexSpringCompressionStiffness"),
+		TEXT("VertexSpringDamping"),
+	};
+
+	static const TArray<FName> VertexFaceSpringConstraintPropertyNames =
+	{
+		TEXT("VertexFaceSpringExtensionStiffness"),
+		TEXT("VertexFaceSpringCompressionStiffness"),
+		TEXT("VertexFaceSpringDamping"),
+	};
+
+	static const TArray<FName> FaceSpringConstraintPropertyNames =
+	{
+		TEXT("FaceSpringExtensionStiffness"),
+		TEXT("FaceSpringCompressionStiffness"),
+		TEXT("FaceSpringDamping"),
+	};
+
+	static bool IsSpringConstraintProperty(const FName& PropertyKey)
+	{
+		return VertexSpringConstraintPropertyNames.Contains(PropertyKey) || VertexFaceSpringConstraintPropertyNames.Contains(PropertyKey) || FaceSpringConstraintPropertyNames.Contains(PropertyKey);
+	}
+
+	static void UpdateSpringConstraintWeights(
+		const ::Chaos::Softs::FEmbeddedSpringFacade& InSpringFacade,
+		::Chaos::Softs::FEmbeddedSpringFacade& OutSpringFacade,
+		const FConstraintMergeData& ConstraintMergeData,
+		const FVector2f& InPropertyBounds, const FVector2f& OutPropertyBounds,
+		const FVector2f& PropertyBounds, const FString& PropertyName)
+	{
+		if (PropertyBounds[0] == PropertyBounds[1])
+		{
+			// If the low high values of the merged property are the same we don't need to build a weight map
+			return;
+		}
+		using namespace ::Chaos::Softs;
+		const FName PropertyNameName(*PropertyName);
+
+
+		TArrayView<float> WeightMap;
+		TConstArrayView<float> InWeightMap;
+		if (VertexSpringConstraintPropertyNames.Contains(PropertyNameName))
+		{
+			FEmbeddedSpringConstraintFacade OutConstraintFacade = OutSpringFacade.GetSpringConstraint(ConstraintMergeData.VertexSpringConstraintIndex);
+			check(ConstraintMergeData.OtherVertexSpringConstraintIndex != INDEX_NONE);
+			FEmbeddedSpringConstraintFacade InConstraintFacade = InSpringFacade.GetSpringConstraintConst(ConstraintMergeData.OtherVertexSpringConstraintIndex);
+
+			if (PropertyNameName == FName(TEXT("VertexSpringExtensionStiffness")))
+			{
+				WeightMap = OutConstraintFacade.GetExtensionStiffness();
+				InWeightMap = InConstraintFacade.GetExtensionStiffnessConst();
+			}
+			else if (PropertyNameName == FName(TEXT("VertexSpringCompressionStiffness")))
+			{
+				WeightMap = OutConstraintFacade.GetCompressionStiffness();
+				InWeightMap = InConstraintFacade.GetCompressionStiffnessConst();
+			}
+			else
+			{
+				check(PropertyNameName == FName(TEXT("VertexSpringDamping")));
+				WeightMap = OutConstraintFacade.GetDamping();
+				InWeightMap = InConstraintFacade.GetDampingConst();
+			}
+			const int32 InNumSprings = InConstraintFacade.GetNumSprings();
+			const int32 OutNumSprings = OutConstraintFacade.GetNumSprings() - InNumSprings;
+			check(OutNumSprings >= 0);
+			FillWeightMap(WeightMap.Left(OutNumSprings), WeightMap.Left(OutNumSprings), PropertyBounds, OutPropertyBounds);
+			FillWeightMap(WeightMap.Right(InNumSprings), InWeightMap, PropertyBounds, InPropertyBounds);
+		}
+		else if (VertexFaceSpringConstraintPropertyNames.Contains(PropertyNameName))
+		{
+			FEmbeddedSpringConstraintFacade OutConstraintFacade = OutSpringFacade.GetSpringConstraint(ConstraintMergeData.VertexFaceSpringConstraintIndex);
+			check(ConstraintMergeData.OtherVertexFaceSpringConstraintIndex != INDEX_NONE);
+			FEmbeddedSpringConstraintFacade InConstraintFacade = InSpringFacade.GetSpringConstraintConst(ConstraintMergeData.OtherVertexFaceSpringConstraintIndex);
+
+			if (PropertyNameName == FName(TEXT("VertexFaceSpringExtensionStiffness")))
+			{
+				WeightMap = OutConstraintFacade.GetExtensionStiffness();
+				InWeightMap = InConstraintFacade.GetExtensionStiffnessConst();
+			}
+			else if (PropertyNameName == FName(TEXT("VertexFaceSpringCompressionStiffness")))
+			{
+				WeightMap = OutConstraintFacade.GetCompressionStiffness();
+				InWeightMap = InConstraintFacade.GetCompressionStiffnessConst();
+			}
+			else
+			{
+				check(PropertyNameName == FName(TEXT("VertexFaceSpringDamping")));
+				WeightMap = OutConstraintFacade.GetDamping();
+				InWeightMap = InConstraintFacade.GetDampingConst();
+			}
+			const int32 InNumSprings = InConstraintFacade.GetNumSprings();
+			const int32 OutNumSprings = OutConstraintFacade.GetNumSprings() - InNumSprings;
+			check(OutNumSprings >= 0);
+			FillWeightMap(WeightMap.Left(OutNumSprings), WeightMap.Left(OutNumSprings), PropertyBounds, OutPropertyBounds);
+			FillWeightMap(WeightMap.Right(InNumSprings), InWeightMap, PropertyBounds, InPropertyBounds);
+		}
+		else
+		{
+			FEmbeddedSpringConstraintFacade OutConstraintFacade = OutSpringFacade.GetSpringConstraint(ConstraintMergeData.FaceSpringConstraintIndex);
+			check(ConstraintMergeData.OtherFaceSpringConstraintIndex != INDEX_NONE);
+			FEmbeddedSpringConstraintFacade InConstraintFacade = InSpringFacade.GetSpringConstraintConst(ConstraintMergeData.OtherFaceSpringConstraintIndex);
+
+			if (PropertyNameName == FName(TEXT("FaceSpringExtensionStiffness")))
+			{
+				WeightMap = OutConstraintFacade.GetExtensionStiffness();
+				InWeightMap = InConstraintFacade.GetExtensionStiffnessConst();
+			}
+			else if (PropertyNameName == FName(TEXT("FaceSpringCompressionStiffness")))
+			{
+				WeightMap = OutConstraintFacade.GetCompressionStiffness();
+				InWeightMap = InConstraintFacade.GetCompressionStiffnessConst();
+			}
+			else
+			{
+				check(PropertyNameName == FName(TEXT("FaceSpringDamping")));
+				WeightMap = OutConstraintFacade.GetDamping();
+				InWeightMap = InConstraintFacade.GetDampingConst();
+			}
+			const int32 InNumSprings = InConstraintFacade.GetNumSprings();
+			const int32 OutNumSprings = OutConstraintFacade.GetNumSprings() - InNumSprings;
+			check(OutNumSprings >= 0);
+			FillWeightMap(WeightMap.Left(OutNumSprings), WeightMap.Left(OutNumSprings), PropertyBounds, OutPropertyBounds);
+			FillWeightMap(WeightMap.Right(InNumSprings), InWeightMap, PropertyBounds, InPropertyBounds);
+		}
+	}
+
 	static ::Chaos::Softs::ECollectionPropertyFlags MergePropertyFlags(
 		const ::Chaos::Softs::FCollectionPropertyConstFacade& InPropertyFacade,
 			  ::Chaos::Softs::FCollectionPropertyMutableFacade& OutPropertyFacade,
@@ -190,7 +344,10 @@ namespace UE::Chaos::ClothAsset::Private
 		const FCollectionClothConstFacade& InClothFacade,
 			  FCollectionClothFacade& OutClothFacade,
 		const ::Chaos::Softs::FCollectionPropertyConstFacade& InPropertyFacade,
-			  ::Chaos::Softs::FCollectionPropertyMutableFacade& OutPropertyFacade)
+			  ::Chaos::Softs::FCollectionPropertyMutableFacade& OutPropertyFacade,
+		const ::Chaos::Softs::FEmbeddedSpringFacade& InSpringFacade,
+			  ::Chaos::Softs::FEmbeddedSpringFacade& OutSpringFacade,
+		const FConstraintMergeData& ConstraintMergeData)
 	{
 		const int32 InNumInKeys = InPropertyFacade.Num();
 		TMap<FString,FMergedProperty> MergedPropertyMaps;
@@ -209,25 +366,32 @@ namespace UE::Chaos::ClothAsset::Private
 			bool bOverrideProperty = true;
 			if(OutKeyIndex != INDEX_NONE)
 			{
-				if(InPropertyFacade.IsInterpolable(InKeyIndex))
+				if (InPropertyFacade.IsInterpolable(InKeyIndex))
 				{
 					// If it exists we compute the min of the property low values and the max of the property high values
 					const FVector2f InPropertyBounds = InPropertyFacade.GetWeightedFloatValue(InKeyIndex);
 					const FVector2f OutPropertyBounds = OutPropertyFacade.GetWeightedFloatValue(OutKeyIndex);
 					const FVector2f PropertyBounds = MergePropertyBounds(InPropertyBounds, OutPropertyBounds);
-					
+
 					const ::Chaos::Softs::ECollectionPropertyFlags PropertyFlags = MergePropertyFlags(
 						InPropertyFacade, OutPropertyFacade, InKeyIndex, OutKeyIndex, InPropertyFlags, InPropertyKey);
-					
-                    OutPropertyFacade.SetFlags(OutKeyIndex, PropertyFlags);
-                    OutPropertyFacade.SetWeightedFloatValue(OutKeyIndex, PropertyBounds);
-    
-                    // We keep the string value to be the one in the output if defined
-                    const FString WeightMapName = BuildWeightMaps(DataflowNode, InClothFacade, OutClothFacade,
-                    	InPropertyBounds, OutPropertyBounds, PropertyBounds, InPropertyKey,
+
+					OutPropertyFacade.SetFlags(OutKeyIndex, PropertyFlags);
+					OutPropertyFacade.SetWeightedFloatValue(OutKeyIndex, PropertyBounds);
+
+					if (IsSpringConstraintProperty(FName(*InPropertyKey)))
+					{
+						UpdateSpringConstraintWeights(InSpringFacade, OutSpringFacade, ConstraintMergeData, InPropertyBounds, OutPropertyBounds, PropertyBounds, InPropertyKey);
+					}
+					else
+					{
+						// We keep the string value to be the one in the output if defined
+						const FString WeightMapName = BuildWeightMaps(DataflowNode, InClothFacade, OutClothFacade,
+							InPropertyBounds, OutPropertyBounds, PropertyBounds, InPropertyKey,
 							InPropertyFacade.GetStringValue(InKeyIndex), OutPropertyFacade.GetStringValue(OutKeyIndex), MergedPropertyMaps);
 
-					OutPropertyFacade.SetStringValue(OutKeyIndex, WeightMapName);
+						OutPropertyFacade.SetStringValue(OutKeyIndex, WeightMapName);
+					}
 					bOverrideProperty = false;
 				}
 			}
@@ -370,6 +534,48 @@ void FChaosClothAssetMergeClothCollectionsNode_v2::Evaluate(UE::Dataflow::FConte
 		FCollectionClothSelectionFacade SelectionFacade(ClothCollection);
 		bAreAnyValid |= SelectionFacade.IsValid();
 
+		FEmbeddedSpringFacade SpringFacade(ClothCollection.Get(), ClothCollectionGroup::SimVertices3D);
+		Private::FConstraintMergeData ConstraintMergeData;
+		if (SpringFacade.IsValid())
+		{
+			for (int32 ConstraintIndex = 0; ConstraintIndex < SpringFacade.GetNumSpringConstraints(); ++ConstraintIndex)
+			{
+				const FEmbeddedSpringConstraintFacade ConstraintFacade = SpringFacade.GetSpringConstraintConst(ConstraintIndex);
+				const FUintVector2 EndPoints = ConstraintFacade.GetConstraintEndPointNumIndices();
+				if (EndPoints == FUintVector2(1, 1))
+				{
+					checkf(ConstraintMergeData.VertexSpringConstraintIndex == INDEX_NONE, TEXT("Multiple vertex spring constraints found"));
+					ConstraintMergeData.VertexSpringConstraintIndex = ConstraintIndex;
+					bAreAnyValid = true;
+				}
+				else if (EndPoints == FUintVector2(1, 3))
+				{
+					const FString& ConstraintName = ConstraintFacade.GetConstraintName();
+					if (ConstraintName == TEXT("VertexFaceRepulsionConstraint"))
+					{
+						checkf(ConstraintMergeData.VertexFaceRepulsionConstraintIndex == INDEX_NONE, TEXT("Multiple vertex-face repulsion constraints found"));
+						ConstraintMergeData.VertexFaceRepulsionConstraintIndex = ConstraintIndex;
+					}
+					else
+					{
+						checkf(ConstraintMergeData.VertexFaceSpringConstraintIndex == INDEX_NONE, TEXT("Multiple vertex-face spring constraints found"));
+						ConstraintMergeData.VertexFaceSpringConstraintIndex = ConstraintIndex;
+					}
+					bAreAnyValid = true;
+				}
+				else if (EndPoints == FUintVector2(3, 3))
+				{
+					checkf(ConstraintMergeData.FaceSpringConstraintIndex == INDEX_NONE, TEXT("Multiple face spring constraints found"));
+					ConstraintMergeData.FaceSpringConstraintIndex = ConstraintIndex;
+					bAreAnyValid = true;
+				}
+				else
+				{
+					checkf(false, TEXT("Unexpected spring constraint type found with end points (%d, %d)"), EndPoints[0], EndPoints[1]);
+				}
+			}
+		}
+
 		// Iterate through the inputs and append them to LOD 0
 		for (int32 InputIndex = 1; InputIndex < Collections.Num(); ++InputIndex)
 		{
@@ -392,6 +598,16 @@ void FChaosClothAssetMergeClothCollectionsNode_v2::Evaluate(UE::Dataflow::FConte
 				}
 			}
 
+			// Springs need to update with offsets. Gather offsets before appending cloth data.
+			const FEmbeddedSpringFacade OtherEmbeddedSpringFacade(OtherClothCollection.Get(), ClothCollectionGroup::SimVertices3D);
+			if (OtherEmbeddedSpringFacade.IsValid())
+			{
+				if (!GroupNameOffsets.Find(ClothCollectionGroup::SimVertices3D))
+				{
+					GroupNameOffsets.Add(ClothCollectionGroup::SimVertices3D) = ClothCollection->NumElements(ClothCollectionGroup::SimVertices3D);
+				}
+			}
+
 			// Append cloth
 			const FCollectionClothConstFacade OtherClothFacade(OtherClothCollection);
 			if (OtherClothFacade.IsValid() && Private::AreSkeletalMeshesCompatible(*this, ClothFacade, OtherClothFacade))
@@ -408,6 +624,101 @@ void FChaosClothAssetMergeClothCollectionsNode_v2::Evaluate(UE::Dataflow::FConte
 				bAreAnyValid = true;
 			}
 
+			// Append springs (with offsets)
+			ConstraintMergeData.ResetOtherData();
+			if (OtherEmbeddedSpringFacade.IsValid())
+			{
+				for (int32 ConstraintIndex = 0; ConstraintIndex < OtherEmbeddedSpringFacade.GetNumSpringConstraints(); ++ConstraintIndex)
+				{
+					const FEmbeddedSpringConstraintFacade OtherConstraintFacade = OtherEmbeddedSpringFacade.GetSpringConstraintConst(ConstraintIndex);
+					const FUintVector2 EndPoints = OtherConstraintFacade.GetConstraintEndPointNumIndices();
+					if (EndPoints == FUintVector2(1, 1))
+					{
+						checkf(ConstraintMergeData.OtherVertexSpringConstraintIndex == INDEX_NONE, TEXT("Multiple vertex spring constraints found"));
+						ConstraintMergeData.OtherVertexSpringConstraintIndex = ConstraintIndex;
+						if (ConstraintMergeData.VertexSpringConstraintIndex == INDEX_NONE)
+						{
+							// Create new constraint.
+							FEmbeddedSpringConstraintFacade NewConstraintFacade = SpringFacade.AddGetSpringConstraint();
+							NewConstraintFacade.Initialize(OtherConstraintFacade, GroupNameOffsets[ClothCollectionGroup::SimVertices3D]);
+							ConstraintMergeData.VertexSpringConstraintIndex = NewConstraintFacade.GetConstraintIndex();
+						}
+						else
+						{
+							// Append to existing constraint
+							FEmbeddedSpringConstraintFacade ConstraintFacade = SpringFacade.GetSpringConstraint(ConstraintMergeData.VertexSpringConstraintIndex);
+							ConstraintFacade.Append(OtherConstraintFacade, GroupNameOffsets[ClothCollectionGroup::SimVertices3D]);
+						}
+
+						bAreAnyValid = true;
+					}
+					else if (EndPoints == FUintVector2(1, 3))
+					{
+						const FString& ConstraintName = OtherConstraintFacade.GetConstraintName();
+						if (ConstraintName == TEXT("VertexFaceRepulsionConstraint"))
+						{
+							checkf(ConstraintMergeData.OtherVertexFaceRepulsionConstraintIndex == INDEX_NONE, TEXT("Multiple vertex-face repulsion constraints found"));
+							ConstraintMergeData.OtherVertexFaceRepulsionConstraintIndex = ConstraintIndex;
+							if (ConstraintMergeData.VertexFaceRepulsionConstraintIndex == INDEX_NONE)
+							{
+								// Create new constraint.
+								FEmbeddedSpringConstraintFacade NewConstraintFacade = SpringFacade.AddGetSpringConstraint();
+								NewConstraintFacade.Initialize(OtherConstraintFacade, GroupNameOffsets[ClothCollectionGroup::SimVertices3D]);
+								ConstraintMergeData.VertexFaceRepulsionConstraintIndex = NewConstraintFacade.GetConstraintIndex();
+							}
+							else
+							{
+								// Append to existing constraint
+								FEmbeddedSpringConstraintFacade ConstraintFacade = SpringFacade.GetSpringConstraint(ConstraintMergeData.VertexFaceRepulsionConstraintIndex);
+								ConstraintFacade.Append(OtherConstraintFacade, GroupNameOffsets[ClothCollectionGroup::SimVertices3D]);
+							}
+						}
+						else
+						{
+							checkf(ConstraintMergeData.OtherVertexFaceSpringConstraintIndex == INDEX_NONE, TEXT("Multiple vertex-face spring constraints found"));
+							ConstraintMergeData.OtherVertexFaceSpringConstraintIndex = ConstraintIndex;
+							if (ConstraintMergeData.VertexFaceSpringConstraintIndex == INDEX_NONE)
+							{
+								// Create new constraint.
+								FEmbeddedSpringConstraintFacade NewConstraintFacade = SpringFacade.AddGetSpringConstraint();
+								NewConstraintFacade.Initialize(OtherConstraintFacade, GroupNameOffsets[ClothCollectionGroup::SimVertices3D]);
+								ConstraintMergeData.VertexFaceSpringConstraintIndex = NewConstraintFacade.GetConstraintIndex();
+							}
+							else
+							{
+								// Append to existing constraint
+								FEmbeddedSpringConstraintFacade ConstraintFacade = SpringFacade.GetSpringConstraint(ConstraintMergeData.VertexFaceSpringConstraintIndex);
+								ConstraintFacade.Append(OtherConstraintFacade, GroupNameOffsets[ClothCollectionGroup::SimVertices3D]);
+							}
+						}
+						bAreAnyValid = true;
+					}
+					else if (EndPoints == FUintVector2(3, 3))
+					{
+						checkf(ConstraintMergeData.OtherFaceSpringConstraintIndex == INDEX_NONE, TEXT("Multiple face spring constraints found"));
+						ConstraintMergeData.OtherFaceSpringConstraintIndex = ConstraintIndex;
+						if (ConstraintMergeData.FaceSpringConstraintIndex == INDEX_NONE)
+						{
+							// Create new constraint.
+							FEmbeddedSpringConstraintFacade NewConstraintFacade = SpringFacade.AddGetSpringConstraint();
+							NewConstraintFacade.Initialize(OtherConstraintFacade, GroupNameOffsets[ClothCollectionGroup::SimVertices3D]);
+							ConstraintMergeData.FaceSpringConstraintIndex = NewConstraintFacade.GetConstraintIndex();
+						}
+						else
+						{
+							// Append to existing constraint
+							FEmbeddedSpringConstraintFacade ConstraintFacade = SpringFacade.GetSpringConstraint(ConstraintMergeData.FaceSpringConstraintIndex);
+							ConstraintFacade.Append(OtherConstraintFacade, GroupNameOffsets[ClothCollectionGroup::SimVertices3D]);
+						}
+						bAreAnyValid = true;
+					}
+					else
+					{
+						checkf(false, TEXT("Unexpected spring constraint type found with end points (%d, %d)"), EndPoints[0], EndPoints[1]);
+					}
+				}
+			}
+
 			// Copy properties
 			const FCollectionPropertyConstFacade OtherPropertyFacade(OtherClothCollection);
 			if (OtherPropertyFacade.IsValid())
@@ -421,7 +732,7 @@ void FChaosClothAssetMergeClothCollectionsNode_v2::Evaluate(UE::Dataflow::FConte
 				}
 				else
 				{
-					Private::AppendInputProperties(*this, OtherClothFacade, ClothFacade, OtherPropertyFacade, PropertyFacade);
+					Private::AppendInputProperties(*this, OtherClothFacade, ClothFacade, OtherPropertyFacade, PropertyFacade, OtherEmbeddedSpringFacade, SpringFacade, ConstraintMergeData);
 				}
 				bAreAnyValid = true;
 			}
@@ -562,6 +873,23 @@ void FChaosClothAssetMergeClothCollectionsNode::Evaluate(UE::Dataflow::FContext&
 		FCollectionClothSelectionFacade SelectionFacade(ClothCollection);
 		bAreAnyValid |= SelectionFacade.IsValid();
 
+		FEmbeddedSpringFacade SpringFacade(ClothCollection.Get(), ClothCollectionGroup::SimVertices3D);
+		Private::FConstraintMergeData ConstraintMergeData;
+		if (SpringFacade.IsValid())
+		{
+			bAreAnyValid = true;
+			if (ClothFacade.HasUserDefinedAttribute<int32>(TEXT("VertexSpringConstraintIndex"), ClothCollectionGroup::Lods))
+			{
+				ConstraintMergeData.VertexSpringConstraintIndex = ClothFacade.GetUserDefinedAttribute<int32>(TEXT("VertexSpringConstraintIndex"), ClothCollectionGroup::Lods)[0];
+				check(ConstraintMergeData.VertexSpringConstraintIndex == INDEX_NONE || (ConstraintMergeData.VertexSpringConstraintIndex >= 0 && ConstraintMergeData.VertexSpringConstraintIndex < SpringFacade.GetNumSpringConstraints()));
+			}
+			if (ClothFacade.HasUserDefinedAttribute<int32>(TEXT("VertexFaceSpringConstraintIndex"), ClothCollectionGroup::Lods))
+			{
+				ConstraintMergeData.VertexFaceSpringConstraintIndex = ClothFacade.GetUserDefinedAttribute<int32>(TEXT("VertexFaceSpringConstraintIndex"), ClothCollectionGroup::Lods)[0];
+				check(ConstraintMergeData.VertexFaceSpringConstraintIndex == INDEX_NONE || (ConstraintMergeData.VertexFaceSpringConstraintIndex >= 0 && ConstraintMergeData.VertexFaceSpringConstraintIndex < SpringFacade.GetNumSpringConstraints()));
+			}
+		}
+
 		// Iterate through the inputs and append them to LOD 0
 		const TArray<const FManagedArrayCollection*> Collections = GetCollections();
 		for (int32 InputIndex = 1; InputIndex < Collections.Num(); ++InputIndex)
@@ -585,6 +913,16 @@ void FChaosClothAssetMergeClothCollectionsNode::Evaluate(UE::Dataflow::FContext&
 				}
 			}
 
+			// Springs need to update with offsets. Gather offsets before appending cloth data.
+			const FEmbeddedSpringFacade OtherEmbeddedSpringFacade(OtherClothCollection.Get(), ClothCollectionGroup::SimVertices3D);
+			if (OtherEmbeddedSpringFacade.IsValid())
+			{
+				if (!GroupNameOffsets.Find(ClothCollectionGroup::SimVertices3D))
+				{
+					GroupNameOffsets.Add(ClothCollectionGroup::SimVertices3D) = ClothCollection->NumElements(ClothCollectionGroup::SimVertices3D);
+				}
+			}
+
 			// Append cloth
 			const FCollectionClothConstFacade OtherClothFacade(OtherClothCollection);
 			if (OtherClothFacade.IsValid() && Private::AreSkeletalMeshesCompatible(*this, ClothFacade, OtherClothFacade))
@@ -600,6 +938,53 @@ void FChaosClothAssetMergeClothCollectionsNode::Evaluate(UE::Dataflow::FContext&
 				SelectionFacade.AppendWithOffsets(OtherSelectionFacade, bUpdateExistingSelections, GroupNameOffsets);
 				bAreAnyValid = true;
 			}
+			// Append springs (with offsets)
+			ConstraintMergeData.ResetOtherData();
+			if (OtherEmbeddedSpringFacade.IsValid())
+			{
+				if (OtherClothFacade.HasUserDefinedAttribute<int32>(TEXT("VertexSpringConstraintIndex"), ClothCollectionGroup::Lods))
+				{
+					ConstraintMergeData.OtherVertexSpringConstraintIndex = OtherClothFacade.GetUserDefinedAttribute<int32>(TEXT("VertexSpringConstraintIndex"), ClothCollectionGroup::Lods)[0];
+					if (ConstraintMergeData.OtherVertexSpringConstraintIndex >= 0 && ConstraintMergeData.OtherVertexSpringConstraintIndex < OtherEmbeddedSpringFacade.GetNumSpringConstraints())
+					{
+						const FEmbeddedSpringConstraintFacade OtherConstraintFacade = OtherEmbeddedSpringFacade.GetSpringConstraintConst(ConstraintMergeData.OtherVertexSpringConstraintIndex);
+
+						if (ConstraintMergeData.VertexSpringConstraintIndex == INDEX_NONE)
+						{
+							FEmbeddedSpringConstraintFacade NewConstraintFacade = SpringFacade.AddGetSpringConstraint();
+							NewConstraintFacade.Initialize(OtherConstraintFacade, GroupNameOffsets[ClothCollectionGroup::SimVertices3D]);
+							ConstraintMergeData.VertexSpringConstraintIndex = NewConstraintFacade.GetConstraintIndex();
+						}
+						else
+						{
+							FEmbeddedSpringConstraintFacade NewConstraintFacade = SpringFacade.GetSpringConstraint(ConstraintMergeData.VertexSpringConstraintIndex);
+							NewConstraintFacade.Append(OtherConstraintFacade, GroupNameOffsets[ClothCollectionGroup::SimVertices3D]);
+						}
+						bAreAnyValid = true;
+					}
+				}
+				if (OtherClothFacade.HasUserDefinedAttribute<int32>(TEXT("VertexFaceSpringConstraintIndex"), ClothCollectionGroup::Lods))
+				{
+					ConstraintMergeData.OtherVertexFaceSpringConstraintIndex = OtherClothFacade.GetUserDefinedAttribute<int32>(TEXT("VertexFaceSpringConstraintIndex"), ClothCollectionGroup::Lods)[0];
+					if (ConstraintMergeData.OtherVertexFaceSpringConstraintIndex >= 0 && ConstraintMergeData.OtherVertexFaceSpringConstraintIndex < OtherEmbeddedSpringFacade.GetNumSpringConstraints())
+					{
+						const FEmbeddedSpringConstraintFacade OtherConstraintFacade = OtherEmbeddedSpringFacade.GetSpringConstraintConst(ConstraintMergeData.OtherVertexFaceSpringConstraintIndex);
+
+						if (ConstraintMergeData.VertexFaceSpringConstraintIndex == INDEX_NONE)
+						{
+							FEmbeddedSpringConstraintFacade NewConstraintFacade = SpringFacade.AddGetSpringConstraint();
+							NewConstraintFacade.Initialize(OtherConstraintFacade, GroupNameOffsets[ClothCollectionGroup::SimVertices3D]);
+							ConstraintMergeData.VertexFaceSpringConstraintIndex = NewConstraintFacade.GetConstraintIndex();
+						}
+						else
+						{
+							FEmbeddedSpringConstraintFacade NewConstraintFacade = SpringFacade.GetSpringConstraint(ConstraintMergeData.VertexFaceSpringConstraintIndex);
+							NewConstraintFacade.Append(OtherConstraintFacade, GroupNameOffsets[ClothCollectionGroup::SimVertices3D]);
+						}
+						bAreAnyValid = true;
+					}
+				}
+			}
 
 			// Copy properties
 			const FCollectionPropertyConstFacade OtherPropertyFacade(OtherClothCollection);
@@ -614,7 +999,7 @@ void FChaosClothAssetMergeClothCollectionsNode::Evaluate(UE::Dataflow::FContext&
 				}
 				else
 				{
-					Private::AppendInputProperties(*this, OtherClothFacade, ClothFacade, OtherPropertyFacade, PropertyFacade);
+					Private::AppendInputProperties(*this, OtherClothFacade, ClothFacade, OtherPropertyFacade, PropertyFacade, OtherEmbeddedSpringFacade, SpringFacade, ConstraintMergeData);
 				}
 				bAreAnyValid = true;
 			}
diff --git a/Engine/Plugins/ChaosClothAssetEditor/Source/ChaosClothAssetDataflowNodes/Private/ChaosClothAsset/SimulationBaseConfigNode.cpp b/Engine/Plugins/ChaosClothAssetEditor/Source/ChaosClothAssetDataflowNodes/Private/ChaosClothAsset/SimulationBaseConfigNode.cpp
--- a/Engine/Plugins/ChaosClothAssetEditor/Source/ChaosClothAssetDataflowNodes/Private/ChaosClothAsset/SimulationBaseConfigNode.cpp
+++ b/Engine/Plugins/ChaosClothAssetEditor/Source/ChaosClothAssetDataflowNodes/Private/ChaosClothAsset/SimulationBaseConfigNode.cpp
@@ -290,6 +290,18 @@ int32 FChaosClothAssetSimulationBaseConfigNode::FPropertyHelper::SetPropertyWeig
 	return PropertyKeyIndex;
 }
 
+int32 FChaosClothAssetSimulationBaseConfigNode::FPropertyHelper::SetPropertyWeighted(
+	const FName& PropertyName,
+	const FVector2f& PropertyValue,
+	const TArray<FName>& SimilarPropertyNames,
+	ECollectionPropertyFlags PropertyFlags)
+{
+	EnumAddFlags(PropertyFlags, ECollectionPropertyFlags::Interpolable);  // Interpolable
+	const int32 PropertyKeyIndex = ConfigNode.AddPropertyHelper(Properties, PropertyName, SimilarPropertyNames, PropertyFlags);
+	Properties.SetWeightedValue(PropertyKeyIndex, PropertyValue[0], PropertyValue[1]);
+	return PropertyKeyIndex;
+}
+
 int32 FChaosClothAssetSimulationBaseConfigNode::FPropertyHelper::SetPropertyWeighted(
 	const FName& PropertyName,
 	const FChaosClothAssetWeightedValue& PropertyValue,
diff --git a/Engine/Plugins/ChaosClothAssetEditor/Source/ChaosClothAssetDataflowNodes/Private/ChaosClothAsset/TerminalNode.cpp b/Engine/Plugins/ChaosClothAssetEditor/Source/ChaosClothAssetDataflowNodes/Private/ChaosClothAsset/TerminalNode.cpp
--- a/Engine/Plugins/ChaosClothAssetEditor/Source/ChaosClothAssetDataflowNodes/Private/ChaosClothAsset/TerminalNode.cpp
+++ b/Engine/Plugins/ChaosClothAssetEditor/Source/ChaosClothAssetDataflowNodes/Private/ChaosClothAsset/TerminalNode.cpp
@@ -9,6 +9,7 @@
 #include "ChaosClothAsset/CollectionClothFacade.h"
 #include "ChaosClothAsset/CollectionClothSelectionFacade.h"
 #include "Animation/Skeleton.h"
+#include "Chaos/CollectionEmbeddedSpringConstraintFacade.h"
 #include "Chaos/CollectionPropertyFacade.h"
 #include "Dataflow/DataflowInputOutput.h"
 #include "Engine/SkinnedAssetCommon.h"
@@ -57,6 +58,11 @@ namespace UE::Chaos::ClothAsset::Private
 					}
 				}
 			}
+			const ::Chaos::Softs::FEmbeddedSpringFacade SpringFacade(const_cast<const FManagedArrayCollection&>(ClothCollection.Get()), ClothCollectionGroup::SimVertices3D);
+			if (SpringFacade.IsValid())
+			{
+				Checksum = SpringFacade.CalculateTypeHash(Checksum);
+			}
 		}
 		return Checksum;
 	}
@@ -183,7 +189,8 @@ void FChaosClothAssetTerminalNode_v2::SetAssetValue(TObjectPtr<UObject> Asset, U
 
 			// Set properties
 			constexpr bool bUpdateExistingProperties = false;
-			Chaos::Softs::FCollectionPropertyMutableFacade(ClothCollection).Append(InClothCollections[LodIndex].ToSharedPtr(), bUpdateExistingProperties);
+			Chaos::Softs::FCollectionPropertyMutableFacade PropertyFacade(ClothCollection);
+			PropertyFacade.Append(InClothCollections[LodIndex].ToSharedPtr(), bUpdateExistingProperties);
 
 			// Set selections
 			FCollectionClothSelectionFacade Selection(ClothCollection);
@@ -203,6 +210,15 @@ void FChaosClothAssetTerminalNode_v2::SetAssetValue(TObjectPtr<UObject> Asset, U
 				}
 			}
 
+			// Set springs
+			Chaos::Softs::FEmbeddedSpringFacade EmbeddedSpringFacade(ClothCollection.Get(), ClothCollectionGroup::SimVertices3D);
+			const Chaos::Softs::FEmbeddedSpringFacade InEmbeddedSpringFacade(InClothCollections[LodIndex].Get(), ClothCollectionGroup::SimVertices3D);
+			if (InEmbeddedSpringFacade.IsValid())
+			{
+				EmbeddedSpringFacade.DefineSchema();
+				constexpr int32 VertexOffset = 0;
+				EmbeddedSpringFacade.Append(InEmbeddedSpringFacade, VertexOffset);
+			}
 
 			// Set physics asset only with LOD 0 at the moment
 			if (LodIndex == 0)
diff --git a/Engine/Plugins/ChaosClothAssetEditor/Source/ChaosClothAssetDataflowNodes/Public/ChaosClothAsset/SimulationBaseConfigNode.h b/Engine/Plugins/ChaosClothAssetEditor/Source/ChaosClothAssetDataflowNodes/Public/ChaosClothAsset/SimulationBaseConfigNode.h
--- a/Engine/Plugins/ChaosClothAssetEditor/Source/ChaosClothAssetDataflowNodes/Public/ChaosClothAsset/SimulationBaseConfigNode.h
+++ b/Engine/Plugins/ChaosClothAssetEditor/Source/ChaosClothAssetDataflowNodes/Public/ChaosClothAsset/SimulationBaseConfigNode.h
@@ -74,6 +74,7 @@ protected:
 		template<typename T UE_REQUIRES(std::is_base_of_v<FChaosClothAssetSimulationBaseConfigNode, T>)>
 		inline int32 SetPropertyString(const T* ConfigStruct, const FString* PropertyValue, const TArray<FName>& SimilarPropertyNames = {}, ECollectionPropertyFlags PropertyFlags = ECollectionPropertyFlags::Animatable);
 
+		int32 SetPropertyWeighted(const FName& PropertyName, const FVector2f& PropertyValue, const TArray<FName>& SimilarPropertyNames = {}, ECollectionPropertyFlags PropertyFlags = ECollectionPropertyFlags::None);
 		int32 SetPropertyWeighted(const FName& PropertyName, const FChaosClothAssetWeightedValue& PropertyValue, const TArray<FName>& SimilarPropertyNames = {}, ECollectionPropertyFlags PropertyFlags = ECollectionPropertyFlags::None);
 		int32 SetPropertyWeighted(const FName& PropertyName, const FChaosClothAssetWeightedValueNonAnimatable& PropertyValue, const TArray<FName>& SimilarPropertyNames = {}, ECollectionPropertyFlags PropertyFlags = ECollectionPropertyFlags::None);
 		int32 SetPropertyWeighted(const FName& PropertyName, const FChaosClothAssetWeightedValueNonAnimatableNoLowHighRange& PropertyValue, const TArray<FName>& SimilarPropertyNames = {}, ECollectionPropertyFlags PropertyFlags = ECollectionPropertyFlags::None);
diff --git a/Engine/Plugins/ChaosClothAssetEditor/Source/ChaosClothAssetEditor/Private/ChaosClothAsset/ClothEditorSimulationVisualization.cpp b/Engine/Plugins/ChaosClothAssetEditor/Source/ChaosClothAssetEditor/Private/ChaosClothAsset/ClothEditorSimulationVisualization.cpp
--- a/Engine/Plugins/ChaosClothAssetEditor/Source/ChaosClothAssetEditor/Private/ChaosClothAsset/ClothEditorSimulationVisualization.cpp
+++ b/Engine/Plugins/ChaosClothAssetEditor/Source/ChaosClothAssetEditor/Private/ChaosClothAsset/ClothEditorSimulationVisualization.cpp
@@ -391,6 +391,13 @@ const FVisualizationOption FVisualizationOption::OptionData[] =
 		}),
 		LOCTEXT("ChaosVisName_DrawKinematicColliderWired", "Draw Kinematic Colliders (Wireframe)"), 
 		LOCTEXT("ChaosVisName_DrawKinematicColliderWired_Tooltip", "Draw kinematic cloth colliders in wireframe.")),
+	FVisualizationOption(
+		FClothVisualizationDebugDraw::CreateLambda([](const FClothEditorSimulationVisualization&, const ::Chaos::FClothVisualizationNoGC& Visualization, FPrimitiveDrawInterface* PDI)
+		{
+			Visualization.DrawClothClothConstraints(PDI);
+		}),
+		LOCTEXT("ChaosVisName_DrawClothClothConstraints", "Cloth-Cloth constraints"), 
+		LOCTEXT("ChaosVisName_DrawClothClothConstraints_Tooltip", "Draw cloth-cloth constraints.")),
 	FVisualizationOption(
 		FClothVisualizationDebugDraw::CreateLambda([](const FClothEditorSimulationVisualization&, const ::Chaos::FClothVisualizationNoGC& Visualization, FPrimitiveDrawInterface* PDI)
 		{
diff --git a/Engine/Source/Runtime/Experimental/Chaos/Private/Chaos/CollectionEmbeddedSpringConstraintFacade.cpp b/Engine/Source/Runtime/Experimental/Chaos/Private/Chaos/CollectionEmbeddedSpringConstraintFacade.cpp
--- a/Engine/Source/Runtime/Experimental/Chaos/Private/Chaos/CollectionEmbeddedSpringConstraintFacade.cpp
+++ b/Engine/Source/Runtime/Experimental/Chaos/Private/Chaos/CollectionEmbeddedSpringConstraintFacade.cpp
@@ -140,6 +140,7 @@ namespace Chaos::Softs
 	const FName FEmbeddedSpringBaseFacade::ConstraintStart(TEXT("ConstraintStart"));
 	const FName FEmbeddedSpringBaseFacade::ConstraintEnd(TEXT("ConstraintEnd"));
 	const FName FEmbeddedSpringBaseFacade::ConstraintEndPointNumIndices(TEXT("ConstraintEndPointNumIndices"));
+	const FName FEmbeddedSpringBaseFacade::ConstraintName(TEXT("ConstraintName"));
 
 	FEmbeddedSpringBaseFacade::FEmbeddedSpringBaseFacade(FManagedArrayCollection& InCollection, const FName& InVerticesGroup)
 		: FPositionTargetFacade(InCollection, InVerticesGroup)
@@ -148,6 +149,7 @@ namespace Chaos::Softs
 		, ConstraintStartAttribute(InCollection, ConstraintStart, SpringConstraintGroupName, GroupName)
 		, ConstraintEndAttribute(InCollection, ConstraintEnd, SpringConstraintGroupName, GroupName)
 		, ConstraintEndPointNumIndicesAttribute(InCollection, ConstraintEndPointNumIndices, SpringConstraintGroupName)
+		, ConstraintNameAttribute(InCollection, ConstraintName, SpringConstraintGroupName)
 
 	{
 		DefineSchema();
@@ -160,6 +162,7 @@ namespace Chaos::Softs
 		, ConstraintStartAttribute(InCollection, ConstraintStart, SpringConstraintGroupName, GroupName)
 		, ConstraintEndAttribute(InCollection, ConstraintEnd, SpringConstraintGroupName, GroupName)
 		, ConstraintEndPointNumIndicesAttribute(InCollection, ConstraintEndPointNumIndices, SpringConstraintGroupName)
+		, ConstraintNameAttribute(InCollection, ConstraintName, SpringConstraintGroupName)
 	{
 	}
 
@@ -168,7 +171,7 @@ namespace Chaos::Softs
 		return FPositionTargetFacade::IsValid()
 			&& SpringLengthAttribute.IsValid() && CompressionStiffnessAttribute.IsValid()
 			&& ConstraintStartAttribute.IsValid() && ConstraintEndAttribute.IsValid()
-			&& ConstraintEndPointNumIndicesAttribute.IsValid();
+			&& ConstraintEndPointNumIndicesAttribute.IsValid() && ConstraintNameAttribute.IsValid();
 	}
 
 	void FEmbeddedSpringBaseFacade::DefineSchema()
@@ -182,6 +185,7 @@ namespace Chaos::Softs
 		ConstraintStartAttribute.Add();
 		ConstraintEndAttribute.Add();
 		ConstraintEndPointNumIndicesAttribute.Add();
+		ConstraintNameAttribute.Add();
 	}
 
 	void FEmbeddedSpringFacade::Append(const FEmbeddedSpringFacade& Other, int32 VertexOffset)
@@ -289,6 +293,7 @@ namespace Chaos::Softs
 		check(!IsConst());
 		ConstraintStartAttribute.ModifyAt(ConstraintIndex, INDEX_NONE);
 		ConstraintEndAttribute.ModifyAt(ConstraintIndex, INDEX_NONE);
+		SetConstraintName(TEXT(""));
 	}
 
 	void FEmbeddedSpringConstraintFacade::Reset()
@@ -426,26 +431,39 @@ namespace Chaos::Softs
 	}
 
 	void FEmbeddedSpringConstraintFacade::Initialize(const TConstArrayView<FIntVector2>& EndPoints, const TConstArrayView<float>& InSpringLength,
-		const TConstArrayView<float>& InExtensionStiffnessWeight, const TConstArrayView<float>& InCompressionStiffnessWeight, const TConstArrayView<float>& InDampingWeight)
+		const TConstArrayView<float>& InExtensionStiffnessWeight, const TConstArrayView<float>& InCompressionStiffnessWeight, const TConstArrayView<float>& InDampingWeight, const FString& InConstraintName)
 	{
 		check(!IsConst());
 		Reset();
 
 		//~ Spring Constraint Group
 		SetConstraintEndPointNumIndices(FUintVector2(1, 1));
+		SetConstraintName(InConstraintName);
+		Append(EndPoints, InSpringLength, InExtensionStiffnessWeight, InCompressionStiffnessWeight, InDampingWeight);
+	}
+
+	void FEmbeddedSpringConstraintFacade::Append(const TConstArrayView<FIntVector2>&EndPoints, const TConstArrayView<float>&InSpringLength,
+		const TConstArrayView<float>&InExtensionStiffnessWeight, const TConstArrayView<float>&InCompressionStiffnessWeight, const TConstArrayView<float>&InDampingWeight)
+	{
+		if (GetConstraintEndPointNumIndices() != FUintVector2(1, 1))
+		{
+			return;
+		}
+
+		const int32 OrigNumSprings = GetNumSprings();
 
 		//~ Group PositionTargets
 		const int32 InNumSprings = EndPoints.Num();
 		check(InNumSprings == InSpringLength.Num());
-		SetNumSprings(InNumSprings);
-		TArrayView<TArray<int32>> SourceIndexArray = GetSourceIndex();
-		TArrayView<TArray<int32>> TargetIndexArray = GetTargetIndex();
-		TArrayView<TArray<float>> SourceWeightArray = GetSourceWeights();
-		TArrayView<TArray<float>> TargetWeightArray = GetTargetWeights();
-		TArrayView<float> SpringLengthArray = GetSpringLength();
-		TArrayView<float> ExtensionStiffnessArray = GetExtensionStiffness();
-		TArrayView<float> CompressionStiffnessArray = GetCompressionStiffness();
-		TArrayView<float> DampingArray = GetDamping();
+		SetNumSprings(OrigNumSprings + InNumSprings);
+		TArrayView<TArray<int32>> SourceIndexArray = GetSourceIndex().Right(InNumSprings);
+		TArrayView<TArray<int32>> TargetIndexArray = GetTargetIndex().Right(InNumSprings);
+		TArrayView<TArray<float>> SourceWeightArray = GetSourceWeights().Right(InNumSprings);
+		TArrayView<TArray<float>> TargetWeightArray = GetTargetWeights().Right(InNumSprings);
+		TArrayView<float> SpringLengthArray = GetSpringLength().Right(InNumSprings);
+		TArrayView<float> ExtensionStiffnessArray = GetExtensionStiffness().Right(InNumSprings);
+		TArrayView<float> CompressionStiffnessArray = GetCompressionStiffness().Right(InNumSprings);
+		TArrayView<float> DampingArray = GetDamping().Right(InNumSprings);
 		for (int32 SpringIdx = 0; SpringIdx < InNumSprings; ++SpringIdx)
 		{
 			SourceIndexArray[SpringIdx] = { EndPoints[SpringIdx][0] };
@@ -460,15 +478,27 @@ namespace Chaos::Softs
 	}
 
 	void FEmbeddedSpringConstraintFacade::Initialize(const FUintVector2& EndPointNumIndices, const TConstArrayView<TArray<int32>>& InSourceIndices, const TConstArrayView<TArray<float>>& InSourceWeights,
-		const TConstArrayView<TArray<int32>>& InTargetIndices, const TConstArrayView<TArray<float>>& InTargetWeights, 
-		const TConstArrayView<float>& InSpringLength, 
-		const TConstArrayView<float>& InExtensionStiffnessWeight, const TConstArrayView<float>& InCompressionStiffnessWeight, const TConstArrayView<float>& InDampingWeight)
+		const TConstArrayView<TArray<int32>>& InTargetIndices, const TConstArrayView<TArray<float>>& InTargetWeights,
+		const TConstArrayView<float>& InSpringLength,
+		const TConstArrayView<float>& InExtensionStiffnessWeight, const TConstArrayView<float>& InCompressionStiffnessWeight, const TConstArrayView<float>& InDampingWeight, const FString& InConstraintName)
 	{
 		check(!IsConst());
 		Reset();
 
 		//~ Spring Constraint Group
 		SetConstraintEndPointNumIndices(EndPointNumIndices);
+		SetConstraintName(InConstraintName);
+
+		Append(InSourceIndices, InSourceWeights, InTargetIndices, InTargetWeights, InSpringLength, InExtensionStiffnessWeight, InCompressionStiffnessWeight, InDampingWeight);
+	}
+
+	void FEmbeddedSpringConstraintFacade::Append(const TConstArrayView<TArray<int32>>& InSourceIndices, const TConstArrayView<TArray<float>>& InSourceWeights,
+		const TConstArrayView<TArray<int32>>& InTargetIndices, const TConstArrayView<TArray<float>>& InTargetWeights,
+		const TConstArrayView<float>& InSpringLength,
+		const TConstArrayView<float>& InExtensionStiffnessWeight, const TConstArrayView<float>& InCompressionStiffnessWeight, const TConstArrayView<float>& InDampingWeight)
+	{
+
+		const int32 OrigNumSprings = GetNumSprings();
 
 		//~ Group PositionTargets
 		const int32 InNumSprings = InSourceIndices.Num();
@@ -476,14 +506,14 @@ namespace Chaos::Softs
 		check(InNumSprings == InTargetIndices.Num());
 		check(InNumSprings == InTargetWeights.Num());
 		check(InNumSprings == InSpringLength.Num());
-		SetNumSprings(InNumSprings);
-		TArrayView<float> SpringLengthArray = GetSpringLength();
-		TArrayView<float> ExtensionStiffnessArray = GetExtensionStiffness();
-		TArrayView<float> CompressionStiffnessArray = GetCompressionStiffness();
-		TArrayView<float> DampingArray = GetDamping();
+		SetNumSprings(OrigNumSprings + InNumSprings);
+		TArrayView<float> SpringLengthArray = GetSpringLength().Right(InNumSprings);
+		TArrayView<float> ExtensionStiffnessArray = GetExtensionStiffness().Right(InNumSprings);
+		TArrayView<float> CompressionStiffnessArray = GetCompressionStiffness().Right(InNumSprings);
+		TArrayView<float> DampingArray = GetDamping().Right(InNumSprings);
 		for (int32 SpringIdx = 0; SpringIdx < InNumSprings; ++SpringIdx)
 		{
-			SetIndicesAndWeights(SpringIdx, TConstArrayView<int32>(InSourceIndices[SpringIdx]), TConstArrayView<float>(InSourceWeights[SpringIdx]),
+			SetIndicesAndWeights(SpringIdx + OrigNumSprings, TConstArrayView<int32>(InSourceIndices[SpringIdx]), TConstArrayView<float>(InSourceWeights[SpringIdx]),
 				TConstArrayView<int32>(InTargetIndices[SpringIdx]), TConstArrayView<float>(InTargetWeights[SpringIdx]));
 			SpringLengthArray[SpringIdx] = InSpringLength[SpringIdx];
 			ExtensionStiffnessArray[SpringIdx] = InExtensionStiffnessWeight.IsValidIndex(SpringIdx) ? InExtensionStiffnessWeight[SpringIdx] : 0.f;
@@ -500,6 +530,7 @@ namespace Chaos::Softs
 
 		//~ Spring Constraint Group
 		SetConstraintEndPointNumIndices(Other.GetConstraintEndPointNumIndices());
+		SetConstraintName(Other.GetConstraintName());
 
 		//~ Group PositionTargets
 		SetNumSprings(Other.GetNumSprings());
