Cloth Remesh Node: don't try remeshing sim mesh if the 2D sim mesh area is close to zero

diff --git a/Engine/Plugins/ChaosClothAssetEditor/Source/ChaosClothAssetDataflowNodes/Private/ChaosClothAsset/RemeshNode.cpp b/Engine/Plugins/ChaosClothAssetEditor/Source/ChaosClothAssetDataflowNodes/Private/ChaosClothAsset/RemeshNode.cpp
--- a/Engine/Plugins/ChaosClothAssetEditor/Source/ChaosClothAssetDataflowNodes/Private/ChaosClothAsset/RemeshNode.cpp
+++ b/Engine/Plugins/ChaosClothAssetEditor/Source/ChaosClothAssetDataflowNodes/Private/ChaosClothAsset/RemeshNode.cpp
@@ -970,6 +970,7 @@ void FChaosClothAssetRemeshNode::RemeshSimMesh(const TSharedRef<const FManagedAr
 	Converter.Convert(ClothCollection, INDEX_NONE, EClothPatternVertexType::Sim2D, Mesh2D);
 
 	const double TotalArea = TMeshQueries<FDynamicMesh3>::GetVolumeArea(Mesh2D).Y;
+	checkf(TotalArea > 0, TEXT("FChaosClothAssetRemeshNode::RemeshSimMesh: Expected 2D Sim mesh to have a positive area"));
 	const int TriangleCount = Mesh2D.TriangleCount();
 
 	// Copy pattern IDs into polygroup layer
@@ -1640,10 +1641,26 @@ void FChaosClothAssetRemeshNode::Evaluate(UE::Dataflow::FContext& Context, const
 		if (ClothFacade.IsValid())  // Can only act on the collection if it is a valid cloth collection
 		{
 			if (bRemeshSim)
-			{
-				EmptySimSelections(OutputClothCollectionRef);
-				RemeshSimMesh(ClothCollection, OutputClothCollectionRef);
-				RebuildTopologyDependentSimData(ClothCollection, OutputClothCollectionRef);
+			{			
+				const float TotalArea2D = Algo::Accumulate(ClothFacade.GetSimIndices2D(), 0.0f, [&ClothFacade](float PreviousSum, const FIntVector3& Tri)
+				{
+					const TArrayView<const FVector2f>& Vertices2D = ClothFacade.GetSimPosition2D();
+					const float TriArea = FVector2f::CrossProduct(Vertices2D[Tri[1]] - Vertices2D[Tri[0]], Vertices2D[Tri[2]] - Vertices2D[Tri[0]]);
+					return PreviousSum + FMath::Abs(TriArea);
+				});
+
+				if (TotalArea2D < UE_SMALL_NUMBER)
+				{
+					FClothDataflowTools::LogAndToastWarning(*this,
+						LOCTEXT("SimMeshAreaZeroHeadline", "2D Sim Mesh Has Zero Area"),
+						LOCTEXT("SimMeshAreaZeroDetails", "The input Cloth Collection has a 2D Sim mesh with zero area. Sim remeshing will not be performed"));
+				}
+				else
+				{
+					EmptySimSelections(OutputClothCollectionRef);
+					RemeshSimMesh(ClothCollection, OutputClothCollectionRef);
+					RebuildTopologyDependentSimData(ClothCollection, OutputClothCollectionRef);
+				}
 			}
 
 			if (bRemeshRender)
diff --git a/Engine/Plugins/Runtime/MeshModelingToolset/Source/ModelingOperators/Private/CleaningOps/RemeshMeshOp.cpp b/Engine/Plugins/Runtime/MeshModelingToolset/Source/ModelingOperators/Private/CleaningOps/RemeshMeshOp.cpp
--- a/Engine/Plugins/Runtime/MeshModelingToolset/Source/ModelingOperators/Private/CleaningOps/RemeshMeshOp.cpp
+++ b/Engine/Plugins/Runtime/MeshModelingToolset/Source/ModelingOperators/Private/CleaningOps/RemeshMeshOp.cpp
@@ -256,6 +256,7 @@ double FRemeshMeshOp::CalculateTargetEdgeLength(const FDynamicMesh3* Mesh, int T
 	double InitialMeshArea = PrecomputedMeshArea;
 	if (InitialMeshArea <= 0)
 	{
+		checkf(Mesh, TEXT("FRemeshMeshOp: Expected either a non-zero precomputed mesh area or a valid Mesh pointer"));
 		InitialMeshArea = TMeshQueries<FDynamicMesh3>::GetVolumeArea(*Mesh).Y;
 	}
 
