Add experimental ImportSimulationCache node to read in cached simulation results to cloth collections.

diff --git a/Engine/Plugins/ChaosClothAssetEditor/ChaosClothAssetEditor.uplugin b/Engine/Plugins/ChaosClothAssetEditor/ChaosClothAssetEditor.uplugin
--- a/Engine/Plugins/ChaosClothAssetEditor/ChaosClothAssetEditor.uplugin
+++ b/Engine/Plugins/ChaosClothAssetEditor/ChaosClothAssetEditor.uplugin
@@ -44,6 +44,10 @@
 			"Name": "GeometryProcessing",
 			"Enabled": true
 		},
+		{
+			"Name": "ChaosCaching",
+			"Enabled": true
+		},
 		{
 			"Name": "ChaosClothAsset",
 			"Enabled": true
@@ -70,7 +74,9 @@
 			"Name": "USDImporter",
 			"Enabled": true,
 			"SupportedTargetPlatforms": [
-				"Win64", "Mac", "Linux"
+				"Win64",
+				"Mac",
+				"Linux"
 			]
 		}
 	]
diff --git a/Engine/Plugins/ChaosClothAssetEditor/Source/ChaosClothAssetDataflowNodes/ChaosClothAssetDataflowNodes.Build.cs b/Engine/Plugins/ChaosClothAssetEditor/Source/ChaosClothAssetDataflowNodes/ChaosClothAssetDataflowNodes.Build.cs
--- a/Engine/Plugins/ChaosClothAssetEditor/Source/ChaosClothAssetDataflowNodes/ChaosClothAssetDataflowNodes.Build.cs
+++ b/Engine/Plugins/ChaosClothAssetEditor/Source/ChaosClothAssetDataflowNodes/ChaosClothAssetDataflowNodes.Build.cs
@@ -18,6 +18,7 @@ public class ChaosClothAssetDataflowNodes : ModuleRules
 			new string[]
 			{
 				"Chaos",
+				"ChaosCaching",
 				"ChaosCloth",
 				"ChaosClothAsset",
 				"ChaosClothAssetEngine",
diff --git a/Engine/Plugins/ChaosClothAssetEditor/Source/ChaosClothAssetDataflowNodes/Private/ChaosClothAsset/ChaosClothAssetDataflowNodesModule.cpp b/Engine/Plugins/ChaosClothAssetEditor/Source/ChaosClothAssetDataflowNodes/Private/ChaosClothAsset/ChaosClothAssetDataflowNodesModule.cpp
--- a/Engine/Plugins/ChaosClothAssetEditor/Source/ChaosClothAssetDataflowNodes/Private/ChaosClothAsset/ChaosClothAssetDataflowNodesModule.cpp
+++ b/Engine/Plugins/ChaosClothAssetEditor/Source/ChaosClothAssetDataflowNodes/Private/ChaosClothAsset/ChaosClothAssetDataflowNodesModule.cpp
@@ -13,6 +13,7 @@
 #include "ChaosClothAsset/DeleteElementNode.h"
 #include "ChaosClothAsset/ImportFilePathCustomization.h"
 #include "ChaosClothAsset/ImportNode.h"
+#include "ChaosClothAsset/ImportSimulationCacheNode.h"
 #include "ChaosClothAsset/MergeClothCollectionsNode.h"
 #include "ChaosClothAsset/ProxyDeformerNode.h"
 #include "ChaosClothAsset/RemeshNode.h"
@@ -95,6 +96,7 @@ namespace UE::Chaos::ClothAsset
 			DATAFLOW_NODE_REGISTER_CREATION_FACTORY(FChaosClothAssetDatasmithImportNode);
 			DATAFLOW_NODE_REGISTER_CREATION_FACTORY(FChaosClothAssetDeleteElementNode);
 			DATAFLOW_NODE_REGISTER_CREATION_FACTORY(FChaosClothAssetImportNode);
+			DATAFLOW_NODE_REGISTER_CREATION_FACTORY(FChaosClothAssetImportSimulationCacheNode);
 			DATAFLOW_NODE_REGISTER_CREATION_FACTORY(FChaosClothAssetMergeClothCollectionsNode_v2);
 			DATAFLOW_NODE_REGISTER_CREATION_FACTORY(FChaosClothAssetProxyDeformerNode_v2);
 			DATAFLOW_NODE_REGISTER_CREATION_FACTORY(FChaosClothAssetRemeshNode);
diff --git a/Engine/Plugins/ChaosClothAssetEditor/Source/ChaosClothAssetDataflowNodes/Private/ChaosClothAsset/ImportSimulationCacheNode.cpp b/Engine/Plugins/ChaosClothAssetEditor/Source/ChaosClothAssetDataflowNodes/Private/ChaosClothAsset/ImportSimulationCacheNode.cpp
new file mode 100644
--- /dev/null
+++ b/Engine/Plugins/ChaosClothAssetEditor/Source/ChaosClothAssetDataflowNodes/Private/ChaosClothAsset/ImportSimulationCacheNode.cpp
@@ -0,0 +1,108 @@
+// Copyright Epic Games, Inc. All Rights Reserved.
+
+#include "ChaosClothAsset/ImportSimulationCacheNode.h"
+#include "ChaosClothAsset/ClothDataflowTools.h"
+#include "ChaosClothAsset/CollectionClothFacade.h"
+#include "Chaos/CacheCollection.h"
+#include "Chaos/ChaosCache.h"
+#include "Dataflow/DataflowInputOutput.h"
+#include "TriangleTypes.h"
+
+#include UE_INLINE_GENERATED_CPP_BY_NAME(ImportSimulationCacheNode)
+
+#define LOCTEXT_NAMESPACE "ChaosClothAssetImportSimulationCacheNode"
+
+FChaosClothAssetImportSimulationCacheNode::FChaosClothAssetImportSimulationCacheNode(const UE::Dataflow::FNodeParameters& InParam, FGuid InGuid)
+	: FDataflowNode(InParam, InGuid)
+{
+	RegisterInputConnection(&Collection);
+	RegisterInputConnection(&ImportedCache);
+	RegisterOutputConnection(&Collection, &Collection);
+}
+
+void FChaosClothAssetImportSimulationCacheNode::Evaluate(UE::Dataflow::FContext& Context, const FDataflowOutput* Out) const
+{
+	if (Out->IsA<FManagedArrayCollection>(&Collection))
+	{
+		using namespace UE::Chaos::ClothAsset;
+
+		// Evaluate in collection
+		FManagedArrayCollection InCollection = GetValue<FManagedArrayCollection>(Context, &Collection);
+		const TSharedRef<FManagedArrayCollection> ClothCollection = MakeShared<FManagedArrayCollection>(MoveTemp(InCollection));
+
+		const UChaosCacheCollection* const InCacheCollection = GetValue(Context, &ImportedCache);
+
+		// Always check for a valid cloth collection/facade to avoid processing non cloth collections
+		FCollectionClothFacade ClothFacade(ClothCollection);
+
+		if (ClothFacade.IsValid() && InCacheCollection && InCacheCollection->Caches.IsValidIndex(CacheIndex) && InCacheCollection->Caches[CacheIndex])
+		{
+			const UChaosCache* const InCache = InCacheCollection->Caches[CacheIndex];
+
+			FCacheUserToken Token = InCache->BeginPlayback();
+
+			FPlaybackTickRecord TickRecord;
+			TickRecord.SetLastTime(CacheTime);
+			TickRecord.SetDt(0.f);
+			TickRecord.SetSpaceTransform(Transform);
+
+			FCacheEvaluationContext CacheContext(TickRecord);
+			CacheContext.bEvaluateTransform = false;
+			CacheContext.bEvaluateCurves = false;
+			CacheContext.bEvaluateEvents = false;
+			CacheContext.bEvaluateChannels = true;
+			CacheContext.bEvaluateNamedTransforms = true;
+			const FCacheEvaluationResult EvaluatedResult = InCache->Evaluate(CacheContext, nullptr);
+			const TArray<float>* const PendingPX = EvaluatedResult.Channels.Find(FName(TEXT("PositionX")));
+			const TArray<float>* const PendingPY = EvaluatedResult.Channels.Find(FName(TEXT("PositionY")));
+			const TArray<float>* const PendingPZ = EvaluatedResult.Channels.Find(TEXT("PositionZ"));
+			const int32 NumCachedParticles = EvaluatedResult.ParticleIndices.Num();
+
+			if (PendingPX && PendingPY && PendingPZ)
+			{
+				TArrayView<FVector3f> SimPositions = ClothFacade.GetSimPosition3D();
+				for (int32 CachedIndex = 0; CachedIndex < NumCachedParticles; ++CachedIndex)
+				{
+					const int32 ParticleIndex = EvaluatedResult.ParticleIndices[CachedIndex] - ParticleOffset;
+					if (SimPositions.IsValidIndex(ParticleIndex))
+					{
+						SimPositions[ParticleIndex].X = (*PendingPX)[CachedIndex];
+						SimPositions[ParticleIndex].Y = (*PendingPY)[CachedIndex];
+						SimPositions[ParticleIndex].Z = (*PendingPZ)[CachedIndex];
+					}
+				}
+				if (bRecalculateNormals)
+				{
+					TArrayView<FVector3f> Normals = ClothFacade.GetSimNormal();
+					for (FVector3f& Normal : Normals)
+					{
+						Normal = FVector3f::ZeroVector;
+					}
+
+					for (const FIntVector& Element : ClothFacade.GetSimIndices3D())
+					{
+						UE::Geometry::FTriangle3f Triangle;
+						Triangle.V[0] = SimPositions[Element[0]];
+						Triangle.V[1] = SimPositions[Element[1]];
+						Triangle.V[2] = SimPositions[Element[2]];
+						const FVector3f& Normal = Triangle.Normal();
+						Normals[Element[0]] += Normal;
+						Normals[Element[1]] += Normal;
+						Normals[Element[2]] += Normal;
+					}
+
+					for (FVector3f& Normal : Normals)
+					{
+						Normal = Normal.GetSafeNormal(UE_SMALL_NUMBER, FVector3f::XAxisVector);
+					}
+				}
+			}
+
+			InCache->EndPlayback(Token);
+		}
+
+		SetValue(Context, MoveTemp(*ClothCollection), &Collection);
+	}
+}
+
+#undef LOCTEXT_NAMESPACE
diff --git a/Engine/Plugins/ChaosClothAssetEditor/Source/ChaosClothAssetDataflowNodes/Public/ChaosClothAsset/ImportSimulationCacheNode.h b/Engine/Plugins/ChaosClothAssetEditor/Source/ChaosClothAssetDataflowNodes/Public/ChaosClothAsset/ImportSimulationCacheNode.h
new file mode 100644
--- /dev/null
+++ b/Engine/Plugins/ChaosClothAssetEditor/Source/ChaosClothAssetDataflowNodes/Public/ChaosClothAsset/ImportSimulationCacheNode.h
@@ -0,0 +1,53 @@
+// Copyright Epic Games, Inc. All Rights Reserved.
+
+#pragma once
+
+#include "Dataflow/DataflowNode.h"
+#include "GeometryCollection/ManagedArrayCollection.h"
+#include "ImportSimulationCacheNode.generated.h"
+
+class UChaosCacheCollection;
+
+/** Set vertex values from a simulation cache. The topology of the Collection will remain the same.*/
+USTRUCT(Meta = (DataflowCloth, Experimental))
+struct FChaosClothAssetImportSimulationCacheNode final : public FDataflowNode 
+{
+	GENERATED_USTRUCT_BODY()
+	DATAFLOW_NODE_DEFINE_INTERNAL(FChaosClothAssetImportSimulationCacheNode, "ImportSimulationCache", "Cloth", "Cloth Simulation Import Cache")
+	DATAFLOW_NODE_RENDER_TYPE("SurfaceRender", FName("FClothCollection"), "Collection")
+
+public:
+
+	FChaosClothAssetImportSimulationCacheNode(const UE::Dataflow::FNodeParameters& InParam, FGuid InGuid = FGuid::NewGuid());
+
+private:
+	virtual void Evaluate(UE::Dataflow::FContext& Context, const FDataflowOutput* Out) const override;
+
+	/** Input/output collection */
+	UPROPERTY(Meta = (Dataflowinput, DataflowOutput, DataflowPassthrough = "Collection"))
+	FManagedArrayCollection Collection;
+
+	/** Cache to import in. */
+	UPROPERTY(EditAnywhere, Category = "Import Simulation Cache", Meta =(Dataflowinput))
+	TObjectPtr<const UChaosCacheCollection> ImportedCache;
+
+	/** Cache index to read */
+	UPROPERTY(EditAnywhere, Category = "Import Simulation Cache")
+	int32 CacheIndex = 0;
+
+	/** Cache time to read */
+	UPROPERTY(EditAnywhere, Category = "Import Simulation Cache")
+	float CacheTime = 0.f;
+
+	/** Transform cache data. */
+	UPROPERTY(EditAnywhere, Category = "Import Simulation Cache")
+	FTransform Transform;
+
+	/** Particle cache offset. */
+	UPROPERTY(EditAnywhere, Category = "Import Simulation Cache")
+	int32 ParticleOffset = 0;
+
+	/** Recalculate Normals based on imported positions. */
+	UPROPERTY(EditAnywhere, Category = "Import Simulation Cache")
+	bool bRecalculateNormals = true;
+};
