Chaos Cloth Asset - Added Reimport buttons to all asset import nodes to trigger re-imports whenever needed.

[RN] Chaos Cloth Asset - Added Reimport buttons to all asset import nodes to trigger re-imports whenever needed.

diff --git a/Engine/Plugins/ChaosClothAssetEditor/Source/ChaosClothAssetDataflowNodes/Public/ChaosClothAsset/SkeletalMeshImportNode.h b/Engine/Plugins/ChaosClothAssetEditor/Source/ChaosClothAssetDataflowNodes/Public/ChaosClothAsset/SkeletalMeshImportNode.h
--- a/Engine/Plugins/ChaosClothAssetEditor/Source/ChaosClothAssetDataflowNodes/Public/ChaosClothAsset/SkeletalMeshImportNode.h
+++ b/Engine/Plugins/ChaosClothAssetEditor/Source/ChaosClothAssetDataflowNodes/Public/ChaosClothAsset/SkeletalMeshImportNode.h
@@ -3,6 +3,7 @@
 #pragma once 
 
 #include "Dataflow/DataflowNode.h"
+#include "Dataflow/DataflowFunctionProperty.h"
 #include "GeometryCollection/ManagedArrayCollection.h"
 #include "SkeletalMeshImportNode.generated.h"
 
@@ -29,6 +30,10 @@ public:
 	UPROPERTY(EditAnywhere, Category = "Skeletal Mesh Import", Meta = (DataflowInput))
 	TObjectPtr<const USkeletalMesh> SkeletalMesh;
 
+	/** Reimport the imported skeletal mesh asset. */
+	UPROPERTY(EditAnywhere, Category = "Skeletal Mesh Import", Meta = (ButtonImage = "Persona.ReimportAsset", DisplayName = "Reimport Asset"))
+	FDataflowFunctionProperty Reimport;
+
 	/** The skeletal mesh LOD to import. */
 	UPROPERTY(EditAnywhere, Category = "Skeletal Mesh Import", Meta = (ClampMin = "0", DisplayName = "LOD Index"))
 	int32 LODIndex = 0;
diff --git a/Engine/Plugins/ChaosClothAssetEditor/Source/ChaosClothAssetDataflowNodes/Public/ChaosClothAsset/StaticMeshImportNode.h b/Engine/Plugins/ChaosClothAssetEditor/Source/ChaosClothAssetDataflowNodes/Public/ChaosClothAsset/StaticMeshImportNode.h
--- a/Engine/Plugins/ChaosClothAssetEditor/Source/ChaosClothAssetDataflowNodes/Public/ChaosClothAsset/StaticMeshImportNode.h
+++ b/Engine/Plugins/ChaosClothAssetEditor/Source/ChaosClothAssetDataflowNodes/Public/ChaosClothAsset/StaticMeshImportNode.h
@@ -3,6 +3,7 @@
 #pragma once 
 
 #include "Dataflow/DataflowNode.h"
+#include "Dataflow/DataflowFunctionProperty.h"
 #include "GeometryCollection/ManagedArrayCollection.h"
 #include "StaticMeshImportNode.generated.h"
 
@@ -30,6 +31,10 @@ public:
 	UPROPERTY(EditAnywhere, Category = "Static Mesh Import", Meta = (DataflowInput))
 	TObjectPtr<const UStaticMesh> StaticMesh;
 
+	/** Reimport the imported static mesh asset. */
+	UPROPERTY(EditAnywhere, Category = "Static Mesh Import", Meta = (ButtonImage = "Persona.ReimportAsset", DisplayName = "Reimport Asset", EditCondition = "!bIsLocked"))
+	FDataflowFunctionProperty Reimport;
+
 	/* Which static mesh Lod to import. */
 	UPROPERTY(EditAnywhere, Category = "Static Mesh Import", Meta = (ClampMin = "0", DisplayName = "LOD Index"))
 	int32 LODIndex = 0;
