Cloth Remesh Node: add optional vertex weights to control mesh resolution

diff --git a/Engine/Plugins/ChaosClothAssetEditor/Source/ChaosClothAssetDataflowNodes/Private/ChaosClothAsset/RemeshNode.cpp b/Engine/Plugins/ChaosClothAssetEditor/Source/ChaosClothAssetDataflowNodes/Private/ChaosClothAsset/RemeshNode.cpp
--- a/Engine/Plugins/ChaosClothAssetEditor/Source/ChaosClothAssetDataflowNodes/Private/ChaosClothAsset/RemeshNode.cpp
+++ b/Engine/Plugins/ChaosClothAssetEditor/Source/ChaosClothAssetDataflowNodes/Private/ChaosClothAsset/RemeshNode.cpp
@@ -34,6 +34,18 @@
 namespace UE::Chaos::ClothAsset::Private
 {
 
+	/** Given weight values for the edge vertices, compute a new edge length multiplier to be used in the remesher. 
+	    A greater edge length means the edge is more likely to be split, resulting in a higher overall vertex density */
+	// TODO: Experiment with other forms
+	// TODO: Add parameters to rescale the 0-1 weight values
+	static double EdgeScalingFunction(float WeightValueA, float WeightValueB)
+	{
+		const float Avg = 0.5f * (WeightValueA + WeightValueB);
+		const float ClampedWeight = FMath::Clamp(Avg, 0.0f, 1.0f);
+		const double LengthScale = 1.0f + ClampedWeight;
+		return LengthScale;
+	}
+
 	//
 	// Functions to support seam remeshing
 	//
@@ -69,7 +81,7 @@ namespace UE::Chaos::ClothAsset::Private
 	// 
 	// The function computes the dot product of normalized vectors AB and BC, and returns true if the result is less than the given threshold
 	// Also returns true if A \approx B or B \approx C
-	bool IsSharpCorner(const FVector3d& A, const FVector3d& B, const FVector3d& C, double CosAngleThreshold)
+	static bool IsSharpCorner(const FVector3d& A, const FVector3d& B, const FVector3d& C, double CosAngleThreshold)
 	{
 		const FVector3d AB = B - A;
 		const FVector3d BC = C - B;
@@ -85,7 +97,7 @@ namespace UE::Chaos::ClothAsset::Private
 		return (CosAngle < CosAngleThreshold);
 	}
 
-	FSeamCollapseParameters GetSeamEdgeCollapseParameters(int32 SeamID, int32 StitchID, const UE::Geometry::FDynamicMesh3& Mesh, const TArray<TArray<FIntVector2>>& Seams, double CosAngleCornerTheshold = 0.5)
+	static FSeamCollapseParameters GetSeamEdgeCollapseParameters(int32 SeamID, int32 StitchID, const UE::Geometry::FDynamicMesh3& Mesh, const TArray<TArray<FIntVector2>>& Seams, double CosAngleCornerTheshold = 0.5)
 	{
 		using namespace UE::Geometry;
 
@@ -259,7 +271,7 @@ namespace UE::Chaos::ClothAsset::Private
 		return Constraints;
 	}
 
-	bool CanSplitSeamEdge(int32 SeamID, int32 StitchID, const TArray<TArray<FIntVector2>>& Seams)
+	static bool CanSplitSeamEdge(int32 SeamID, int32 StitchID, const TArray<TArray<FIntVector2>>& Seams)
 	{
 		using namespace UE::Geometry;
 
@@ -307,7 +319,7 @@ namespace UE::Chaos::ClothAsset::Private
 	}
 
 
-	void FindCoincidentBoundaryVertices(const UE::Geometry::FDynamicMesh3& Mesh, TArray<FIntVector2>& Pairs)
+	static void FindCoincidentBoundaryVertices(const UE::Geometry::FDynamicMesh3& Mesh, TArray<FIntVector2>& Pairs)
 	{
 		using namespace UE::Geometry;
 
@@ -380,7 +392,7 @@ namespace UE::Chaos::ClothAsset::Private
 	}
 
 
-	void RemeshSeams(UE::Geometry::FDynamicMesh3& Mesh, TArray<TArray<FIntVector2>>& Seams, double TargetEdgeLength)
+	static void RemeshSeams(UE::Geometry::FDynamicMesh3& Mesh, TArray<TArray<FIntVector2>>& Seams, double TargetEdgeLength, const UE::Geometry::FDynamicMeshWeightAttribute* const DensityMapLayer)
 	{
 		using namespace UE::Geometry;
 
@@ -402,7 +414,17 @@ namespace UE::Chaos::ClothAsset::Private
 					continue;
 				}
 
-				const double EdgeALength = (Mesh.GetVertex(SideAVertexA) - Mesh.GetVertex(SideAVertexB)).Length();
+				double EdgeALength = (Mesh.GetVertex(SideAVertexA) - Mesh.GetVertex(SideAVertexB)).Length();
+
+				if (DensityMapLayer)
+				{
+					float WeightValueA;
+					DensityMapLayer->GetValue(SideAVertexA, &WeightValueA);
+					float WeightValueB;
+					DensityMapLayer->GetValue(SideAVertexB, &WeightValueB);
+
+					EdgeALength = EdgeScalingFunction(WeightValueA, WeightValueB) * EdgeALength;
+				}
 
 				const int32 SideBVertexA = Seam[StitchID][1];
 				const int32 SideBVertexB = Seam[StitchID + 1][1];
@@ -412,7 +434,17 @@ namespace UE::Chaos::ClothAsset::Private
 					continue;
 				}
 
-				const double EdgeBLength = (Mesh.GetVertex(SideBVertexA) - Mesh.GetVertex(SideBVertexB)).Length();
+				double EdgeBLength = (Mesh.GetVertex(SideBVertexA) - Mesh.GetVertex(SideBVertexB)).Length();
+
+				if (DensityMapLayer)
+				{
+					float WeightValueA;
+					DensityMapLayer->GetValue(SideBVertexA, &WeightValueA);
+					float WeightValueB;
+					DensityMapLayer->GetValue(SideBVertexB, &WeightValueB);
+
+					EdgeBLength = EdgeScalingFunction(WeightValueA, WeightValueB) * EdgeBLength;
+				}
 
 				if ((EdgeALength < MinLength) && (EdgeBLength < MinLength))
 				{
@@ -493,7 +525,7 @@ namespace UE::Chaos::ClothAsset::Private
 		double NewPositionEdgeParameter = 0.5;
 	};
 
-	FEdgeCollapseParameters GetBoundaryEdgeCollapseParameters(const UE::Geometry::FDynamicMesh3& Mesh, const TArray<TArray<FIntVector2>>& Seams, const UE::Geometry::FIndex2i& EdgeVerts, double CosAngleThreshold = 0.5)
+	static FEdgeCollapseParameters GetBoundaryEdgeCollapseParameters(const UE::Geometry::FDynamicMesh3& Mesh, const TArray<TArray<FIntVector2>>& Seams, const UE::Geometry::FIndex2i& EdgeVerts, double CosAngleThreshold = 0.5)
 	{
 		using namespace UE::Geometry;
 
@@ -610,7 +642,7 @@ namespace UE::Chaos::ClothAsset::Private
 		return Constraints;
 	}
 
-	void RemeshBoundaries(UE::Geometry::FDynamicMesh3& Mesh, const TArray<TArray<FIntVector2>>& Seams, double TargetEdgeLength)
+	static void RemeshBoundaries(UE::Geometry::FDynamicMesh3& Mesh, const TArray<TArray<FIntVector2>>& Seams, double TargetEdgeLength, const UE::Geometry::FDynamicMeshWeightAttribute* const DensityMapLayer)
 	{
 		using namespace UE::Geometry;
 
@@ -666,7 +698,19 @@ namespace UE::Chaos::ClothAsset::Private
 				continue;
 			}
 
-			const double EdgeLength = (Mesh.GetVertex(EdgeVerts[0]) - Mesh.GetVertex(EdgeVerts[1])).Length();
+			double EdgeLength = (Mesh.GetVertex(EdgeVerts[0]) - Mesh.GetVertex(EdgeVerts[1])).Length();
+
+			if (DensityMapLayer)
+			{
+				float WeightValueA;
+				DensityMapLayer->GetValue(EdgeVerts[0], &WeightValueA);
+				float WeightValueB;
+				DensityMapLayer->GetValue(EdgeVerts[1], &WeightValueB);
+
+				EdgeLength = EdgeScalingFunction(WeightValueA, WeightValueB) * EdgeLength;
+			}
+
+
 			if (EdgeLength < MinLength)
 			{
 				const UE::Chaos::ClothAsset::Private::FEdgeCollapseParameters CollapseParams = Private::GetBoundaryEdgeCollapseParameters(Mesh, Seams, EdgeVerts);
@@ -690,7 +734,7 @@ namespace UE::Chaos::ClothAsset::Private
 	// Remeshing away from Seams/Boundaries
 	// 
 
-	bool Remesh(UE::Geometry::FDynamicMesh3& Mesh, double TargetEdgeLength, int32 Iterations, float SmoothingRate, bool bUniformSmoothing, const TArray<TArray<FIntVector2>>& Seams, UE::Geometry::FCompactMaps* CompactMaps = nullptr)
+	static bool Remesh(UE::Geometry::FDynamicMesh3& Mesh, double TargetEdgeLength, int32 Iterations, float SmoothingRate, bool bUniformSmoothing, const TArray<TArray<FIntVector2>>& Seams, const FString& DensityMapName, UE::Geometry::FCompactMaps* CompactMaps = nullptr)
 	{
 		using namespace UE::Geometry;
 
@@ -775,6 +819,41 @@ namespace UE::Chaos::ClothAsset::Private
 		RemeshOp.bUseWorldSpace = false;
 		RemeshOp.bParallel = true;
 
+		// RemeshOp makes a copy of the mesh to operate on, so we can't just pass the FDynamicMeshWeightAttribute pointer into CustomEdgeLengthScaleF
+		int32 DensityMapLayerIndex = -1;
+		bool bFoundDensityMapLayer = false;
+		if (!DensityMapName.IsEmpty() && SourceMesh->HasAttributes())
+		{
+			for (int32 WeightLayerIndex = 0; WeightLayerIndex < SourceMesh->Attributes()->NumWeightLayers(); ++WeightLayerIndex)
+			{
+				if (SourceMesh->Attributes()->GetWeightLayer(WeightLayerIndex)->GetName() == DensityMapName)
+				{
+					DensityMapLayerIndex = WeightLayerIndex;
+					bFoundDensityMapLayer = true;
+					break;
+				}
+			}
+		}
+
+		RemeshOp.CustomEdgeLengthScaleF = [bFoundDensityMapLayer, DensityMapLayerIndex](const FDynamicMesh3& Mesh, int VertexA, int VertexB) -> double
+		{
+			if (bFoundDensityMapLayer)
+			{
+				check(Mesh.HasAttributes());
+				const FDynamicMeshWeightAttribute* const DensityMapLayer = Mesh.Attributes()->GetWeightLayer(DensityMapLayerIndex);
+				check(DensityMapLayer);
+
+				float WeightValueA;
+				DensityMapLayer->GetValue(VertexA, &WeightValueA);
+
+				float WeightValueB;
+				DensityMapLayer->GetValue(VertexB, &WeightValueB);
+
+				return EdgeScalingFunction(WeightValueA, WeightValueB);
+			}
+
+			return 1.0;
+		};
 
 		// Set up constraints for Cloth Seam edges
 		UE::Geometry::FMeshConstraints Constraints;
@@ -835,7 +914,7 @@ namespace UE::Chaos::ClothAsset::Private
 	}
 
 
-	bool Simplify(UE::Geometry::FDynamicMesh3& Mesh, int TargetVertexCount, bool bCoarsenBoundaries, UE::Geometry::FCompactMaps* CompactMaps = nullptr)
+	static bool Simplify(UE::Geometry::FDynamicMesh3& Mesh, int TargetVertexCount, bool bCoarsenBoundaries, UE::Geometry::FCompactMaps* CompactMaps = nullptr)
 	{
 		using namespace UE::Geometry;
 
@@ -927,6 +1006,8 @@ FChaosClothAssetRemeshNode::FChaosClothAssetRemeshNode(const UE::Dataflow::FNode
 {
 	RegisterInputConnection(&Collection);
 	RegisterOutputConnection(&Collection, &Collection);
+	RegisterInputConnection(&DensityMapSim.StringValue, GET_MEMBER_NAME_CHECKED(FChaosClothAssetConnectableIStringValue, StringValue));
+	RegisterInputConnection(&DensityMapRender.StringValue, GET_MEMBER_NAME_CHECKED(FChaosClothAssetConnectableIStringValue, StringValue));
 }
 
 
@@ -949,6 +1030,7 @@ void FChaosClothAssetRemeshNode::EmptySimSelections(const TSharedRef<FManagedArr
 }
 
 void FChaosClothAssetRemeshNode::RemeshSimMesh(const TSharedRef<const FManagedArrayCollection>& ClothCollection,
+	const FString& DensityMapName,
 	const TSharedRef<FManagedArrayCollection>& OutClothCollection) const
 {
 	using namespace UE::Geometry;
@@ -986,6 +1068,20 @@ void FChaosClothAssetRemeshNode::RemeshSimMesh(const TSharedRef<const FManagedAr
 		PatternIndexLayer->SetValue(FaceIndex, PatternID);
 	}
 
+	// Look for the DensityMap weight layer
+	const FDynamicMeshWeightAttribute* DensityMapLayer = nullptr;
+	if (Mesh2D.HasAttributes())
+	{
+		for (int32 WeightLayerIndex = 0; WeightLayerIndex < Mesh2D.Attributes()->NumWeightLayers(); ++WeightLayerIndex)
+		{
+			if (Mesh2D.Attributes()->GetWeightLayer(WeightLayerIndex)->GetName() == DensityMapName)
+			{
+				DensityMapLayer = Mesh2D.Attributes()->GetWeightLayer(WeightLayerIndex);
+				break;
+			}
+		}
+	}
+
 	// Grab the seam information from the input collection
 	
 	// Seams are comprised of a set of Stitches. Each Stitch is simply a pair of vertex indices indicating vertices that should be welded to form the 3D mesh.
@@ -1079,21 +1175,21 @@ void FChaosClothAssetRemeshNode::RemeshSimMesh(const TSharedRef<const FManagedAr
 
 	for (int32 ResampleIter = 0; ResampleIter < IterationsSim; ++ResampleIter)
 	{
-		UE::Chaos::ClothAsset::Private::RemeshSeams(Mesh2D, Seams, TargetEdgeLength);
+		UE::Chaos::ClothAsset::Private::RemeshSeams(Mesh2D, Seams, TargetEdgeLength, DensityMapLayer);
 	}
 
 	// Remesh boundaries
 
 	for (int32 ResampleIter = 0; ResampleIter < IterationsSim; ++ResampleIter)
 	{
-		UE::Chaos::ClothAsset::Private::RemeshBoundaries(Mesh2D, Seams, TargetEdgeLength);
+		UE::Chaos::ClothAsset::Private::RemeshBoundaries(Mesh2D, Seams, TargetEdgeLength, DensityMapLayer);
 	}
 
 	// Do the remeshing of the rest of the mesh
 
 	UE::Geometry::FCompactMaps CompactMaps;
 	constexpr bool bUniformSmoothing = true;
-	UE::Chaos::ClothAsset::Private::Remesh(Mesh2D, TargetEdgeLength, IterationsSim, SmoothingSim, bUniformSmoothing, Seams, &CompactMaps);
+	UE::Chaos::ClothAsset::Private::Remesh(Mesh2D, TargetEdgeLength, IterationsSim, SmoothingSim, bUniformSmoothing, Seams, DensityMapName, &CompactMaps);
 
 	// Update stitches
 	for (TArray<FIntVector2>& Seam : Seams)
@@ -1310,6 +1406,7 @@ void FChaosClothAssetRemeshNode::EmptyRenderSelections(const TSharedRef<FManaged
 
 
 void FChaosClothAssetRemeshNode::RemeshRenderMesh(const TSharedRef<const FManagedArrayCollection>& ClothCollection,
+	const FString& DensityMapName,
 	const TSharedRef<FManagedArrayCollection>& OutClothCollection) const
 {
 	using namespace UE::Geometry;
@@ -1344,15 +1441,29 @@ void FChaosClothAssetRemeshNode::RemeshRenderMesh(const TSharedRef<const FManage
 
 		FClothGeometryTools::BuildConnectedSeams(Stitches, DynamicMesh, Seams);
 
+		// Add density map for for Render mesh
+		const UE::Geometry::FDynamicMeshWeightAttribute* DensityMapLayer = nullptr;
+		if (DynamicMesh.HasAttributes())
+		{
+			for (int32 WeightLayerIndex = 0; WeightLayerIndex < DynamicMesh.Attributes()->NumWeightLayers(); ++WeightLayerIndex)
+			{
+				if (DynamicMesh.Attributes()->GetWeightLayer(WeightLayerIndex)->GetName() == DensityMapName)
+				{
+					DensityMapLayer = DynamicMesh.Attributes()->GetWeightLayer(WeightLayerIndex);
+					break;
+				}
+			}
+		}
+
 		for (int RemeshPass = 0; RemeshPass < RenderSeamRemeshIterations; ++RemeshPass)
 		{
-			UE::Chaos::ClothAsset::Private::RemeshSeams(DynamicMesh, Seams, TargetEdgeLength);
+			UE::Chaos::ClothAsset::Private::RemeshSeams(DynamicMesh, Seams, TargetEdgeLength, DensityMapLayer);
 		}
 
 		// Also remesh the open boundaries that are not constrained by seams
 		for (int RemeshPass = 0; RemeshPass < RenderSeamRemeshIterations; ++RemeshPass)
 		{
-			UE::Chaos::ClothAsset::Private::RemeshBoundaries(DynamicMesh, Seams, TargetEdgeLength);
+			UE::Chaos::ClothAsset::Private::RemeshBoundaries(DynamicMesh, Seams, TargetEdgeLength, DensityMapLayer);
 		}
 	}
 
@@ -1361,7 +1472,7 @@ void FChaosClothAssetRemeshNode::RemeshRenderMesh(const TSharedRef<const FManage
 	if (RemeshMethodRender == EChaosClothAssetRemeshMethod::Remesh)
 	{
 		constexpr bool bUniformSmoothing = false;	// uniform smoothing can distort the UV layer pretty badly
-		const bool bSuccess = UE::Chaos::ClothAsset::Private::Remesh(DynamicMesh, TargetEdgeLength, IterationsRender, SmoothingRender, bUniformSmoothing, Seams, &CompactMaps);
+		const bool bSuccess = UE::Chaos::ClothAsset::Private::Remesh(DynamicMesh, TargetEdgeLength, IterationsRender, SmoothingRender, bUniformSmoothing, Seams, DensityMapName,  &CompactMaps);
 		check(bSuccess);
 	}
 	else
@@ -1609,6 +1720,9 @@ void FChaosClothAssetRemeshNode::Evaluate(UE::Dataflow::FContext& Context, const
 		const TSharedRef<FManagedArrayCollection> ClothCollection = MakeShared<FManagedArrayCollection>(MoveTemp(InCollection));
 		FCollectionClothFacade ClothFacade(ClothCollection);
 
+		const FString SimDensityMapName = GetValue<FString>(Context, &DensityMapSim.StringValue);
+		const FString RenderDensityMapName = GetValue<FString>(Context, &DensityMapRender.StringValue);
+
 		// Copy collection to output
 		TSharedPtr<FManagedArrayCollection> OutputClothCollection = MakeShared<FManagedArrayCollection>();
 
@@ -1658,7 +1772,7 @@ void FChaosClothAssetRemeshNode::Evaluate(UE::Dataflow::FContext& Context, const
 				else
 				{
 					EmptySimSelections(OutputClothCollectionRef);
-					RemeshSimMesh(ClothCollection, OutputClothCollectionRef);
+					RemeshSimMesh(ClothCollection, SimDensityMapName, OutputClothCollectionRef);
 					RebuildTopologyDependentSimData(ClothCollection, OutputClothCollectionRef);
 				}
 			}
@@ -1666,7 +1780,7 @@ void FChaosClothAssetRemeshNode::Evaluate(UE::Dataflow::FContext& Context, const
 			if (bRemeshRender)
 			{
 				EmptyRenderSelections(OutputClothCollectionRef);
-				RemeshRenderMesh(ClothCollection, OutputClothCollectionRef);
+				RemeshRenderMesh(ClothCollection, RenderDensityMapName, OutputClothCollectionRef);
 			}
 		}
 
diff --git a/Engine/Plugins/ChaosClothAssetEditor/Source/ChaosClothAssetDataflowNodes/Public/ChaosClothAsset/RemeshNode.h b/Engine/Plugins/ChaosClothAssetEditor/Source/ChaosClothAssetDataflowNodes/Public/ChaosClothAsset/RemeshNode.h
--- a/Engine/Plugins/ChaosClothAssetEditor/Source/ChaosClothAssetDataflowNodes/Public/ChaosClothAsset/RemeshNode.h
+++ b/Engine/Plugins/ChaosClothAssetEditor/Source/ChaosClothAssetDataflowNodes/Public/ChaosClothAsset/RemeshNode.h
@@ -4,6 +4,7 @@
 
 #include "Dataflow/DataflowNode.h"
 #include "GeometryCollection/ManagedArrayCollection.h"
+#include "ChaosClothAsset/ConnectableValue.h"
 #include "RemeshNode.generated.h"
 
 #if UE_ENABLE_INCLUDE_ORDER_DEPRECATED_IN_5_5
@@ -47,6 +48,9 @@ public:
 	UPROPERTY(EditAnywhere, Category = "Sim Mesh", meta = (UIMin = "0.0", UIMax = "1.0", ClampMin = "0.0", ClampMax = "1.0", EditCondition = "bRemeshSim"))
 	double SmoothingSim = 0.25;
 
+	UPROPERTY(EditAnywhere, Category = "Sim Mesh", Meta = (EditCondition = "bRemeshSim"))
+	FChaosClothAssetConnectableIStringValue DensityMapSim;
+
 	UPROPERTY(EditAnywhere, Category = "Render Mesh")
 	bool bRemeshRender = false;
 
@@ -70,6 +74,9 @@ public:
 	UPROPERTY(EditAnywhere, Category = "Render Mesh", meta = (UIMin = "0", UIMax = "20", ClampMin = "0", ClampMax = "100", EditCondition = "bRemeshRender && bRemeshRenderSeams"))
 	int32 RenderSeamRemeshIterations = 1;
 
+	UPROPERTY(EditAnywhere, Category = "Render Mesh", Meta = (EditCondition = "bRemeshRender && RemeshMethodRender == EChaosClothAssetRemeshMethod::Remesh"))
+	FChaosClothAssetConnectableIStringValue DensityMapRender;
+
 	FChaosClothAssetRemeshNode(const UE::Dataflow::FNodeParameters& InParam, FGuid InGuid = FGuid::NewGuid());
 
 	virtual void Evaluate(UE::Dataflow::FContext& Context, const FDataflowOutput* Out) const override;
@@ -79,6 +86,7 @@ private:
 	void EmptySimSelections(const TSharedRef<FManagedArrayCollection>& ClothCollection) const;
 
 	void RemeshSimMesh(const TSharedRef<const FManagedArrayCollection>& ClothCollection,
+		const FString& DensityMapName,
 		const TSharedRef<FManagedArrayCollection>& OutClothCollection) const;
 
 	void RebuildTopologyDependentSimData(const TSharedRef<const FManagedArrayCollection>& ClothCollection,
@@ -87,5 +95,6 @@ private:
 	void EmptyRenderSelections(const TSharedRef<FManagedArrayCollection>& ClothCollection) const;
 
 	void RemeshRenderMesh(const TSharedRef<const FManagedArrayCollection>& ClothCollection,
+		const FString& DensityMapName,
 		const TSharedRef<FManagedArrayCollection>& OutClothCollection) const;
 };
