Add ability to apply simulation cache results to the collection render mesh (via proxy deformer)

diff --git a/Engine/Plugins/ChaosClothAssetEditor/Source/ChaosClothAssetDataflowNodes/Private/ChaosClothAsset/ImportSimulationCacheNode.cpp b/Engine/Plugins/ChaosClothAssetEditor/Source/ChaosClothAssetDataflowNodes/Private/ChaosClothAsset/ImportSimulationCacheNode.cpp
--- a/Engine/Plugins/ChaosClothAssetEditor/Source/ChaosClothAssetDataflowNodes/Private/ChaosClothAsset/ImportSimulationCacheNode.cpp
+++ b/Engine/Plugins/ChaosClothAssetEditor/Source/ChaosClothAssetDataflowNodes/Private/ChaosClothAsset/ImportSimulationCacheNode.cpp
@@ -12,6 +12,163 @@
 
 #define LOCTEXT_NAMESPACE "ChaosClothAssetImportSimulationCacheNode"
 
+namespace UE::Chaos::ClothAsset::Private
+{
+	static void UpdateSimulationPositions(const FCacheEvaluationResult& EvaluatedResult, const int32 ParticleOffset, TArrayView<FVector3f> SimPositions)
+	{
+		const TArray<float>* const PendingPX = EvaluatedResult.Channels.Find(FName(TEXT("PositionX")));
+		const TArray<float>* const PendingPY = EvaluatedResult.Channels.Find(FName(TEXT("PositionY")));
+		const TArray<float>* const PendingPZ = EvaluatedResult.Channels.Find(TEXT("PositionZ"));
+		check(PendingPX && PendingPY && PendingPZ);
+		const int32 NumCachedParticles = EvaluatedResult.ParticleIndices.Num();
+		for (int32 CachedIndex = 0; CachedIndex < NumCachedParticles; ++CachedIndex)
+		{
+			const int32 ParticleIndex = EvaluatedResult.ParticleIndices[CachedIndex] - ParticleOffset;
+			if (SimPositions.IsValidIndex(ParticleIndex))
+			{
+				SimPositions[ParticleIndex].X = (*PendingPX)[CachedIndex];
+				SimPositions[ParticleIndex].Y = (*PendingPY)[CachedIndex];
+				SimPositions[ParticleIndex].Z = (*PendingPZ)[CachedIndex];
+			}
+		}
+	}
+
+	static void RecalculateSimulationNormals(TConstArrayView<FVector3f> SimPositions, const FCollectionClothConstFacade& ClothFacade, TArrayView<FVector3f> Normals)
+	{
+		for (FVector3f& Normal : Normals)
+		{
+			Normal = FVector3f::ZeroVector;
+		}
+
+		for (const FIntVector& Element : ClothFacade.GetSimIndices3D())
+		{
+			UE::Geometry::FTriangle3f Triangle;
+			Triangle.V[0] = SimPositions[Element[0]];
+			Triangle.V[1] = SimPositions[Element[1]];
+			Triangle.V[2] = SimPositions[Element[2]];
+			const FVector3f& Normal = Triangle.Normal();
+			Normals[Element[0]] += Normal;
+			Normals[Element[1]] += Normal;
+			Normals[Element[2]] += Normal;
+		}
+
+		for (FVector3f& Normal : Normals)
+		{
+			Normal = Normal.GetSafeNormal(UE_SMALL_NUMBER, FVector3f::XAxisVector);
+		}
+	}
+
+	// This follows GpuSkinVertexFactor.ush
+	static void ApplyRenderDeformer(FCollectionClothFacade& ClothFacade, const TConstArrayView<FVector3f>& SimPositions, const TConstArrayView<FVector3f>& SimNormals)
+	{
+		for (int32 SectionIndex = 0; SectionIndex < ClothFacade.GetNumRenderPatterns(); ++SectionIndex)
+		{
+			FCollectionClothRenderPatternFacade RenderPatternFacade = ClothFacade.GetRenderPattern(SectionIndex);
+			const int32 RenderDeformerNumInfluences = RenderPatternFacade.GetRenderDeformerNumInfluences();
+			if (RenderDeformerNumInfluences > 0)
+			{
+				TArrayView<FVector3f> RenderPositions = RenderPatternFacade.GetRenderPosition();
+				TArrayView<FVector3f> RenderNormal = RenderPatternFacade.GetRenderNormal();
+				TArrayView<FVector3f> RenderTangentU = RenderPatternFacade.GetRenderTangentU();
+				TArrayView<FVector3f> RenderTangentV = RenderPatternFacade.GetRenderTangentV();
+
+				TConstArrayView<TArray<FVector4f>> PositionBaryCoordsAndDist = RenderPatternFacade.GetRenderDeformerPositionBaryCoordsAndDist();
+				TConstArrayView<TArray<FVector4f>> NormalBaryCoordsAndDist = RenderPatternFacade.GetRenderDeformerNormalBaryCoordsAndDist();
+				TConstArrayView<TArray<FVector4f>> TangentBaryCoordsAndDist = RenderPatternFacade.GetRenderDeformerTangentBaryCoordsAndDist();
+				TConstArrayView<TArray<FIntVector3>> SimIndices3D = RenderPatternFacade.GetRenderDeformerSimIndices3D();
+				TConstArrayView<TArray<float>> DeformerWeight = RenderPatternFacade.GetRenderDeformerWeight();
+				TConstArrayView<float> SkinningBlend = RenderPatternFacade.GetRenderDeformerSkinningBlend();
+
+				for (int32 Index = 0; Index < RenderPatternFacade.GetNumRenderVertices(); ++Index)
+				{
+					if (SkinningBlend[Index] < 1)
+					{
+						FVector3f AveragedSimPosition(0.f);
+						FVector3f NormalPosition(0.f);
+						FVector3f TangentPosition(0.f);
+
+						float SumWeights = 0.f;
+						float SimulWeight = 0.f;
+						for (int32 Influence = 0; Influence < RenderDeformerNumInfluences; ++Influence)
+						{
+							const FVector3f& SimPosA = SimPositions[SimIndices3D[Index][Influence].X];
+							const FVector3f& SimPosB = SimPositions[SimIndices3D[Index][Influence].Y];
+							const FVector3f& SimPosC = SimPositions[SimIndices3D[Index][Influence].Z];
+
+							const FVector3f& SimNormalA = SimNormals[SimIndices3D[Index][Influence].X];
+							const FVector3f& SimNormalB = SimNormals[SimIndices3D[Index][Influence].Y];
+							const FVector3f& SimNormalC = SimNormals[SimIndices3D[Index][Influence].Z];
+
+							SimulWeight += 1.f - SkinningBlend[Index];
+
+							float Weight = 1.f;
+							if (RenderDeformerNumInfluences > 1)
+							{
+								Weight = DeformerWeight[Index][Influence];
+								SumWeights += Weight;
+							}
+							else
+							{
+								Weight = 1.f;
+								SumWeights = 1.f;
+							}
+
+							auto InterpolatePosition = [Weight, &SimPosA, &SimPosB, &SimPosC, &SimNormalA, &SimNormalB, &SimNormalC](const FVector4f& BaryCoord)->FVector3f
+								{
+									// Note: coordinates are calculated with inverted normals, so subtract them here.
+									return Weight * (
+										BaryCoord.X * (SimPosA - SimNormalA * BaryCoord.W) +
+										BaryCoord.Y * (SimPosB - SimNormalB * BaryCoord.W) +
+										BaryCoord.Z * (SimPosC - SimNormalC * BaryCoord.W));
+								};
+
+							const FVector4f& BaryCoordPos = PositionBaryCoordsAndDist[Index][Influence];
+							AveragedSimPosition += InterpolatePosition(BaryCoordPos); 
+
+							const FVector4f& BaryCoordNormal = NormalBaryCoordsAndDist[Index][Influence];
+							NormalPosition += InterpolatePosition(BaryCoordNormal);
+
+							const FVector4f& BaryCoordTangent = TangentBaryCoordsAndDist[Index][Influence];
+							TangentPosition += InterpolatePosition(BaryCoordTangent);
+						}
+
+
+						FVector3f Normal(0.f);
+						FVector3f TangentU(0.f);
+						FVector3f TangentV(0.f);
+						if (SumWeights > UE_KINDA_SMALL_NUMBER)
+						{
+							float InvWeight = 1.f / SumWeights;
+							AveragedSimPosition *= InvWeight;
+							TangentPosition *= InvWeight;
+							NormalPosition *= InvWeight;
+
+							TangentU = (TangentPosition - AveragedSimPosition).GetSafeNormal();
+							Normal = (NormalPosition - AveragedSimPosition).GetSafeNormal();
+
+							TangentV = FVector3f::CrossProduct(TangentU, Normal).GetSafeNormal();
+						}
+						else
+						{
+							SimulWeight = 0.f; // Fallback to skinned pos
+						}
+
+						if (RenderDeformerNumInfluences > 1)
+						{
+							SimulWeight /= RenderDeformerNumInfluences;
+						}
+
+						RenderPositions[Index] = FMath::Lerp(RenderPositions[Index], AveragedSimPosition, SimulWeight);
+						RenderNormal[Index] = FMath::Lerp(RenderNormal[Index], Normal, SimulWeight).GetSafeNormal();
+						RenderTangentU[Index] = FMath::Lerp(RenderTangentU[Index], TangentU, SimulWeight).GetSafeNormal();
+						RenderTangentV[Index] = FMath::Lerp(RenderTangentV[Index], TangentV, SimulWeight).GetSafeNormal();
+					}
+				}
+			}
+		}
+	}
+}
+
 FChaosClothAssetImportSimulationCacheNode::FChaosClothAssetImportSimulationCacheNode(const UE::Dataflow::FNodeParameters& InParam, FGuid InGuid)
 	: FDataflowNode(InParam, InGuid)
 {
@@ -53,48 +210,54 @@ void FChaosClothAssetImportSimulationCacheNode::Evaluate(UE::Dataflow::FContext&
 			CacheContext.bEvaluateChannels = true;
 			CacheContext.bEvaluateNamedTransforms = true;
 			const FCacheEvaluationResult EvaluatedResult = InCache->Evaluate(CacheContext, nullptr);
+
 			const TArray<float>* const PendingPX = EvaluatedResult.Channels.Find(FName(TEXT("PositionX")));
 			const TArray<float>* const PendingPY = EvaluatedResult.Channels.Find(FName(TEXT("PositionY")));
 			const TArray<float>* const PendingPZ = EvaluatedResult.Channels.Find(TEXT("PositionZ"));
-			const int32 NumCachedParticles = EvaluatedResult.ParticleIndices.Num();
 
 			if (PendingPX && PendingPY && PendingPZ)
 			{
-				TArrayView<FVector3f> SimPositions = ClothFacade.GetSimPosition3D();
-				for (int32 CachedIndex = 0; CachedIndex < NumCachedParticles; ++CachedIndex)
+				if (bUpdateSimulationMesh)
 				{
-					const int32 ParticleIndex = EvaluatedResult.ParticleIndices[CachedIndex] - ParticleOffset;
-					if (SimPositions.IsValidIndex(ParticleIndex))
+					TArrayView<FVector3f> SimPositions = ClothFacade.GetSimPosition3D();
+					Private::UpdateSimulationPositions(EvaluatedResult, ParticleOffset, SimPositions);
+					if (bRecalculateNormals)
 					{
-						SimPositions[ParticleIndex].X = (*PendingPX)[CachedIndex];
-						SimPositions[ParticleIndex].Y = (*PendingPY)[CachedIndex];
-						SimPositions[ParticleIndex].Z = (*PendingPZ)[CachedIndex];
+						Private::RecalculateSimulationNormals(SimPositions, ClothFacade, ClothFacade.GetSimNormal());
 					}
 				}
-				if (bRecalculateNormals)
+
+				if (bUpdateRenderMesh)
 				{
-					TArrayView<FVector3f> Normals = ClothFacade.GetSimNormal();
-					for (FVector3f& Normal : Normals)
+					TArray<FVector3f> CachedSimPositions;
+					TArray<FVector3f> CachedSimNormals;
+					TArrayView<FVector3f> SimPositions;
+					TArrayView<FVector3f> SimNormals;
+					if (bUpdateSimulationMesh)
 					{
-						Normal = FVector3f::ZeroVector;
+						SimPositions = ClothFacade.GetSimPosition3D();
+						if (bRecalculateNormals)
+						{
+							SimNormals = ClothFacade.GetSimNormal();
+						}
+						else
+						{
+							CachedSimNormals.SetNumUninitialized(ClothFacade.GetNumSimVertices3D());
+							SimNormals = CachedSimNormals;
+							Private::RecalculateSimulationNormals(SimPositions, ClothFacade, SimNormals);
+						}
 					}
-
-					for (const FIntVector& Element : ClothFacade.GetSimIndices3D())
+					else
 					{
-						UE::Geometry::FTriangle3f Triangle;
-						Triangle.V[0] = SimPositions[Element[0]];
-						Triangle.V[1] = SimPositions[Element[1]];
-						Triangle.V[2] = SimPositions[Element[2]];
-						const FVector3f& Normal = Triangle.Normal();
-						Normals[Element[0]] += Normal;
-						Normals[Element[1]] += Normal;
-						Normals[Element[2]] += Normal;
+						CachedSimPositions = ClothFacade.GetSimPosition3D();
+						SimPositions = CachedSimPositions;
+						Private::UpdateSimulationPositions(EvaluatedResult, ParticleOffset, SimPositions);
+						CachedSimNormals.SetNumUninitialized(ClothFacade.GetNumSimVertices3D());
+						SimNormals = CachedSimNormals;
+						Private::RecalculateSimulationNormals(SimPositions, ClothFacade, SimNormals);
 					}
 
-					for (FVector3f& Normal : Normals)
-					{
-						Normal = Normal.GetSafeNormal(UE_SMALL_NUMBER, FVector3f::XAxisVector);
-					}
+					Private::ApplyRenderDeformer(ClothFacade, SimPositions, SimNormals);
 				}
 			}
 
diff --git a/Engine/Plugins/ChaosClothAssetEditor/Source/ChaosClothAssetDataflowNodes/Public/ChaosClothAsset/ImportSimulationCacheNode.h b/Engine/Plugins/ChaosClothAssetEditor/Source/ChaosClothAssetDataflowNodes/Public/ChaosClothAsset/ImportSimulationCacheNode.h
--- a/Engine/Plugins/ChaosClothAssetEditor/Source/ChaosClothAssetDataflowNodes/Public/ChaosClothAsset/ImportSimulationCacheNode.h
+++ b/Engine/Plugins/ChaosClothAssetEditor/Source/ChaosClothAssetDataflowNodes/Public/ChaosClothAsset/ImportSimulationCacheNode.h
@@ -47,7 +47,16 @@ private:
 	UPROPERTY(EditAnywhere, Category = "Import Simulation Cache")
 	int32 ParticleOffset = 0;
 
-	/** Recalculate Normals based on imported positions. */
+	/** Update simulation mesh from cache. */
 	UPROPERTY(EditAnywhere, Category = "Import Simulation Cache")
+	bool bUpdateSimulationMesh = true;
+
+	/** Recalculate simulation normals based on imported positions. */
+	UPROPERTY(EditAnywhere, Category = "Import Simulation Cache", Meta = (EditCondition = "bUpdateSimulationMesh"))
 	bool bRecalculateNormals = true;
+
+	/** Update render mesh from cache via proxy deformer data. */
+	UPROPERTY(EditAnywhere, Category = "Import Simulation Cache")
+	bool bUpdateRenderMesh = false;
+
 };
