VertexFaceSpringConfigNode: use tet mesh to find vertex-triangle pairs

diff --git a/Engine/Plugins/ChaosClothAssetEditor/Source/ChaosClothAssetDataflowNodes/Private/ChaosClothAsset/SimulationClothVertexFaceSpringConfigNode.cpp b/Engine/Plugins/ChaosClothAssetEditor/Source/ChaosClothAssetDataflowNodes/Private/ChaosClothAsset/SimulationClothVertexFaceSpringConfigNode.cpp
--- a/Engine/Plugins/ChaosClothAssetEditor/Source/ChaosClothAssetDataflowNodes/Private/ChaosClothAsset/SimulationClothVertexFaceSpringConfigNode.cpp
+++ b/Engine/Plugins/ChaosClothAssetEditor/Source/ChaosClothAssetDataflowNodes/Private/ChaosClothAsset/SimulationClothVertexFaceSpringConfigNode.cpp
@@ -12,7 +12,7 @@
 #include "Algo/MaxElement.h"
 #include "MeshQueries.h"
 #include "TriangleTypes.h"
-
+#include "CompGeom/Delaunay3.h"
 
 #include "ChaosClothAsset/ClothDataflowTools.h"
 
@@ -254,6 +254,124 @@ namespace UE::Chaos::ClothAsset::Private
 		}
 	}
 
+
+	void AppendConstraintsTetrahedralize(const TSet<int32>& SourceVertices, const TSet<int32>& TargetFaces, float Radius, int32 DisableNeighborDistance, const FCollectionClothConstFacade& Cloth, TSet<FVertexFaceConstraint>& Constraints)
+	{
+		// Speed up vertex-triangle queries by constructing this map up front
+		TMap<int32, TArray<int32>> VertexToTriangleMap;
+		TConstArrayView<FIntVector3> SimTris = Cloth.GetSimIndices3D();
+		for (const int32& TriIndex : TargetFaces)
+		{
+			const FIntVector3& SimTri = SimTris[TriIndex];
+
+			for (int32 TriCornerIndex = 0; TriCornerIndex < 3; ++TriCornerIndex)
+			{
+				const int32 VertexIndex = SimTri[TriCornerIndex];
+				if (!VertexToTriangleMap.Find(VertexIndex))
+				{
+					VertexToTriangleMap.Add(VertexIndex);
+				}
+				VertexToTriangleMap[VertexIndex].Add(TriIndex);
+			}
+		}
+
+		// Check if the three input vertices correspond to a single triangle on the Cloth sim mesh
+		auto IsSimTri = [](const FIntVector3& TestIndices, const TMap<int32, TArray<int32>>& VertexToTriangleMap, const FCollectionClothConstFacade& Cloth, int32& OutClothTriIndex)
+		{
+			check(TestIndices[0] != TestIndices[1] && TestIndices[0] != TestIndices[2] && TestIndices[1] != TestIndices[2]);
+
+			const TArray<int32> CandidateTriangles = VertexToTriangleMap[TestIndices[0]];
+			for (const int32 CandidateTriIndex : CandidateTriangles)
+			{
+				const FIntVector3& Tri = Cloth.GetSimIndices3D()[CandidateTriIndex];
+				check(Tri[0] != Tri[1] && Tri[0] != Tri[2] && Tri[1] != Tri[2]);
+
+				if ((Tri[0] == TestIndices[0] || Tri[0] == TestIndices[1] || Tri[0] == TestIndices[2]) &&
+					(Tri[1] == TestIndices[0] || Tri[1] == TestIndices[1] || Tri[1] == TestIndices[2]) &&
+					(Tri[2] == TestIndices[0] || Tri[2] == TestIndices[1] || Tri[2] == TestIndices[2]))
+				{
+					OutClothTriIndex = CandidateTriIndex;
+					return true;
+				}
+			}
+
+			return false;
+		};
+
+
+		const FFaceSetMeshAdapter MeshAdapter(TargetFaces, Cloth);
+
+		// Throw all source vertices and target triangle vertices into the tet mesher
+		TSet<int32> AllVertexIndices;
+		for (const int32 SourceVertex : SourceVertices)
+		{
+			AllVertexIndices.Add(SourceVertex);
+		}
+		for (const int32 TargetFace : TargetFaces)
+		{
+			if (MeshAdapter.IsTriangle(TargetFace))
+			{
+				const UE::Geometry::FIndex3i Triangle = MeshAdapter.GetTriangle(TargetFace);
+				AllVertexIndices.Add(Triangle[0]);
+				AllVertexIndices.Add(Triangle[1]);
+				AllVertexIndices.Add(Triangle[2]);
+			}
+		}
+		TConstArrayView<FVector3f> Positions = Cloth.GetSimPosition3D();
+		TArray<FVector3f> TetInputPoints;
+		for (const int32 VertexIndex : AllVertexIndices)
+		{
+			check(Positions.IsValidIndex(VertexIndex));
+			TetInputPoints.Add(Positions[VertexIndex]);
+		}
+
+		// Compute the tetrahedral mesh
+		UE::Geometry::FDelaunay3 Delaunay;
+		if (Delaunay.Triangulate(TetInputPoints))
+		{
+			const TArray<FIntVector4> Tets = Delaunay.GetTetrahedra();
+			for (const FIntVector4& Tet : Tets)
+			{
+				for (int32 TetFaceIndex = 0; TetFaceIndex < 4; ++TetFaceIndex)
+				{
+					// For each tet face, if it corresponds to a triangle in the sim mesh, use it to create a constraint
+					const FIntVector3 TetFace(Tet[TetFaceIndex], Tet[(TetFaceIndex + 1) % 4], Tet[(TetFaceIndex + 2) % 4]);
+					int32 SimMeshTriIndex = -1;
+					if (IsSimTri(TetFace, VertexToTriangleMap, Cloth, SimMeshTriIndex))
+					{
+						const int32 VertIndex = Tet[(TetFaceIndex + 3) % 4];
+						check(0 <= VertIndex && VertIndex < Cloth.GetNumSimVertices3D());
+						check(0 <= SimMeshTriIndex && SimMeshTriIndex < Cloth.GetNumSimFaces());
+
+						const FIntVector3 SimMeshTri = Cloth.GetSimIndices3D()[SimMeshTriIndex];
+						check(VertIndex != SimMeshTri[0] && VertIndex != SimMeshTri[1] && VertIndex != SimMeshTri[2]);
+
+						// Only take constraints where the vertex is in the set of source vertices and triangle is in the set of target triangles
+						if (!SourceVertices.Contains(VertIndex) || !TargetFaces.Contains(SimMeshTriIndex))
+						{
+							continue;
+						}
+
+						// Check if the vertex is coplanar to the triangle and skip this contraint if it is
+						const FVector3f VertPos = Cloth.GetSimPosition3D()[VertIndex];
+						const FVector3f A = Cloth.GetSimPosition3D()[SimMeshTri[0]] - VertPos;
+						const FVector3f B = Cloth.GetSimPosition3D()[SimMeshTri[1]] - VertPos;
+						const FVector3f C = Cloth.GetSimPosition3D()[SimMeshTri[2]] - VertPos;
+						const float Vol = FVector3f::DotProduct(A, FVector3f::CrossProduct(B, C));
+						if (FMath::Abs(Vol) < UE_KINDA_SMALL_NUMBER)
+						{
+							continue;
+						}
+
+						Constraints.Add({ VertIndex, SimMeshTriIndex, FVertexFaceConstraint::EIntersectionType::Closest });
+					}
+				}
+			}
+		}
+
+	}
+
+
 	void StoreConstraintsEnds(const TSet<FVertexFaceConstraint>& Constraints, const FCollectionClothConstFacade& Cloth, TArray<int32>& SourceVertices, TArray<FIntVector>& TargetVertices, TArray<FVector3f>& TargetWeights)
 	{
 		using namespace UE::Geometry;
@@ -644,6 +762,9 @@ void FChaosClothAssetSimulationClothVertexFaceSpringConfigNode::CreateConstraint
 				case EChaosClothAssetClothVertexFaceSpringConstructionMethod::AllWithinRadius:
 					Private::AppendConstraintsAllWithinRadius(SourceSet, TargetSet, Data.Radius, Data.DisableNeighborDistance, ClothFacade, Constraints);
 					break;
+				case EChaosClothAssetClothVertexFaceSpringConstructionMethod::Tetrahedralize:
+					Private::AppendConstraintsTetrahedralize(SourceSet, TargetSet, Data.Radius, Data.DisableNeighborDistance, ClothFacade, Constraints);
+					break;
 				default:
 					checkNoEntry();
 				}
diff --git a/Engine/Plugins/ChaosClothAssetEditor/Source/ChaosClothAssetDataflowNodes/Public/ChaosClothAsset/SimulationClothVertexFaceSpringConfigNode.h b/Engine/Plugins/ChaosClothAssetEditor/Source/ChaosClothAssetDataflowNodes/Public/ChaosClothAsset/SimulationClothVertexFaceSpringConfigNode.h
--- a/Engine/Plugins/ChaosClothAssetEditor/Source/ChaosClothAssetDataflowNodes/Public/ChaosClothAsset/SimulationClothVertexFaceSpringConfigNode.h
+++ b/Engine/Plugins/ChaosClothAssetEditor/Source/ChaosClothAssetDataflowNodes/Public/ChaosClothAsset/SimulationClothVertexFaceSpringConfigNode.h
@@ -18,6 +18,8 @@ enum struct EChaosClothAssetClothVertexFaceSpringConstructionMethod
 	SourceToRayIntersectionTarget,
 	// For each source, find all targets within a radius.
 	AllWithinRadius,
+	// Create a tet mesh and find corresponding tet face-vertex pairs
+	Tetrahedralize,
 };
 
 /** Data to procedurally generate ClothVertexFaceSpring Constraints.*/
