Cloth Remesh Dataflow Node: replace a single target percent with a range that a vertex weight map can interpolate between

diff --git a/Engine/Plugins/ChaosClothAssetEditor/Source/ChaosClothAssetDataflowNodes/Private/ChaosClothAsset/ChaosClothAssetDataflowNodesModule.cpp b/Engine/Plugins/ChaosClothAssetEditor/Source/ChaosClothAssetDataflowNodes/Private/ChaosClothAsset/ChaosClothAssetDataflowNodesModule.cpp
--- a/Engine/Plugins/ChaosClothAssetEditor/Source/ChaosClothAssetDataflowNodes/Private/ChaosClothAsset/ChaosClothAssetDataflowNodesModule.cpp
+++ b/Engine/Plugins/ChaosClothAssetEditor/Source/ChaosClothAssetDataflowNodes/Private/ChaosClothAsset/ChaosClothAssetDataflowNodesModule.cpp
@@ -99,7 +99,7 @@ namespace UE::Chaos::ClothAsset
 			DATAFLOW_NODE_REGISTER_CREATION_FACTORY(FChaosClothAssetImportSimulationCacheNode);
 			DATAFLOW_NODE_REGISTER_CREATION_FACTORY(FChaosClothAssetMergeClothCollectionsNode_v2);
 			DATAFLOW_NODE_REGISTER_CREATION_FACTORY(FChaosClothAssetProxyDeformerNode_v2);
-			DATAFLOW_NODE_REGISTER_CREATION_FACTORY(FChaosClothAssetRemeshNode);
+			DATAFLOW_NODE_REGISTER_CREATION_FACTORY(FChaosClothAssetRemeshNode_v2);
 			DATAFLOW_NODE_REGISTER_CREATION_FACTORY(FChaosClothAssetReverseNormalsNode);
 			DATAFLOW_NODE_REGISTER_CREATION_FACTORY(FChaosClothAssetSelectionNode_v2);
 			DATAFLOW_NODE_REGISTER_CREATION_FACTORY(FChaosClothAssetSelectionToIntMapNode);
@@ -154,6 +154,7 @@ PRAGMA_DISABLE_DEPRECATION_WARNINGS
 			DATAFLOW_NODE_REGISTER_CREATION_FACTORY(FChaosClothAssetAddWeightMapNode);
 			DATAFLOW_NODE_REGISTER_CREATION_FACTORY(FChaosClothAssetMergeClothCollectionsNode);
 			DATAFLOW_NODE_REGISTER_CREATION_FACTORY(FChaosClothAssetProxyDeformerNode);
+			DATAFLOW_NODE_REGISTER_CREATION_FACTORY(FChaosClothAssetRemeshNode);
 			DATAFLOW_NODE_REGISTER_CREATION_FACTORY(FChaosClothAssetSelectionNode);
 			DATAFLOW_NODE_REGISTER_CREATION_FACTORY(FChaosClothAssetSimulationLongRangeAttachmentConfigNode);
 			DATAFLOW_NODE_REGISTER_CREATION_FACTORY(FChaosClothAssetSkeletalMeshImportNode);
diff --git a/Engine/Plugins/ChaosClothAssetEditor/Source/ChaosClothAssetDataflowNodes/Private/ChaosClothAsset/RemeshNode.cpp b/Engine/Plugins/ChaosClothAssetEditor/Source/ChaosClothAssetDataflowNodes/Private/ChaosClothAsset/RemeshNode.cpp
--- a/Engine/Plugins/ChaosClothAssetEditor/Source/ChaosClothAssetDataflowNodes/Private/ChaosClothAsset/RemeshNode.cpp
+++ b/Engine/Plugins/ChaosClothAssetEditor/Source/ChaosClothAssetDataflowNodes/Private/ChaosClothAsset/RemeshNode.cpp
@@ -34,16 +34,16 @@
 namespace UE::Chaos::ClothAsset::Private
 {
 
-	/** Given weight values for the edge vertices, compute a new edge length multiplier to be used in the remesher. 
-	    A greater edge length means the edge is more likely to be split, resulting in a higher overall vertex density */
-	// TODO: Experiment with other forms
-	// TODO: Add parameters to rescale the 0-1 weight values
-	static double EdgeScalingFunction(float WeightValueA, float WeightValueB)
+	/** 
+	 * Given weight values for the edge vertices, compute a new edge length multiplier to be used in the remesher. 
+	 * A greater edge length means the edge is more likely to be split, resulting in a higher overall vertex density 
+	 */
+	static double EdgeScalingFunction(float LocalTriangleDensityWeightValueA, float LocalTriangleDensityWeightValueB)
 	{
-		const float Avg = 0.5f * (WeightValueA + WeightValueB);
-		const float ClampedWeight = FMath::Clamp(Avg, 0.0f, 1.0f);
-		const double LengthScale = 1.0f + ClampedWeight;
-		return LengthScale;
+		// Interpret the weight values as desired local triangle density scalars. E.g. A value of 2.0 on all vertices should roughly double the total number of triangles.
+		// To turn this into a scalar that we can multiply the computed edge lengths by during remesing, we take the square root of the vertex weights since AvgEdgeLength \propto Sqrt(AvgTriangleArea)
+		// Then we take the simple mean average of the two edge vertex values
+		return 0.5f * (FMath::Sqrt(LocalTriangleDensityWeightValueA) + FMath::Sqrt(LocalTriangleDensityWeightValueB));
 	}
 
 	//
@@ -392,7 +392,10 @@ namespace UE::Chaos::ClothAsset::Private
 	}
 
 
-	static void RemeshSeams(UE::Geometry::FDynamicMesh3& Mesh, TArray<TArray<FIntVector2>>& Seams, double TargetEdgeLength, const UE::Geometry::FDynamicMeshWeightAttribute* const DensityMapLayer)
+	static void RemeshSeams(UE::Geometry::FDynamicMesh3& Mesh, TArray<TArray<FIntVector2>>& Seams, 
+		double TargetEdgeLength, 
+		const UE::Geometry::FDynamicMeshWeightAttribute* const DensityMapLayer,
+		const FVector2f& DensityLowHigh)
 	{
 		using namespace UE::Geometry;
 
@@ -420,11 +423,19 @@ namespace UE::Chaos::ClothAsset::Private
 				{
 					float WeightValueA;
 					DensityMapLayer->GetValue(SideAVertexA, &WeightValueA);
+					WeightValueA = DensityLowHigh[0] * (1.0 - WeightValueA) + DensityLowHigh[1] * WeightValueA;
+
 					float WeightValueB;
 					DensityMapLayer->GetValue(SideAVertexB, &WeightValueB);
+					WeightValueB = DensityLowHigh[0] * (1.0 - WeightValueB) + DensityLowHigh[1] * WeightValueB;
 
 					EdgeALength = EdgeScalingFunction(WeightValueA, WeightValueB) * EdgeALength;
 				}
+				else
+				{
+					EdgeALength = EdgeScalingFunction(DensityLowHigh[0], DensityLowHigh[0]) * EdgeALength;
+				}
+				
 
 				const int32 SideBVertexA = Seam[StitchID][1];
 				const int32 SideBVertexB = Seam[StitchID + 1][1];
@@ -440,11 +451,18 @@ namespace UE::Chaos::ClothAsset::Private
 				{
 					float WeightValueA;
 					DensityMapLayer->GetValue(SideBVertexA, &WeightValueA);
+					WeightValueA = DensityLowHigh[0] * (1.0 - WeightValueA) + DensityLowHigh[1] * WeightValueA;
+
 					float WeightValueB;
 					DensityMapLayer->GetValue(SideBVertexB, &WeightValueB);
+					WeightValueB = DensityLowHigh[0] * (1.0 - WeightValueB) + DensityLowHigh[1] * WeightValueB;
 
 					EdgeBLength = EdgeScalingFunction(WeightValueA, WeightValueB) * EdgeBLength;
 				}
+				else
+				{
+					EdgeBLength = EdgeScalingFunction(DensityLowHigh[0], DensityLowHigh[0]) * EdgeALength;
+				}
 
 				if ((EdgeALength < MinLength) && (EdgeBLength < MinLength))
 				{
@@ -642,7 +660,11 @@ namespace UE::Chaos::ClothAsset::Private
 		return Constraints;
 	}
 
-	static void RemeshBoundaries(UE::Geometry::FDynamicMesh3& Mesh, const TArray<TArray<FIntVector2>>& Seams, double TargetEdgeLength, const UE::Geometry::FDynamicMeshWeightAttribute* const DensityMapLayer)
+	static void RemeshBoundaries(UE::Geometry::FDynamicMesh3& Mesh, 
+		const TArray<TArray<FIntVector2>>& Seams, 
+		double TargetEdgeLength, 
+		const UE::Geometry::FDynamicMeshWeightAttribute* const DensityMapLayer,
+		const FVector2f& DensityLowHigh)
 	{
 		using namespace UE::Geometry;
 
@@ -704,11 +726,18 @@ namespace UE::Chaos::ClothAsset::Private
 			{
 				float WeightValueA;
 				DensityMapLayer->GetValue(EdgeVerts[0], &WeightValueA);
+				WeightValueA = DensityLowHigh[0] * (1.0 - WeightValueA) + DensityLowHigh[1] * WeightValueA;
+
 				float WeightValueB;
 				DensityMapLayer->GetValue(EdgeVerts[1], &WeightValueB);
+				WeightValueB = DensityLowHigh[0] * (1.0 - WeightValueB) + DensityLowHigh[1] * WeightValueB;
 
 				EdgeLength = EdgeScalingFunction(WeightValueA, WeightValueB) * EdgeLength;
 			}
+			else
+			{
+				EdgeLength = EdgeScalingFunction(DensityLowHigh[0], DensityLowHigh[0]) * EdgeLength;
+			}
 
 
 			if (EdgeLength < MinLength)
@@ -734,7 +763,15 @@ namespace UE::Chaos::ClothAsset::Private
 	// Remeshing away from Seams/Boundaries
 	// 
 
-	static bool Remesh(UE::Geometry::FDynamicMesh3& Mesh, double TargetEdgeLength, int32 Iterations, float SmoothingRate, bool bUniformSmoothing, const TArray<TArray<FIntVector2>>& Seams, const FString& DensityMapName, UE::Geometry::FCompactMaps* CompactMaps = nullptr)
+	static bool Remesh(UE::Geometry::FDynamicMesh3& Mesh, 
+		double TargetEdgeLength, 
+		int32 Iterations, 
+		float SmoothingRate, 
+		bool bUniformSmoothing, 
+		const TArray<TArray<FIntVector2>>& Seams, 
+		const FString& DensityMapName, 
+		const FVector2f& DensityLowHigh,
+		UE::Geometry::FCompactMaps* CompactMaps = nullptr)
 	{
 		using namespace UE::Geometry;
 
@@ -835,7 +872,7 @@ namespace UE::Chaos::ClothAsset::Private
 			}
 		}
 
-		RemeshOp.CustomEdgeLengthScaleF = [bFoundDensityMapLayer, DensityMapLayerIndex](const FDynamicMesh3& Mesh, int VertexA, int VertexB) -> double
+		RemeshOp.CustomEdgeLengthScaleF = [bFoundDensityMapLayer, DensityMapLayerIndex, &DensityLowHigh](const FDynamicMesh3& Mesh, int VertexA, int VertexB) -> double
 		{
 			if (bFoundDensityMapLayer)
 			{
@@ -845,14 +882,18 @@ namespace UE::Chaos::ClothAsset::Private
 
 				float WeightValueA;
 				DensityMapLayer->GetValue(VertexA, &WeightValueA);
+				WeightValueA = DensityLowHigh[0] * (1 - WeightValueA) + DensityLowHigh[1] * WeightValueA;
 
 				float WeightValueB;
 				DensityMapLayer->GetValue(VertexB, &WeightValueB);
+				WeightValueB = DensityLowHigh[0] * (1 - WeightValueB) + DensityLowHigh[1] * WeightValueB;
 
 				return EdgeScalingFunction(WeightValueA, WeightValueB);
 			}
-
-			return 1.0;
+			else
+			{
+				return EdgeScalingFunction(DensityLowHigh[0], DensityLowHigh[0]);
+			}
 		};
 
 		// Set up constraints for Cloth Seam edges
@@ -998,717 +1039,828 @@ namespace UE::Chaos::ClothAsset::Private
 		return true;
 	}
 
-}
+	static void EmptySimSelections(const TSharedRef<FManagedArrayCollection>& ClothCollection)
+	{
+		FCollectionClothSelectionFacade SelectionFacade(ClothCollection);
 
+		const TArray<FName> SelectionNames = SelectionFacade.GetNames();
+		for (const FName& SelectionName : SelectionNames)
+		{
+			const FName GroupName = SelectionFacade.GetSelectionGroup(SelectionName);
+			if (GroupName == ClothCollectionGroup::SimVertices3D || GroupName == ClothCollectionGroup::SimVertices2D || GroupName == ClothCollectionGroup::SimFaces)
+			{
+				TSet<int32>& SelectionSet = SelectionFacade.GetSelectionSet(SelectionName);
+				SelectionSet.Empty();
+			}
+		}
+	}
 
-FChaosClothAssetRemeshNode::FChaosClothAssetRemeshNode(const UE::Dataflow::FNodeParameters& InParam, FGuid InGuid)
-	: FDataflowNode(InParam, InGuid)
-{
-	RegisterInputConnection(&Collection);
-	RegisterOutputConnection(&Collection, &Collection);
-	RegisterInputConnection(&DensityMapSim.StringValue, GET_MEMBER_NAME_CHECKED(FChaosClothAssetConnectableIStringValue, StringValue));
-	RegisterInputConnection(&DensityMapRender.StringValue, GET_MEMBER_NAME_CHECKED(FChaosClothAssetConnectableIStringValue, StringValue));
-}
+	static void RebuildTopologyDependentSimData(const TSharedRef<const FManagedArrayCollection>& InClothCollection,
+		const TSharedRef<FManagedArrayCollection>& OutClothCollection)
+	{
+		using namespace UE::Chaos::ClothAsset;
+		using namespace ::Chaos::Softs;
 
+		FCollectionClothConstFacade InClothFacade(InClothCollection);
+		FCollectionClothFacade OutClothFacade(OutClothCollection);
+		FCollectionClothSelectionFacade OutSelectionFacade(OutClothCollection);
+		OutSelectionFacade.DefineSchema();
 
-void FChaosClothAssetRemeshNode::EmptySimSelections(const TSharedRef<FManagedArrayCollection>& ClothCollection) const
-{
-	using namespace UE::Chaos::ClothAsset;
+		// Check that weight maps and skinning info have been interpolated over
+		for (const FName& InWeightMapName : InClothFacade.GetWeightMapNames())
+		{
+			const FName* FoundName = Algo::Find(OutClothFacade.GetWeightMapNames(), InWeightMapName);
+			checkf(FoundName, TEXT("Weight map %s was not copied to the output cloth collection"), *InWeightMapName.ToString());
+		}
+		if (InClothFacade.GetSimBoneIndices().Num() > 0 && OutClothFacade.GetNumSimVertices3D() > 0)
+		{
+			checkf(OutClothFacade.GetSimBoneIndices().Num() > 0, TEXT("Skinning bone indices not copied to the sim mesh of the output cloth collection"));
+		}
+		if (InClothFacade.GetSimBoneWeights().Num() > 0 && OutClothFacade.GetNumSimVertices3D() > 0)
+		{
+			checkf(OutClothFacade.GetSimBoneWeights().Num() > 0, TEXT("Skinning bone weights not copied to the sim mesh of the output cloth collection"));
+		}
 
-	FCollectionClothSelectionFacade SelectionFacade(ClothCollection);
+		FCollectionPropertyConstFacade InProperties(InClothCollection);
 
-	const TArray<FName> SelectionNames = SelectionFacade.GetNames();
-	for (const FName& SelectionName : SelectionNames)
-	{
-		const FName GroupName = SelectionFacade.GetSelectionGroup(SelectionName);
-		if (GroupName == ClothCollectionGroup::SimVertices3D || GroupName == ClothCollectionGroup::SimVertices2D || GroupName == ClothCollectionGroup::SimFaces)
+		// Reconstruct KinematicVertexSet
+		const FString MaxDistanceString = GET_MEMBER_NAME_STRING_CHECKED(FChaosClothAssetSimulationMaxDistanceConfigNode, MaxDistance);
+		const FName KinematicVertices3DName = GET_MEMBER_NAME_CHECKED(FChaosClothAssetSimulationMaxDistanceConfigNode, KinematicVertices3D);
+		if (InProperties.GetKeyIndex(MaxDistanceString) != INDEX_NONE)
 		{
-			TSet<int32>& SelectionSet = SelectionFacade.GetSelectionSet(SelectionName);
-			SelectionSet.Empty();
+			const FName MaxDistanceMapName(InProperties.GetStringValue(MaxDistanceString, MaxDistanceString));
+			OutSelectionFacade.FindOrAddSelectionSet(KinematicVertices3DName, ClothCollectionGroup::SimVertices3D) =
+				FClothGeometryTools::GenerateKinematicVertices3D(OutClothCollection, MaxDistanceMapName, InProperties.GetWeightedFloatValue(MaxDistanceString, FVector2f(0.f, 1.f)), NAME_None);
 		}
-	}
-}
-
-void FChaosClothAssetRemeshNode::RemeshSimMesh(const TSharedRef<const FManagedArrayCollection>& ClothCollection,
-	const FString& DensityMapName,
-	const TSharedRef<FManagedArrayCollection>& OutClothCollection) const
-{
-	using namespace UE::Geometry;
-	using namespace UE::Chaos::ClothAsset;
 
-	FCollectionClothConstFacade InClothFacade(ClothCollection);
+		// Reconstruct collision spheres
+		const FString SelfCollisionSphereStiffnessString = GET_MEMBER_NAME_STRING_CHECKED(FChaosClothAssetSimulationSelfCollisionSpheresConfigNode, SelfCollisionSphereStiffness);
+		const FString SelfCollisionSphereRadiusString = GET_MEMBER_NAME_STRING_CHECKED(FChaosClothAssetSimulationSelfCollisionSpheresConfigNode, SelfCollisionSphereRadius);
+		const FString SelfCollisionSphereRadiusCullMultiplierString = GET_MEMBER_NAME_STRING_CHECKED(FChaosClothAssetSimulationSelfCollisionSpheresConfigNode, SelfCollisionSphereRadiusCullMultiplier);
+		const FString SelfCollisionSphereSetNameString = GET_MEMBER_NAME_STRING_CHECKED(FChaosClothAssetSimulationSelfCollisionSpheresConfigNode, SelfCollisionSphereSetName);
 
-	if (InClothFacade.GetNumSimPatterns() == 0)
-	{
-		FClothGeometryTools::DeleteSimMesh(OutClothCollection);
-		return;
-	}
+		if (InProperties.GetKeyIndex(SelfCollisionSphereStiffnessString) != INDEX_NONE)
+		{
+			const float SelfCollisionSphereRadius = InProperties.GetValue<float>(SelfCollisionSphereRadiusString);
+			const float SelfCollisionSphereRadiusCullMultiplier = InProperties.GetValue<float>(SelfCollisionSphereRadiusCullMultiplierString);
+			const float CullDiameterSq = FMath::Square(SelfCollisionSphereRadius * SelfCollisionSphereRadiusCullMultiplier * 2.f);
 
-	// Convert input patterns to a DynamicMesh
+			if (OutClothFacade.IsValid() && CullDiameterSq > 0.f)
+			{
+				TConstArrayView<FVector3f> SimPositions = OutClothFacade.GetSimPosition3D();
+				TSet<int32> VertexSet;
+				FClothGeometryTools::SampleVertices(SimPositions, CullDiameterSq, VertexSet);
 
-	FClothPatternToDynamicMesh Converter;
+				const FName SelectionSetName(*InProperties.GetStringValue(SelfCollisionSphereSetNameString, SelfCollisionSphereSetNameString));
+				OutSelectionFacade.FindOrAddSelectionSet(SelectionSetName, ClothCollectionGroup::SimVertices3D) = VertexSet;
+			}
+		}
 
-	FDynamicMesh3 Mesh2D;
-	Converter.Convert(ClothCollection, INDEX_NONE, EClothPatternVertexType::Sim2D, Mesh2D);
+		// Reconstruct long-range attachments
 
-	const double TotalArea = TMeshQueries<FDynamicMesh3>::GetVolumeArea(Mesh2D).Y;
-	checkf(TotalArea > 0, TEXT("FChaosClothAssetRemeshNode::RemeshSimMesh: Expected 2D Sim mesh to have a positive area"));
-	const int TriangleCount = Mesh2D.TriangleCount();
+		// v1 (weight map)
+		PRAGMA_DISABLE_DEPRECATION_WARNINGS
+			const FString FixedEndWeightMapString = GET_MEMBER_NAME_STRING_CHECKED(FChaosClothAssetSimulationLongRangeAttachmentConfigNode, FixedEndWeightMap);
+		if (InProperties.GetKeyIndex(FixedEndWeightMapString) != INDEX_NONE)
+		{
+			FString UseGeodesicTethersString = GET_MEMBER_NAME_STRING_CHECKED(FChaosClothAssetSimulationLongRangeAttachmentConfigNode, bUseGeodesicTethers);
+			UseGeodesicTethersString.RemoveFromStart(TEXT("b"), ESearchCase::CaseSensitive);  // Property collection names doesn't use the b prefix for booleans
+			const bool bUseGeodesicTethers = InProperties.GetValue<bool>(UseGeodesicTethersString);
+			const FName FixedEndWeightMap(InProperties.GetStringValue(FixedEndWeightMapString));
 
-	// Copy pattern IDs into polygroup layer
+			UE::Chaos::ClothAsset::FClothEngineTools::GenerateTethers(OutClothCollection, FixedEndWeightMap, bUseGeodesicTethers);
+		}
+		PRAGMA_ENABLE_DEPRECATION_WARNINGS
 
-	Mesh2D.EnableAttributes();
-	const int32 PatternIndexLayerID = Mesh2D.Attributes()->NumPolygroupLayers();
-	Mesh2D.Attributes()->SetNumPolygroupLayers(PatternIndexLayerID + 1);
-	FDynamicMeshPolygroupAttribute* const PatternIndexLayer = Mesh2D.Attributes()->GetPolygroupLayer((int)PatternIndexLayerID);
-	check(Mesh2D.TriangleCount() == InClothFacade.GetNumSimFaces());
-	for (int32 FaceIndex = 0; FaceIndex < InClothFacade.GetNumSimFaces(); ++FaceIndex)
-	{
-		const int32 PatternID = InClothFacade.FindSimPatternByFaceIndex(FaceIndex);
-		PatternIndexLayer->SetValue(FaceIndex, PatternID);
+			// v2 (vertex set)
+			const FString FixedEndSetString = GET_MEMBER_NAME_STRING_CHECKED(FChaosClothAssetSimulationLongRangeAttachmentConfigNode_v2, FixedEndSet);
+		if (InProperties.GetKeyIndex(FixedEndSetString) != INDEX_NONE)
+		{
+			// Regenerate using the Kinematic vertices.
+			FString UseGeodesicTethersString = GET_MEMBER_NAME_STRING_CHECKED(FChaosClothAssetSimulationLongRangeAttachmentConfigNode_v2, bUseGeodesicTethers);
+			UseGeodesicTethersString.RemoveFromStart(TEXT("b"), ESearchCase::CaseSensitive);  // Property collection names doesn't use the b prefix for booleans
+			const bool bUseGeodesicTethers = InProperties.GetValue<bool>(UseGeodesicTethersString);
+			UE::Chaos::ClothAsset::FClothEngineTools::GenerateTethersFromSelectionSet(OutClothCollection, KinematicVertices3DName, bUseGeodesicTethers);
+		}
 	}
 
-	// Look for the DensityMap weight layer
-	const FDynamicMeshWeightAttribute* DensityMapLayer = nullptr;
-	if (Mesh2D.HasAttributes())
+
+	static void EmptyRenderSelections(const TSharedRef<FManagedArrayCollection>& ClothCollection)
 	{
-		for (int32 WeightLayerIndex = 0; WeightLayerIndex < Mesh2D.Attributes()->NumWeightLayers(); ++WeightLayerIndex)
+		using namespace UE::Chaos::ClothAsset;
+
+		FCollectionClothSelectionFacade SelectionFacade(ClothCollection);
+
+		const TArray<FName> SelectionNames = SelectionFacade.GetNames();
+		for (const FName& SelectionName : SelectionNames)
 		{
-			if (Mesh2D.Attributes()->GetWeightLayer(WeightLayerIndex)->GetName() == DensityMapName)
+			const FName GroupName = SelectionFacade.GetSelectionGroup(SelectionName);
+			if (GroupName == ClothCollectionGroup::RenderVertices || GroupName == ClothCollectionGroup::RenderFaces)
 			{
-				DensityMapLayer = Mesh2D.Attributes()->GetWeightLayer(WeightLayerIndex);
-				break;
+				TSet<int32>& SelectionSet = SelectionFacade.GetSelectionSet(SelectionName);
+				SelectionSet.Empty();
 			}
 		}
 	}
 
 	// Grab the seam information from the input collection
-	
 	// Seams are comprised of a set of Stitches. Each Stitch is simply a pair of vertex indices indicating vertices that should be welded to form the 3D mesh.
-	
 	// Stitches are given in random order within the Seam. To make remeshing them easier, we will find connected strips of stitches and store them in sequential order.
 	// So the vertices in Stitch N are connected to the vertices in Stitch N+1.
-
-	TArray<TArray<FIntVector2>> Seams;
-	for (int32 SeamIndex = 0; SeamIndex < InClothFacade.GetNumSeams(); ++SeamIndex)
+	static void GetSimMeshSeams(const TSharedRef<const FManagedArrayCollection>& ClothCollection, const FCollectionClothConstFacade& ClothFacade, const UE::Geometry::FDynamicMesh3& Mesh2D, TArray<TArray<FIntVector2>>& OutSeams)
 	{
-		FCollectionClothSeamConstFacade SeamFacade = InClothFacade.GetSeam(SeamIndex);
-		FClothGeometryTools::BuildConnectedSeams2D(ClothCollection, SeamIndex, Mesh2D, Seams);
-
-		// Check seams are valid
-		for (int32 SeamID = 0; SeamID < Seams.Num(); ++SeamID)
+		for (int32 SeamIndex = 0; SeamIndex < ClothFacade.GetNumSeams(); ++SeamIndex)
 		{
-			const TArray<FIntVector2>& SubSeam = Seams[SeamID];
-			for (int32 StitchID = 0; StitchID < SubSeam.Num() - 1; ++StitchID)
+			FCollectionClothSeamConstFacade SeamFacade = ClothFacade.GetSeam(SeamIndex);
+			FClothGeometryTools::BuildConnectedSeams2D(ClothCollection, SeamIndex, Mesh2D, OutSeams);
+
+			// Check seams are valid
+			for (int32 SeamID = 0; SeamID < OutSeams.Num(); ++SeamID)
 			{
-				const int32 NextStitchID = StitchID + 1;
-				for (int32 Side = 0; Side < 2; ++Side)
+				const TArray<FIntVector2>& SubSeam = OutSeams[SeamID];
+				for (int32 StitchID = 0; StitchID < SubSeam.Num() - 1; ++StitchID)
 				{
-					const int32 StitchVert = SubSeam[StitchID][Side];
-					const int32 NextStitchVert = SubSeam[NextStitchID][Side];
-
-					if (StitchVert == NextStitchVert)
+					const int32 NextStitchID = StitchID + 1;
+					for (int32 Side = 0; Side < 2; ++Side)
 					{
-						continue;
-					}
+						const int32 StitchVert = SubSeam[StitchID][Side];
+						const int32 NextStitchVert = SubSeam[NextStitchID][Side];
 
-					const int32 FoundEdge = Mesh2D.FindEdge(StitchVert, NextStitchVert);
+						if (StitchVert == NextStitchVert)
+						{
+							continue;
+						}
+
+						const int32 FoundEdge = Mesh2D.FindEdge(StitchVert, NextStitchVert);
 
-					// This would indicate a problem in BuildConnectedSeams
-					checkf(FoundEdge != UE::Geometry::FDynamicMesh3::InvalidID, TEXT("Could not find a mesh edge between sequential seam vertices %d, %d"), StitchVert, NextStitchVert);
+						// This would indicate a problem in BuildConnectedSeams
+						checkf(FoundEdge != UE::Geometry::FDynamicMesh3::InvalidID, TEXT("Could not find a mesh edge between sequential seam vertices %d, %d"), StitchVert, NextStitchVert);
+					}
 				}
 			}
 		}
-	}
 
 
-	// For remeshing purposes, we also want to find and constrain any vertices that connect the two seam sides together
-	// (This is only relevant for "internal" seams, which connect vertices within the same pattern)
-	for (TArray<FIntVector2>& Seam : Seams)
-	{
-		if (Seam.Num() == 0)
-		{
-			continue;
-		}
-
-		for (const FIntVector2& EndStitch : {Seam[0], Seam.Last()})
+		// For remeshing purposes, we also want to find and constrain any vertices that connect the two seam sides together
+		// (This is only relevant for "internal" seams, which connect vertices within the same pattern)
+		for (TArray<FIntVector2>& Seam : OutSeams)
 		{
-			for (const int NeighborA : Mesh2D.VtxVerticesItr(EndStitch[0]))
+			if (Seam.Num() == 0)
 			{
-				bool bMatchFound = false;
+				continue;
+			}
 
-				for (const int NeighborB : Mesh2D.VtxVerticesItr(EndStitch[1]))
+			for (const FIntVector2& EndStitch : { Seam[0], Seam.Last() })
+			{
+				for (const int NeighborA : Mesh2D.VtxVerticesItr(EndStitch[0]))
 				{
-					if (NeighborA == NeighborB)
-					{
-						if (Seam.Contains(FIntVector2(NeighborA, NeighborB)))
-						{
-							continue;
-						}
+					bool bMatchFound = false;
 
-						if (EndStitch == Seam[0])
-						{
-							Seam.Insert(FIntVector2(NeighborA, NeighborB), 0);
-						}
-						else
+					for (const int NeighborB : Mesh2D.VtxVerticesItr(EndStitch[1]))
+					{
+						if (NeighborA == NeighborB)
 						{
-							Seam.Add(FIntVector2(NeighborA, NeighborB));
+							if (Seam.Contains(FIntVector2(NeighborA, NeighborB)))
+							{
+								continue;
+							}
+
+							if (EndStitch == Seam[0])
+							{
+								Seam.Insert(FIntVector2(NeighborA, NeighborB), 0);
+							}
+							else
+							{
+								Seam.Add(FIntVector2(NeighborA, NeighborB));
+							}
+
+							bMatchFound = true;
+							break;
 						}
+					}
 
-						bMatchFound = true;
+					if (bMatchFound)
+					{
 						break;
 					}
 				}
-
-				if (bMatchFound)
-				{
-					break;
-				}
 			}
 		}
 	}
 
-	// Remesh seams
-
-	const int TargetTriangleCount = FMath::RoundToInt(static_cast<float>(TargetPercentSim) / 100.0f * static_cast<float>(TriangleCount));
-	const double TargetEdgeLength = FRemeshMeshOp::CalculateTargetEdgeLength(nullptr, TargetTriangleCount, TotalArea);
-
-	for (int32 ResampleIter = 0; ResampleIter < IterationsSim; ++ResampleIter)
+	static void RemeshSimMesh(const TSharedRef<const FManagedArrayCollection>& ClothCollection,
+		int32 IterationsSim,
+		double SmoothingSim,
+		const FString& DensityMapName,
+		const FVector2f& DensityLowHigh,
+		const TSharedRef<FManagedArrayCollection>& OutClothCollection)
 	{
-		UE::Chaos::ClothAsset::Private::RemeshSeams(Mesh2D, Seams, TargetEdgeLength, DensityMapLayer);
-	}
-
-	// Remesh boundaries
-
-	for (int32 ResampleIter = 0; ResampleIter < IterationsSim; ++ResampleIter)
-	{
-		UE::Chaos::ClothAsset::Private::RemeshBoundaries(Mesh2D, Seams, TargetEdgeLength, DensityMapLayer);
-	}
-
-	// Do the remeshing of the rest of the mesh
+		using namespace UE::Geometry;
+		using namespace UE::Chaos::ClothAsset;
 
-	UE::Geometry::FCompactMaps CompactMaps;
-	constexpr bool bUniformSmoothing = true;
-	UE::Chaos::ClothAsset::Private::Remesh(Mesh2D, TargetEdgeLength, IterationsSim, SmoothingSim, bUniformSmoothing, Seams, DensityMapName, &CompactMaps);
+		FCollectionClothConstFacade InClothFacade(ClothCollection);
 
-	// Update stitches
-	for (TArray<FIntVector2>& Seam : Seams)
-	{
-		for (FIntVector2& Stitch : Seam)
+		if (InClothFacade.GetNumSimPatterns() == 0)
 		{
-			Stitch[0] = CompactMaps.GetVertexMapping(Stitch[0]);
-			Stitch[1] = CompactMaps.GetVertexMapping(Stitch[1]);
-			checkf(Stitch[0] != FDynamicMesh3::InvalidID, TEXT("Stitch vertex %d was deleted by remeshing"), Stitch[0]);
-			checkf(Stitch[1] != FDynamicMesh3::InvalidID, TEXT("Stitch vertex %d was deleted by remeshing"), Stitch[1]);
+			FClothGeometryTools::DeleteSimMesh(OutClothCollection);
+			return;
 		}
-	}
 
+		// Convert input patterns to a DynamicMesh
 
-	// Project the 3D vertices onto the input 3D mesh
+		FClothPatternToDynamicMesh Converter;
 
-	// For each 2D vertex, we will find the closest triangle on the input 2D mesh, then look up that triangle on the input 3D mesh to get the final 3D location.
-	// We will do this pattern-by-pattern to handle issues where the patterns overlap in 2D space.
+		FDynamicMesh3 Mesh2D;
+		Converter.Convert(ClothCollection, INDEX_NONE, EClothPatternVertexType::Sim2D, Mesh2D);
 
-	const FDynamicMeshPolygroupAttribute* const NewPatternIndexLayer = Mesh2D.Attributes()->GetPolygroupLayer((int)PatternIndexLayerID);
-	check(NewPatternIndexLayer);
+		const double TotalArea = TMeshQueries<FDynamicMesh3>::GetVolumeArea(Mesh2D).Y;
+		checkf(TotalArea > 0, TEXT("FChaosClothAssetRemeshNode::RemeshSimMesh: Expected 2D Sim mesh to have a positive area"));
+		const int TriangleCount = Mesh2D.TriangleCount();
 
-	FDynamicMesh3 Mesh3D;
-	Mesh3D.Copy(Mesh2D);
+		// Copy pattern IDs into polygroup layer
 
-	TMap<int32, TSet<int32>> PatternVertexIDs;
-	for (int32 TID : Mesh2D.TriangleIndicesItr())
-	{
-		const int32 PatternID = NewPatternIndexLayer->GetValue(TID);
-		if (!PatternVertexIDs.Contains(PatternID))
+		Mesh2D.EnableAttributes();
+		const int32 PatternIndexLayerID = Mesh2D.Attributes()->NumPolygroupLayers();
+		Mesh2D.Attributes()->SetNumPolygroupLayers(PatternIndexLayerID + 1);
+		FDynamicMeshPolygroupAttribute* const PatternIndexLayer = Mesh2D.Attributes()->GetPolygroupLayer((int)PatternIndexLayerID);
+		check(Mesh2D.TriangleCount() == InClothFacade.GetNumSimFaces());
+		for (int32 FaceIndex = 0; FaceIndex < InClothFacade.GetNumSimFaces(); ++FaceIndex)
 		{
-			PatternVertexIDs.Add(PatternID, TSet<int32>());
+			const int32 PatternID = InClothFacade.FindSimPatternByFaceIndex(FaceIndex);
+			PatternIndexLayer->SetValue(FaceIndex, PatternID);
 		}
 
-		const FIndex3i Tri = Mesh2D.GetTriangle(TID);
-		PatternVertexIDs[PatternID].Add(Tri[0]);
-		PatternVertexIDs[PatternID].Add(Tri[1]);
-		PatternVertexIDs[PatternID].Add(Tri[2]);
-	}
-
-	for (int32 PatternID = 0; PatternID < InClothFacade.GetNumSimPatterns(); ++PatternID)
-	{
-		if (!PatternVertexIDs.Contains(PatternID))
+		// Look for the DensityMap weight layer
+		const FDynamicMeshWeightAttribute* DensityMapLayer = nullptr;
+		if (Mesh2D.HasAttributes())
 		{
-			continue;
+			for (int32 WeightLayerIndex = 0; WeightLayerIndex < Mesh2D.Attributes()->NumWeightLayers(); ++WeightLayerIndex)
+			{
+				if (Mesh2D.Attributes()->GetWeightLayer(WeightLayerIndex)->GetName() == DensityMapName)
+				{
+					DensityMapLayer = Mesh2D.Attributes()->GetWeightLayer(WeightLayerIndex);
+					break;
+				}
+			}
 		}
 
-		FDynamicMesh3 ProjectionTarget2D;
-		Converter.Convert(ClothCollection, PatternID, EClothPatternVertexType::Sim2D, ProjectionTarget2D);
-		FDynamicMesh3 ProjectionTarget3D;
-		Converter.Convert(ClothCollection, PatternID, EClothPatternVertexType::Sim3D, ProjectionTarget3D);
-		TSharedPtr<UE::Geometry::FDynamicMeshAABBTree3> ProjectionTargetSpatial = MakeShared<UE::Geometry::FDynamicMeshAABBTree3>(&ProjectionTarget2D, true);
 
-		for (const int32 VID : PatternVertexIDs[PatternID])
-		{
-			const FVector3d SrcVert = Mesh2D.GetVertex(VID);
+		TArray<TArray<FIntVector2>> Seams;
+		Private::GetSimMeshSeams(ClothCollection, InClothFacade, Mesh2D, Seams);
 
-			double Distance;
-			const int NearestTriangle = ProjectionTargetSpatial->FindNearestTriangle(SrcVert, Distance);
 
-			const FDistPoint3Triangle3d Dist = TMeshQueries<FDynamicMesh3>::TriangleDistance(ProjectionTarget2D, NearestTriangle, SrcVert);
-			const FVector3d Bary = Dist.TriangleBaryCoords;
-			const FVector3d InterpolatedPoint = ProjectionTarget3D.GetTriBaryPoint(NearestTriangle, Bary[0], Bary[1], Bary[2]);
+		// Remesh seams
 
-			Mesh3D.SetVertex(VID, InterpolatedPoint);
-		}
-	}
+		const double TargetEdgeLength = TMeshQueries<FDynamicMesh3>::AverageEdgeLength(Mesh2D);
 
-	// Build the output cloth sim mesh
+		for (int32 ResampleIter = 0; ResampleIter < IterationsSim; ++ResampleIter)
+		{
+			UE::Chaos::ClothAsset::Private::RemeshSeams(Mesh2D, Seams, TargetEdgeLength, DensityMapLayer, DensityLowHigh);
+		}
 
-	constexpr bool bAppendToExistingMesh = false;
-	constexpr bool bTransferWeightMaps = true;
-	constexpr bool bTransferSimSkinningData = true;
-	TMap<int, int32> DynamicMeshToClothVertexMap;
-	FClothGeometryTools::BuildSimMeshFromDynamicMeshes(OutClothCollection, Mesh2D, Mesh3D, PatternIndexLayerID, bTransferWeightMaps, bTransferSimSkinningData, bAppendToExistingMesh, DynamicMeshToClothVertexMap);
+		// Remesh boundaries
 
+		for (int32 ResampleIter = 0; ResampleIter < IterationsSim; ++ResampleIter)
+		{
+			UE::Chaos::ClothAsset::Private::RemeshBoundaries(Mesh2D, Seams, TargetEdgeLength, DensityMapLayer, DensityLowHigh);
+		}
 
-	// Re-apply the seam info from the input sim mesh. This will create a new Seam for each set of connected stitches
+		// Do the remeshing of the rest of the mesh
 
-	FCollectionClothFacade OutClothFacade(OutClothCollection);
-	for (int32 SeamIndex = 0; SeamIndex < Seams.Num(); ++SeamIndex)
-	{
-		const TArray<FIntVector2>& Seam = Seams[SeamIndex];
+		UE::Geometry::FCompactMaps CompactMaps;
+		constexpr bool bUniformSmoothing = true;
+		UE::Chaos::ClothAsset::Private::Remesh(Mesh2D, TargetEdgeLength, IterationsSim, SmoothingSim, bUniformSmoothing, Seams, DensityMapName, DensityLowHigh, &CompactMaps);
 
-		TArray<FIntVector2> NewSeam;
-		for (const FIntVector2& Stitch : Seam)
+		// Update stitches
+		for (TArray<FIntVector2>& Seam : Seams)
 		{
-			if (Stitch[0] == Stitch[1])
+			for (FIntVector2& Stitch : Seam)
 			{
-				continue;
+				Stitch[0] = CompactMaps.GetVertexMapping(Stitch[0]);
+				Stitch[1] = CompactMaps.GetVertexMapping(Stitch[1]);
+				checkf(Stitch[0] != FDynamicMesh3::InvalidID, TEXT("Stitch vertex %d was deleted by remeshing"), Stitch[0]);
+				checkf(Stitch[1] != FDynamicMesh3::InvalidID, TEXT("Stitch vertex %d was deleted by remeshing"), Stitch[1]);
 			}
-
-			FIntVector2 NewStitch;
-			NewStitch[0] = DynamicMeshToClothVertexMap[Stitch[0]];
-			NewStitch[1] = DynamicMeshToClothVertexMap[Stitch[1]];
-			NewSeam.Add(NewStitch);
 		}
 
-		FCollectionClothSeamFacade NewSeamFacade = OutClothFacade.AddGetSeam();
-		NewSeamFacade.Initialize(NewSeam);
-	}
-}
 
+		// Project the 3D vertices onto the input 3D mesh
 
-void FChaosClothAssetRemeshNode::RebuildTopologyDependentSimData(const TSharedRef<const FManagedArrayCollection>& InClothCollection,
-	const TSharedRef<FManagedArrayCollection>& OutClothCollection) const
-{
-	using namespace UE::Chaos::ClothAsset;
-	using namespace ::Chaos::Softs;
+		// For each 2D vertex, we will find the closest triangle on the input 2D mesh, then look up that triangle on the input 3D mesh to get the final 3D location.
+		// We will do this pattern-by-pattern to handle issues where the patterns overlap in 2D space.
 
-	FCollectionClothConstFacade InClothFacade(InClothCollection);
-	FCollectionClothFacade OutClothFacade(OutClothCollection);
-	FCollectionClothSelectionFacade OutSelectionFacade(OutClothCollection);
-	OutSelectionFacade.DefineSchema();
+		const FDynamicMeshPolygroupAttribute* const NewPatternIndexLayer = Mesh2D.Attributes()->GetPolygroupLayer((int)PatternIndexLayerID);
+		check(NewPatternIndexLayer);
 
-	// Check that weight maps and skinning info have been interpolated over
-	for (const FName& InWeightMapName : InClothFacade.GetWeightMapNames())
-	{
-		const FName* FoundName = Algo::Find(OutClothFacade.GetWeightMapNames(), InWeightMapName);
-		checkf(FoundName, TEXT("Weight map %s was not copied to the output cloth collection"), *InWeightMapName.ToString());
-	}
-	if (InClothFacade.GetSimBoneIndices().Num() > 0 && OutClothFacade.GetNumSimVertices3D() > 0)
-	{
-		checkf(OutClothFacade.GetSimBoneIndices().Num() > 0, TEXT("Skinning bone indices not copied to the sim mesh of the output cloth collection"));
-	}
-	if (InClothFacade.GetSimBoneWeights().Num() > 0 && OutClothFacade.GetNumSimVertices3D() > 0)
-	{
-		checkf(OutClothFacade.GetSimBoneWeights().Num() > 0, TEXT("Skinning bone weights not copied to the sim mesh of the output cloth collection"));
-	}
+		FDynamicMesh3 Mesh3D;
+		Mesh3D.Copy(Mesh2D);
 
-	Chaos::Softs::FCollectionPropertyConstFacade InProperties(InClothCollection);
+		TMap<int32, TSet<int32>> PatternVertexIDs;
+		for (int32 TID : Mesh2D.TriangleIndicesItr())
+		{
+			const int32 PatternID = NewPatternIndexLayer->GetValue(TID);
+			if (!PatternVertexIDs.Contains(PatternID))
+			{
+				PatternVertexIDs.Add(PatternID, TSet<int32>());
+			}
 
-	// Reconstruct KinematicVertexSet
-	const FString MaxDistanceString = GET_MEMBER_NAME_STRING_CHECKED(FChaosClothAssetSimulationMaxDistanceConfigNode, MaxDistance);
-	const FName KinematicVertices3DName = GET_MEMBER_NAME_CHECKED(FChaosClothAssetSimulationMaxDistanceConfigNode, KinematicVertices3D);
-	if (InProperties.GetKeyIndex(MaxDistanceString) != INDEX_NONE)
-	{
-		const FName MaxDistanceMapName(InProperties.GetStringValue(MaxDistanceString, MaxDistanceString));
-		OutSelectionFacade.FindOrAddSelectionSet(KinematicVertices3DName, ClothCollectionGroup::SimVertices3D) =
-			FClothGeometryTools::GenerateKinematicVertices3D(OutClothCollection, MaxDistanceMapName, InProperties.GetWeightedFloatValue(MaxDistanceString, FVector2f(0.f, 1.f)), NAME_None);
-	}
+			const FIndex3i Tri = Mesh2D.GetTriangle(TID);
+			PatternVertexIDs[PatternID].Add(Tri[0]);
+			PatternVertexIDs[PatternID].Add(Tri[1]);
+			PatternVertexIDs[PatternID].Add(Tri[2]);
+		}
 
-	// Reconstruct collision spheres
-	const FString SelfCollisionSphereStiffnessString = GET_MEMBER_NAME_STRING_CHECKED(FChaosClothAssetSimulationSelfCollisionSpheresConfigNode, SelfCollisionSphereStiffness);
-	const FString SelfCollisionSphereRadiusString = GET_MEMBER_NAME_STRING_CHECKED(FChaosClothAssetSimulationSelfCollisionSpheresConfigNode, SelfCollisionSphereRadius);
-	const FString SelfCollisionSphereRadiusCullMultiplierString = GET_MEMBER_NAME_STRING_CHECKED(FChaosClothAssetSimulationSelfCollisionSpheresConfigNode, SelfCollisionSphereRadiusCullMultiplier);
-	const FString SelfCollisionSphereSetNameString = GET_MEMBER_NAME_STRING_CHECKED(FChaosClothAssetSimulationSelfCollisionSpheresConfigNode, SelfCollisionSphereSetName);
+		for (int32 PatternID = 0; PatternID < InClothFacade.GetNumSimPatterns(); ++PatternID)
+		{
+			if (!PatternVertexIDs.Contains(PatternID))
+			{
+				continue;
+			}
 
-	if (InProperties.GetKeyIndex(SelfCollisionSphereStiffnessString) != INDEX_NONE)
-	{
-		const float SelfCollisionSphereRadius = InProperties.GetValue<float>(SelfCollisionSphereRadiusString);
-		const float SelfCollisionSphereRadiusCullMultiplier = InProperties.GetValue<float>(SelfCollisionSphereRadiusCullMultiplierString);
-		const float CullDiameterSq = FMath::Square(SelfCollisionSphereRadius * SelfCollisionSphereRadiusCullMultiplier * 2.f);
+			FDynamicMesh3 ProjectionTarget2D;
+			Converter.Convert(ClothCollection, PatternID, EClothPatternVertexType::Sim2D, ProjectionTarget2D);
+			FDynamicMesh3 ProjectionTarget3D;
+			Converter.Convert(ClothCollection, PatternID, EClothPatternVertexType::Sim3D, ProjectionTarget3D);
+			TSharedPtr<UE::Geometry::FDynamicMeshAABBTree3> ProjectionTargetSpatial = MakeShared<UE::Geometry::FDynamicMeshAABBTree3>(&ProjectionTarget2D, true);
 
-		if (OutClothFacade.IsValid() && CullDiameterSq > 0.f)
-		{
-			TConstArrayView<FVector3f> SimPositions = OutClothFacade.GetSimPosition3D();
-			TSet<int32> VertexSet;
-			FClothGeometryTools::SampleVertices(SimPositions, CullDiameterSq, VertexSet);
+			for (const int32 VID : PatternVertexIDs[PatternID])
+			{
+				const FVector3d SrcVert = Mesh2D.GetVertex(VID);
 
-			const FName SelectionSetName(*InProperties.GetStringValue(SelfCollisionSphereSetNameString, SelfCollisionSphereSetNameString));
-			OutSelectionFacade.FindOrAddSelectionSet(SelectionSetName, ClothCollectionGroup::SimVertices3D) = VertexSet;
+				double Distance;
+				const int NearestTriangle = ProjectionTargetSpatial->FindNearestTriangle(SrcVert, Distance);
+
+				const FDistPoint3Triangle3d Dist = TMeshQueries<FDynamicMesh3>::TriangleDistance(ProjectionTarget2D, NearestTriangle, SrcVert);
+				const FVector3d Bary = Dist.TriangleBaryCoords;
+				const FVector3d InterpolatedPoint = ProjectionTarget3D.GetTriBaryPoint(NearestTriangle, Bary[0], Bary[1], Bary[2]);
+
+				Mesh3D.SetVertex(VID, InterpolatedPoint);
+			}
 		}
-	}
 
-	// Reconstruct long-range attachments
+		// Build the output cloth sim mesh
 
-	// v1 (weight map)
-PRAGMA_DISABLE_DEPRECATION_WARNINGS
-	const FString FixedEndWeightMapString = GET_MEMBER_NAME_STRING_CHECKED(FChaosClothAssetSimulationLongRangeAttachmentConfigNode, FixedEndWeightMap);
-	if (InProperties.GetKeyIndex(FixedEndWeightMapString) != INDEX_NONE)
-	{
-		FString UseGeodesicTethersString = GET_MEMBER_NAME_STRING_CHECKED(FChaosClothAssetSimulationLongRangeAttachmentConfigNode, bUseGeodesicTethers);
-		UseGeodesicTethersString.RemoveFromStart(TEXT("b"), ESearchCase::CaseSensitive);  // Property collection names doesn't use the b prefix for booleans
-		const bool bUseGeodesicTethers = InProperties.GetValue<bool>(UseGeodesicTethersString);
-		const FName FixedEndWeightMap(InProperties.GetStringValue(FixedEndWeightMapString));
+		constexpr bool bAppendToExistingMesh = false;
+		constexpr bool bTransferWeightMaps = true;
+		constexpr bool bTransferSimSkinningData = true;
+		TMap<int, int32> DynamicMeshToClothVertexMap;
+		FClothGeometryTools::BuildSimMeshFromDynamicMeshes(OutClothCollection, Mesh2D, Mesh3D, PatternIndexLayerID, bTransferWeightMaps, bTransferSimSkinningData, bAppendToExistingMesh, DynamicMeshToClothVertexMap);
 
-		UE::Chaos::ClothAsset::FClothEngineTools::GenerateTethers(OutClothCollection, FixedEndWeightMap, bUseGeodesicTethers);
-	}
-PRAGMA_ENABLE_DEPRECATION_WARNINGS
 
-	// v2 (vertex set)
-	const FString FixedEndSetString = GET_MEMBER_NAME_STRING_CHECKED(FChaosClothAssetSimulationLongRangeAttachmentConfigNode_v2, FixedEndSet);
-	if (InProperties.GetKeyIndex(FixedEndSetString) != INDEX_NONE)
-	{
-		// Regenerate using the Kinematic vertices.
-		FString UseGeodesicTethersString = GET_MEMBER_NAME_STRING_CHECKED(FChaosClothAssetSimulationLongRangeAttachmentConfigNode_v2, bUseGeodesicTethers);
-		UseGeodesicTethersString.RemoveFromStart(TEXT("b"), ESearchCase::CaseSensitive);  // Property collection names doesn't use the b prefix for booleans
-		const bool bUseGeodesicTethers = InProperties.GetValue<bool>(UseGeodesicTethersString);
-		UE::Chaos::ClothAsset::FClothEngineTools::GenerateTethersFromSelectionSet(OutClothCollection, KinematicVertices3DName, bUseGeodesicTethers);
-	}
-}
+		// Re-apply the seam info from the input sim mesh. This will create a new Seam for each set of connected stitches
 
+		FCollectionClothFacade OutClothFacade(OutClothCollection);
+		for (int32 SeamIndex = 0; SeamIndex < Seams.Num(); ++SeamIndex)
+		{
+			const TArray<FIntVector2>& Seam = Seams[SeamIndex];
 
-void FChaosClothAssetRemeshNode::EmptyRenderSelections(const TSharedRef<FManagedArrayCollection>& ClothCollection) const
-{
-	using namespace UE::Chaos::ClothAsset;
+			TArray<FIntVector2> NewSeam;
+			for (const FIntVector2& Stitch : Seam)
+			{
+				if (Stitch[0] == Stitch[1])
+				{
+					continue;
+				}
 
-	FCollectionClothSelectionFacade SelectionFacade(ClothCollection);
+				FIntVector2 NewStitch;
+				NewStitch[0] = DynamicMeshToClothVertexMap[Stitch[0]];
+				NewStitch[1] = DynamicMeshToClothVertexMap[Stitch[1]];
+				NewSeam.Add(NewStitch);
+			}
 
-	const TArray<FName> SelectionNames = SelectionFacade.GetNames();
-	for (const FName& SelectionName : SelectionNames)
-	{
-		const FName GroupName = SelectionFacade.GetSelectionGroup(SelectionName);
-		if (GroupName == ClothCollectionGroup::RenderVertices || GroupName == ClothCollectionGroup::RenderFaces)
-		{
-			TSet<int32>& SelectionSet = SelectionFacade.GetSelectionSet(SelectionName);
-			SelectionSet.Empty();
+			FCollectionClothSeamFacade NewSeamFacade = OutClothFacade.AddGetSeam();
+			NewSeamFacade.Initialize(NewSeam);
 		}
 	}
-}
 
 
+	static void RemeshRenderMesh(const TSharedRef<const FManagedArrayCollection>& ClothCollection,
+		EChaosClothAssetRemeshMethod RemeshMethodRender,
+		bool bRemeshRenderSeams,
+		int32 RenderSeamRemeshIterations,
+		int32 IterationsRender,
+		double SmoothingRender,
+		const FString& DensityMapName,
+		const FVector2f& DensityLowHigh,
+		int32 TargetPercentRender,
+		const TSharedRef<FManagedArrayCollection>& OutClothCollection)
+	{
+		using namespace UE::Geometry;
+		using namespace UE::Chaos::ClothAsset;
 
+		// Get the source mesh
+		FClothPatternToDynamicMesh Converter;
+		FDynamicMesh3 DynamicMesh;
 
-void FChaosClothAssetRemeshNode::RemeshRenderMesh(const TSharedRef<const FManagedArrayCollection>& ClothCollection,
-	const FString& DensityMapName,
-	const TSharedRef<FManagedArrayCollection>& OutClothCollection) const
-{
-	using namespace UE::Geometry;
-	using namespace UE::Chaos::ClothAsset;
-
-	// Get the source mesh
-	FClothPatternToDynamicMesh Converter;
-	FDynamicMesh3 DynamicMesh;
-
-	// NOTE: When applied to the Render mesh, this Convert function will assign PatternIDs to the MaterialID attribute of the DynamicMesh. 
-	// After remeshing we will use the MatrialID attribute to determine which triangles should go into which output pattern.
-	Converter.Convert(ClothCollection, INDEX_NONE, EClothPatternVertexType::Render, DynamicMesh);
+		// NOTE: When applied to the Render mesh, this Convert function will assign PatternIDs to the MaterialID attribute of the DynamicMesh. 
+		// After remeshing we will use the MatrialID attribute to determine which triangles should go into which output pattern.
+		Converter.Convert(ClothCollection, INDEX_NONE, EClothPatternVertexType::Render, DynamicMesh);
 
-	check(DynamicMesh.HasAttributes());
+		check(DynamicMesh.HasAttributes());
 
-	const int InputMeshVertexCount = DynamicMesh.VertexCount();
-	const int InputMeshTriangleCount = DynamicMesh.TriangleCount();
+		const int InputMeshVertexCount = DynamicMesh.VertexCount();
+		const int InputMeshTriangleCount = DynamicMesh.TriangleCount();
 
-	const bool bHasUVs = (DynamicMesh.Attributes()->PrimaryUV() != nullptr);
+		const bool bHasUVs = (DynamicMesh.Attributes()->PrimaryUV() != nullptr);
 
-	const int TargetTriangleCount = FMath::RoundToInt(static_cast<float>(TargetPercentRender) / 100.0f * static_cast<float>(InputMeshTriangleCount));
-	const double TargetEdgeLength = FRemeshMeshOp::CalculateTargetEdgeLength(&DynamicMesh, TargetTriangleCount);
+		const double TargetEdgeLength = TMeshQueries<FDynamicMesh3>::AverageEdgeLength(DynamicMesh);
 
-	TArray<TArray<FIntVector2>> Seams;
+		TArray<TArray<FIntVector2>> Seams;
 
-	if (bRemeshRenderSeams)
-	{
-		// Create pseudo-stitches based on boundary vertex proximity. These stitches aren't going to actually weld vertices together, but they will guide boundary remeshing.
-		// The goal is to maintain a vertex pairing along boundaries in order to avoid holes opening up when the mesh deforms due to skinning.
-		TArray<FIntVector2> Stitches;
-		UE::Chaos::ClothAsset::Private::FindCoincidentBoundaryVertices(DynamicMesh, Stitches);
+		if (bRemeshRenderSeams)
+		{
+			// Create pseudo-stitches based on boundary vertex proximity. These stitches aren't going to actually weld vertices together, but they will guide boundary remeshing.
+			// The goal is to maintain a vertex pairing along boundaries in order to avoid holes opening up when the mesh deforms due to skinning.
+			TArray<FIntVector2> Stitches;
+			UE::Chaos::ClothAsset::Private::FindCoincidentBoundaryVertices(DynamicMesh, Stitches);
 
-		FClothGeometryTools::BuildConnectedSeams(Stitches, DynamicMesh, Seams);
+			FClothGeometryTools::BuildConnectedSeams(Stitches, DynamicMesh, Seams);
 
-		// Add density map for for Render mesh
-		const UE::Geometry::FDynamicMeshWeightAttribute* DensityMapLayer = nullptr;
-		if (DynamicMesh.HasAttributes())
-		{
-			for (int32 WeightLayerIndex = 0; WeightLayerIndex < DynamicMesh.Attributes()->NumWeightLayers(); ++WeightLayerIndex)
+			// Add density map for for Render mesh
+			const UE::Geometry::FDynamicMeshWeightAttribute* DensityMapLayer = nullptr;
+			if (DynamicMesh.HasAttributes())
 			{
-				if (DynamicMesh.Attributes()->GetWeightLayer(WeightLayerIndex)->GetName() == DensityMapName)
+				for (int32 WeightLayerIndex = 0; WeightLayerIndex < DynamicMesh.Attributes()->NumWeightLayers(); ++WeightLayerIndex)
 				{
-					DensityMapLayer = DynamicMesh.Attributes()->GetWeightLayer(WeightLayerIndex);
-					break;
+					if (DynamicMesh.Attributes()->GetWeightLayer(WeightLayerIndex)->GetName() == DensityMapName)
+					{
+						DensityMapLayer = DynamicMesh.Attributes()->GetWeightLayer(WeightLayerIndex);
+						break;
+					}
 				}
 			}
+
+			for (int RemeshPass = 0; RemeshPass < RenderSeamRemeshIterations; ++RemeshPass)
+			{
+				UE::Chaos::ClothAsset::Private::RemeshSeams(DynamicMesh, Seams, TargetEdgeLength, DensityMapLayer, DensityLowHigh);
+			}
+
+			// Also remesh the open boundaries that are not constrained by seams
+			for (int RemeshPass = 0; RemeshPass < RenderSeamRemeshIterations; ++RemeshPass)
+			{
+				UE::Chaos::ClothAsset::Private::RemeshBoundaries(DynamicMesh, Seams, TargetEdgeLength, DensityMapLayer, DensityLowHigh);
+			}
 		}
 
-		for (int RemeshPass = 0; RemeshPass < RenderSeamRemeshIterations; ++RemeshPass)
+
+		UE::Geometry::FCompactMaps CompactMaps;
+		if (RemeshMethodRender == EChaosClothAssetRemeshMethod::Remesh)
 		{
-			UE::Chaos::ClothAsset::Private::RemeshSeams(DynamicMesh, Seams, TargetEdgeLength, DensityMapLayer);
+			constexpr bool bUniformSmoothing = false;	// uniform smoothing can distort the UV layer pretty badly
+			const bool bSuccess = UE::Chaos::ClothAsset::Private::Remesh(DynamicMesh, TargetEdgeLength, IterationsRender, SmoothingRender, bUniformSmoothing, Seams, DensityMapName, DensityLowHigh, &CompactMaps);
+			check(bSuccess);
 		}
-
-		// Also remesh the open boundaries that are not constrained by seams
-		for (int RemeshPass = 0; RemeshPass < RenderSeamRemeshIterations; ++RemeshPass)
+		else
 		{
-			UE::Chaos::ClothAsset::Private::RemeshBoundaries(DynamicMesh, Seams, TargetEdgeLength, DensityMapLayer);
+			const bool bCoarsenBoundariesDuringSimplify = !bRemeshRenderSeams;
+			const int TargetVertexCount = FMath::RoundToInt(static_cast<float>(TargetPercentRender) / 100.0f * static_cast<float>(InputMeshVertexCount));
+			UE::Chaos::ClothAsset::Private::Simplify(DynamicMesh, TargetVertexCount, bCoarsenBoundariesDuringSimplify, &CompactMaps);
 		}
-	}
 
+		// Collect outputs
 
-	UE::Geometry::FCompactMaps CompactMaps;
-	if (RemeshMethodRender == EChaosClothAssetRemeshMethod::Remesh)
-	{
-		constexpr bool bUniformSmoothing = false;	// uniform smoothing can distort the UV layer pretty badly
-		const bool bSuccess = UE::Chaos::ClothAsset::Private::Remesh(DynamicMesh, TargetEdgeLength, IterationsRender, SmoothingRender, bUniformSmoothing, Seams, DensityMapName,  &CompactMaps);
-		check(bSuccess);
-	}
-	else
-	{
-		const bool bCoarsenBoundariesDuringSimplify = !bRemeshRenderSeams;
-		const int TargetVertexCount = FMath::RoundToInt(static_cast<float>(TargetPercentRender) / 100.0f * static_cast<float>(InputMeshVertexCount));
-		UE::Chaos::ClothAsset::Private::Simplify(DynamicMesh, TargetVertexCount, bCoarsenBoundariesDuringSimplify, &CompactMaps);
-	}
+		//
+		// Normals
+		//
 
-	// Collect outputs
+		const bool bHasNormals = (DynamicMesh.Attributes()->PrimaryNormals() != nullptr);
+		check(bHasNormals);
 
-	//
-	// Normals
-	//
+		TArray<FVector3f> Normals;
+		Normals.SetNum(DynamicMesh.VertexCount());
+		const FDynamicMeshNormalOverlay* const NormalOverlay = DynamicMesh.Attributes()->PrimaryNormals();
+		for (const int TriangleIndex : DynamicMesh.TriangleIndicesItr())
+		{
+			const FIndex3i Tri = DynamicMesh.GetTriangle(TriangleIndex);
 
-	const bool bHasNormals = (DynamicMesh.Attributes()->PrimaryNormals() != nullptr);
-	check(bHasNormals);
+			for (int TriangleVertexIndex = 0; TriangleVertexIndex < 3; ++TriangleVertexIndex)
+			{
+				const int VertexIndex = Tri[TriangleVertexIndex];
 
-	TArray<FVector3f> Normals;
-	Normals.SetNum(DynamicMesh.VertexCount());
-	const FDynamicMeshNormalOverlay* const NormalOverlay = DynamicMesh.Attributes()->PrimaryNormals();
-	for (const int TriangleIndex : DynamicMesh.TriangleIndicesItr())
-	{
-		const FIndex3i Tri = DynamicMesh.GetTriangle(TriangleIndex);
+				// NOTE: This assumes one normal per vertex in the overlay (i.e. no "hard edges")
+				Normals[VertexIndex] = NormalOverlay->GetElementAtVertex(TriangleIndex, VertexIndex);
+			}
+		}
+
+		//
+		// Tangents
+		//
 
-		for (int TriangleVertexIndex = 0; TriangleVertexIndex < 3; ++TriangleVertexIndex)
+		TArray<FVector3f> TangentUs;
+		TArray<FVector3f> TangentVs;
 		{
-			const int VertexIndex = Tri[TriangleVertexIndex];
+			const bool bHasTangentUs = (DynamicMesh.Attributes()->PrimaryTangents() != nullptr);
+			const bool bHasTangentVs = (DynamicMesh.Attributes()->PrimaryBiTangents() != nullptr);
+			if (!bHasTangentUs || !bHasTangentVs)
+			{
+				FMeshTangentsf::ComputeDefaultOverlayTangents(DynamicMesh);
+			}
+			TangentUs.SetNumZeroed(DynamicMesh.VertexCount());
+			TangentVs.SetNumZeroed(DynamicMesh.VertexCount());
 
-			// NOTE: This assumes one normal per vertex in the overlay (i.e. no "hard edges")
-			Normals[VertexIndex] = NormalOverlay->GetElementAtVertex(TriangleIndex, VertexIndex);
-		}
-	}
+			const FDynamicMeshNormalOverlay* const TangentUOverlay = DynamicMesh.Attributes()->PrimaryTangents();
+			const FDynamicMeshNormalOverlay* const TangentVOverlay = DynamicMesh.Attributes()->PrimaryBiTangents();
 
-	//
-	// Tangents
-	//
+			for (const int TriangleIndex : DynamicMesh.TriangleIndicesItr())
+			{
+				const FIndex3i Tri = DynamicMesh.GetTriangle(TriangleIndex);
 
-	TArray<FVector3f> TangentUs;
-	TArray<FVector3f> TangentVs;
-	{
-		const bool bHasTangentUs = (DynamicMesh.Attributes()->PrimaryTangents() != nullptr);
-		const bool bHasTangentVs = (DynamicMesh.Attributes()->PrimaryBiTangents() != nullptr);
-		if (!bHasTangentUs || !bHasTangentVs)
-		{
-			FMeshTangentsf::ComputeDefaultOverlayTangents(DynamicMesh);
+				for (int TriangleVertexIndex = 0; TriangleVertexIndex < 3; ++TriangleVertexIndex)
+				{
+					const int VertexIndex = Tri[TriangleVertexIndex];
+
+					TangentUs[VertexIndex] += TangentUOverlay->GetElementAtVertex(TriangleIndex, VertexIndex);
+					TangentVs[VertexIndex] += TangentVOverlay->GetElementAtVertex(TriangleIndex, VertexIndex);
+				}
+			}
+
+			for (int32 VertexIndex = 0; VertexIndex < DynamicMesh.VertexCount(); ++VertexIndex)
+			{
+				TangentUs[VertexIndex].Normalize();
+				TangentVs[VertexIndex].Normalize();
+			}
 		}
-		TangentUs.SetNumZeroed(DynamicMesh.VertexCount());
-		TangentVs.SetNumZeroed(DynamicMesh.VertexCount());
 
-		const FDynamicMeshNormalOverlay* const TangentUOverlay = DynamicMesh.Attributes()->PrimaryTangents();
-		const FDynamicMeshNormalOverlay* const TangentVOverlay = DynamicMesh.Attributes()->PrimaryBiTangents();
+		//
+		// UVs
+		//
 
-		for (const int TriangleIndex : DynamicMesh.TriangleIndicesItr())
+		// TODO: Account for multiple UV layers
+		TArray<FVector2f> UVs;
+		if (bHasUVs)
 		{
-			const FIndex3i Tri = DynamicMesh.GetTriangle(TriangleIndex);
+			UVs.SetNum(DynamicMesh.VertexCount());
 
-			for (int TriangleVertexIndex = 0; TriangleVertexIndex < 3; ++TriangleVertexIndex)
+			const FDynamicMeshUVOverlay* const UVOverlay = DynamicMesh.Attributes()->PrimaryUV();
+
+			// Assume no seams in the dynamic mesh UV overlay
+			checkSlow(!UVOverlay->HasInteriorSeamEdges());
+
+			for (const int TriangleIndex : DynamicMesh.TriangleIndicesItr())
 			{
-				const int VertexIndex = Tri[TriangleVertexIndex];
+				const FIndex3i Tri = DynamicMesh.GetTriangle(TriangleIndex);
 
-				TangentUs[VertexIndex] += TangentUOverlay->GetElementAtVertex(TriangleIndex, VertexIndex);
-				TangentVs[VertexIndex] += TangentVOverlay->GetElementAtVertex(TriangleIndex, VertexIndex);
+				for (int TriangleVertexIndex = 0; TriangleVertexIndex < 3; ++TriangleVertexIndex)
+				{
+					const int VertexIndex = Tri[TriangleVertexIndex];
+					UVs[VertexIndex] = UVOverlay->GetElementAtVertex(TriangleIndex, VertexIndex);
+				}
 			}
 		}
 
-		for (int32 VertexIndex = 0; VertexIndex < DynamicMesh.VertexCount(); ++VertexIndex)
+		//
+		// Skin Weights
+		//
+
+		FDynamicMeshAttributeSet* Attributes = DynamicMesh.Attributes();
+		check(Attributes);
+
+		TArray<TArray<int32>> BoneIndices;
+		TArray<TArray<float>> BoneWeights;
+		BoneIndices.SetNum(DynamicMesh.VertexCount());
+		BoneWeights.SetNum(DynamicMesh.VertexCount());
+
+		for (const TPair<FName, TUniquePtr<FDynamicMeshVertexSkinWeightsAttribute>>& SkinWeightLayer : Attributes->GetSkinWeightsAttributes())
 		{
-			TangentUs[VertexIndex].Normalize();
-			TangentVs[VertexIndex].Normalize();
+			const FDynamicMeshVertexSkinWeightsAttribute* const SkinWeightAttribute = SkinWeightLayer.Value.Get();
+
+			for (const int TriangleIndex : DynamicMesh.TriangleIndicesItr())
+			{
+				const FIndex3i Tri = DynamicMesh.GetTriangle(TriangleIndex);
+
+				for (int TriangleVertexIndex = 0; TriangleVertexIndex < 3; ++TriangleVertexIndex)
+				{
+					const int VertexIndex = Tri[TriangleVertexIndex];
+					SkinWeightAttribute->GetValue(VertexIndex, BoneIndices[VertexIndex], BoneWeights[VertexIndex]);
+				}
+			}
 		}
-	}
 
-	//
-	// UVs
-	//
 
-	// TODO: Account for multiple UV layers
-	TArray<FVector2f> UVs;
-	if (bHasUVs)
-	{
-		UVs.SetNum(DynamicMesh.VertexCount());
+		// Find the set of triangles per MaterialID
 
-		const FDynamicMeshUVOverlay* const UVOverlay = DynamicMesh.Attributes()->PrimaryUV();
+		TMap<int32, TArray<int32>> MaterialTriangles;
 
-		// Assume no seams in the dynamic mesh UV overlay
-		checkSlow(!UVOverlay->HasInteriorSeamEdges());
+		const FDynamicMeshMaterialAttribute* const MaterialAttribute = Attributes->GetMaterialID();
+		check(MaterialAttribute);
 
-		for (const int TriangleIndex : DynamicMesh.TriangleIndicesItr())
+		for (const int32 TriangleID : DynamicMesh.TriangleIndicesItr())
 		{
-			const FIndex3i Tri = DynamicMesh.GetTriangle(TriangleIndex);
-
-			for (int TriangleVertexIndex = 0; TriangleVertexIndex < 3; ++TriangleVertexIndex)
+			const int32 MaterialID = MaterialAttribute->GetValue(TriangleID);
+			if (!MaterialTriangles.Contains(MaterialID))
 			{
-				const int VertexIndex = Tri[TriangleVertexIndex];
-				UVs[VertexIndex] = UVOverlay->GetElementAtVertex(TriangleIndex, VertexIndex);
+				MaterialTriangles.Add(MaterialID);
 			}
+			MaterialTriangles[MaterialID].Add(TriangleID);
 		}
-	}
 
-	//
-	// Skin Weights
-	//
+		TArray<int32> MaterialIDs;
+		MaterialTriangles.GetKeys(MaterialIDs);
 
-	FDynamicMeshAttributeSet* Attributes = DynamicMesh.Attributes();
-	check(Attributes);
+		const int32 NumMaterials = MaterialIDs.Num();
 
-	TArray<TArray<int32>> BoneIndices;
-	TArray<TArray<float>> BoneWeights;
-	BoneIndices.SetNum(DynamicMesh.VertexCount());
-	BoneWeights.SetNum(DynamicMesh.VertexCount());
+		//
+		// Populate output cloth collection
+		// 
 
-	for (const TPair<FName, TUniquePtr<FDynamicMeshVertexSkinWeightsAttribute>>& SkinWeightLayer : Attributes->GetSkinWeightsAttributes())
-	{
-		const FDynamicMeshVertexSkinWeightsAttribute* const SkinWeightAttribute = SkinWeightLayer.Value.Get();
+		FClothGeometryTools::DeleteRenderMesh(OutClothCollection);
+		FCollectionClothFacade OutClothFacade(OutClothCollection);
+		OutClothFacade.SetNumRenderPatterns(NumMaterials);
 
-		for (const int TriangleIndex : DynamicMesh.TriangleIndicesItr())
+		for (int32 DestPatternID = 0; DestPatternID < MaterialIDs.Num(); ++DestPatternID)
 		{
-			const FIndex3i Tri = DynamicMesh.GetTriangle(TriangleIndex);
+			FCollectionClothRenderPatternFacade OutClothPatternFacade = OutClothFacade.GetRenderPattern(DestPatternID);
+			check(OutClothPatternFacade.GetNumRenderFaces() == 0);
+			check(OutClothPatternFacade.GetNumRenderVertices() == 0);
 
-			for (int TriangleVertexIndex = 0; TriangleVertexIndex < 3; ++TriangleVertexIndex)
+			const int32 SourceMaterialID = MaterialIDs[DestPatternID];
+			check(MaterialTriangles.Contains(SourceMaterialID));
+			const TArray<int32>& TriangleIDs = MaterialTriangles[SourceMaterialID];
+
+			TSet<int32> VertexIndices;
+			for (int32 PatternTriangleIndex = 0; PatternTriangleIndex < TriangleIDs.Num(); ++PatternTriangleIndex)
 			{
-				const int VertexIndex = Tri[TriangleVertexIndex];
-				SkinWeightAttribute->GetValue(VertexIndex, BoneIndices[VertexIndex], BoneWeights[VertexIndex]);
+				const int32 TInd = TriangleIDs[PatternTriangleIndex];
+				const FIndex3i Tri = DynamicMesh.GetTriangle(TInd);
+				VertexIndices.Add(Tri[0]);
+				VertexIndices.Add(Tri[1]);
+				VertexIndices.Add(Tri[2]);
 			}
-		}
-	}
+			const TArray<int32> SourceVertexIndicesArray = VertexIndices.Array();
+			const int32 NumVerticesThisPattern = SourceVertexIndicesArray.Num();
+
+			OutClothPatternFacade.SetNumRenderVertices(NumVerticesThisPattern);
+			TArrayView<FVector3f> RenderPosition = OutClothPatternFacade.GetRenderPosition();
+			TArrayView<FVector3f> RenderNormal = OutClothPatternFacade.GetRenderNormal();
+			TArrayView<FVector3f> RenderTangentU = OutClothPatternFacade.GetRenderTangentU();
+			TArrayView<FVector3f> RenderTangentV = OutClothPatternFacade.GetRenderTangentV();
+			TArrayView<TArray<FVector2f>> RenderUVs = OutClothPatternFacade.GetRenderUVs();
+			TArrayView<FLinearColor> RenderColor = OutClothPatternFacade.GetRenderColor();
+			TArrayView<TArray<int32>> RenderBoneIndices = OutClothPatternFacade.GetRenderBoneIndices();
+			TArrayView<TArray<float>> RenderBoneWeights = OutClothPatternFacade.GetRenderBoneWeights();
+
+			TMap<int32, int32> SourceToDestVertexMap;
+			for (int32 PatternVertexIndex = 0; PatternVertexIndex < NumVerticesThisPattern; ++PatternVertexIndex)
+			{
+				const int32 SourceVertexIndex = SourceVertexIndicesArray[PatternVertexIndex];
 
+				SourceToDestVertexMap.Add(SourceVertexIndex, PatternVertexIndex);
 
-	// Find the set of triangles per MaterialID
+				RenderPosition[PatternVertexIndex] = FVector3f(DynamicMesh.GetVertex(SourceVertexIndex));
+				if (bHasUVs)
+				{
+					RenderUVs[PatternVertexIndex].SetNum(MAX_TEXCOORDS);
+					RenderUVs[PatternVertexIndex][0] = UVs[SourceVertexIndex];
+				}
+				if (bHasNormals)
+				{
+					RenderNormal[PatternVertexIndex] = Normals[SourceVertexIndex];
+				}
+				RenderTangentU[PatternVertexIndex] = TangentUs[SourceVertexIndex];
+				RenderTangentV[PatternVertexIndex] = TangentVs[SourceVertexIndex];
+				RenderColor[PatternVertexIndex] = FLinearColor::White;
+				RenderBoneIndices[PatternVertexIndex] = BoneIndices[SourceVertexIndex];
+				RenderBoneWeights[PatternVertexIndex] = BoneWeights[SourceVertexIndex];
+			}
 
-	TMap<int32, TArray<int32>> MaterialTriangles;
+			OutClothPatternFacade.SetNumRenderFaces(TriangleIDs.Num());
+			for (int32 PatternTriangleIndex = 0; PatternTriangleIndex < TriangleIDs.Num(); ++PatternTriangleIndex)
+			{
+				const int32 VertexOffset = OutClothPatternFacade.GetRenderVerticesOffset();
+				TArrayView<FIntVector3> RenderIndices = OutClothPatternFacade.GetRenderIndices();
 
-	const FDynamicMeshMaterialAttribute* const MaterialAttribute = Attributes->GetMaterialID();
-	check(MaterialAttribute);
+				const int32 TInd = TriangleIDs[PatternTriangleIndex];
+				const FIndex3i SourceTri = DynamicMesh.GetTriangle(TInd);
 
-	for (const int32 TriangleID : DynamicMesh.TriangleIndicesItr())
-	{
-		const int32 MaterialID = MaterialAttribute->GetValue(TriangleID);
-		if (!MaterialTriangles.Contains(MaterialID))
-		{
-			MaterialTriangles.Add(MaterialID);
+				RenderIndices[PatternTriangleIndex][0] = VertexOffset + SourceToDestVertexMap[SourceTri[0]];
+				RenderIndices[PatternTriangleIndex][1] = VertexOffset + SourceToDestVertexMap[SourceTri[1]];
+				RenderIndices[PatternTriangleIndex][2] = VertexOffset + SourceToDestVertexMap[SourceTri[2]];
+			}
+
+			FCollectionClothConstFacade InClothFacade(ClothCollection);
+			FCollectionClothRenderPatternConstFacade InPatternFacade = InClothFacade.GetRenderPattern(SourceMaterialID);
+			OutClothPatternFacade.SetRenderMaterialPathName(InPatternFacade.GetRenderMaterialPathName());
 		}
-		MaterialTriangles[MaterialID].Add(TriangleID);
 	}
 
-	TArray<int32> MaterialIDs;
-	MaterialTriangles.GetKeys(MaterialIDs);
 
-	const int32 NumMaterials = MaterialIDs.Num();
+}
+
+///////////////////////////////////////////////////////////////////////////////////////////////////////////////////////////////
 
-	//
-	// Populate output cloth collection
-	// 
 
-	FClothGeometryTools::DeleteRenderMesh(OutClothCollection);
-	FCollectionClothFacade OutClothFacade(OutClothCollection);
-	OutClothFacade.SetNumRenderPatterns(NumMaterials);
+FChaosClothAssetRemeshNode_v2::FChaosClothAssetRemeshNode_v2(const UE::Dataflow::FNodeParameters& InParam, FGuid InGuid)
+	: FDataflowNode(InParam, InGuid)
+{
+	RegisterInputConnection(&Collection);
+	RegisterOutputConnection(&Collection, &Collection);
+	RegisterInputConnection(&DensityMapSim.WeightMap);
+	RegisterInputConnection(&DensityMapRender.WeightMap);
+}
 
-	for (int32 DestPatternID = 0; DestPatternID < MaterialIDs.Num(); ++DestPatternID)
+
+void FChaosClothAssetRemeshNode_v2::Evaluate(UE::Dataflow::FContext& Context, const FDataflowOutput* Out) const
+{
+	if (Out->IsA<FManagedArrayCollection>(&Collection))
 	{
-		FCollectionClothRenderPatternFacade OutClothPatternFacade = OutClothFacade.GetRenderPattern(DestPatternID);
-		check(OutClothPatternFacade.GetNumRenderFaces() == 0);
-		check(OutClothPatternFacade.GetNumRenderVertices() == 0);
-
-		const int32 SourceMaterialID = MaterialIDs[DestPatternID];
-		check(MaterialTriangles.Contains(SourceMaterialID));
-		const TArray<int32>& TriangleIDs = MaterialTriangles[SourceMaterialID];
-
-		TSet<int32> VertexIndices;
-		for (int32 PatternTriangleIndex = 0; PatternTriangleIndex < TriangleIDs.Num(); ++PatternTriangleIndex)
-		{
-			const int32 TInd = TriangleIDs[PatternTriangleIndex];
-			const FIndex3i Tri = DynamicMesh.GetTriangle(TInd);
-			VertexIndices.Add(Tri[0]);
-			VertexIndices.Add(Tri[1]);
-			VertexIndices.Add(Tri[2]);
-		}
-		const TArray<int32> SourceVertexIndicesArray = VertexIndices.Array();
-		const int32 NumVerticesThisPattern = SourceVertexIndicesArray.Num();
-
-		OutClothPatternFacade.SetNumRenderVertices(NumVerticesThisPattern);
-		TArrayView<FVector3f> RenderPosition = OutClothPatternFacade.GetRenderPosition();
-		TArrayView<FVector3f> RenderNormal = OutClothPatternFacade.GetRenderNormal();
-		TArrayView<FVector3f> RenderTangentU = OutClothPatternFacade.GetRenderTangentU();
-		TArrayView<FVector3f> RenderTangentV = OutClothPatternFacade.GetRenderTangentV();
-		TArrayView<TArray<FVector2f>> RenderUVs = OutClothPatternFacade.GetRenderUVs();
-		TArrayView<FLinearColor> RenderColor = OutClothPatternFacade.GetRenderColor();
-		TArrayView<TArray<int32>> RenderBoneIndices = OutClothPatternFacade.GetRenderBoneIndices();
-		TArrayView<TArray<float>> RenderBoneWeights = OutClothPatternFacade.GetRenderBoneWeights();
-		
-		TMap<int32, int32> SourceToDestVertexMap;
-		for (int32 PatternVertexIndex = 0; PatternVertexIndex < NumVerticesThisPattern; ++PatternVertexIndex)
+		using namespace UE::Chaos::ClothAsset;
+
+		// Evaluate in collection
+		FManagedArrayCollection InCollection = GetValue<FManagedArrayCollection>(Context, &Collection);
+		const TSharedRef<FManagedArrayCollection> ClothCollection = MakeShared<FManagedArrayCollection>(MoveTemp(InCollection));
+		FCollectionClothFacade ClothFacade(ClothCollection);
+
+		const FString SimDensityMapName = GetValue<FString>(Context, &DensityMapSim.WeightMap);
+		const FString RenderDensityMapName = GetValue<FString>(Context, &DensityMapRender.WeightMap);
+
+		// Copy collection to output
+		TSharedRef<FManagedArrayCollection> OutputClothCollection = MakeShared<FManagedArrayCollection>();
+
+		if ((bRemeshSim || bRemeshRender) && ClothFacade.IsValid(EClothCollectionOptionalSchemas::RenderDeformer))
 		{
-			const int32 SourceVertexIndex = SourceVertexIndicesArray[PatternVertexIndex];
+			FClothDataflowTools::LogAndToastWarning(*this,
+				LOCTEXT("InputHasDeformerDataHeadline", "Proxy Deformer Data Found"),
+				LOCTEXT("InputHasDeformerDataDetails", "The input Cloth Collection has Proxy Deformer data that will be removed by the Remesh node. Default deformer bindings will be computed in the final asset. Consider placing ProxyDeformer Node after the Remesh Node."));
 
-			SourceToDestVertexMap.Add(SourceVertexIndex, PatternVertexIndex);
+			// Don't copy proxy deformer data
+			const TArray<FName> GroupsToSkip = TArray<FName>();
+			TArray<TTuple<FName, FName>> AttributesToSkip;
+			AttributesToSkip.Add({ ClothCollectionAttribute::RenderDeformerNumInfluences, ClothCollectionGroup::RenderPatterns });
+			AttributesToSkip.Add({ ClothCollectionAttribute::RenderDeformerPositionBaryCoordsAndDist, ClothCollectionGroup::RenderVertices });
+			AttributesToSkip.Add({ ClothCollectionAttribute::RenderDeformerNormalBaryCoordsAndDist, ClothCollectionGroup::RenderVertices });
+			AttributesToSkip.Add({ ClothCollectionAttribute::RenderDeformerTangentBaryCoordsAndDist, ClothCollectionGroup::RenderVertices });
+			AttributesToSkip.Add({ ClothCollectionAttribute::RenderDeformerSimIndices3D, ClothCollectionGroup::RenderVertices });
+			AttributesToSkip.Add({ ClothCollectionAttribute::RenderDeformerWeight, ClothCollectionGroup::RenderVertices });
+			AttributesToSkip.Add({ ClothCollectionAttribute::RenderDeformerSkinningBlend, ClothCollectionGroup::RenderVertices });
 
-			RenderPosition[PatternVertexIndex] = FVector3f(DynamicMesh.GetVertex(SourceVertexIndex));
-			if (bHasUVs)
-			{
-				RenderUVs[PatternVertexIndex].SetNum(MAX_TEXCOORDS);
-				RenderUVs[PatternVertexIndex][0] = UVs[SourceVertexIndex];
-			}
-			if (bHasNormals)
-			{
-				RenderNormal[PatternVertexIndex] = Normals[SourceVertexIndex];
-			}
-			RenderTangentU[PatternVertexIndex] = TangentUs[SourceVertexIndex];
-			RenderTangentV[PatternVertexIndex] = TangentVs[SourceVertexIndex];
-			RenderColor[PatternVertexIndex] = FLinearColor::White;
-			RenderBoneIndices[PatternVertexIndex] = BoneIndices[SourceVertexIndex];
-			RenderBoneWeights[PatternVertexIndex] = BoneWeights[SourceVertexIndex];
+			ClothCollection->CopyTo(&OutputClothCollection.Get(), GroupsToSkip, AttributesToSkip);
+		}
+		else
+		{
+			ClothCollection->CopyTo(&OutputClothCollection.Get());
 		}
 
-		OutClothPatternFacade.SetNumRenderFaces(TriangleIDs.Num());
-		for (int32 PatternTriangleIndex = 0; PatternTriangleIndex < TriangleIDs.Num(); ++PatternTriangleIndex)
+		if (ClothFacade.IsValid())  // Can only act on the collection if it is a valid cloth collection
 		{
-			const int32 VertexOffset = OutClothPatternFacade.GetRenderVerticesOffset();
-			TArrayView<FIntVector3> RenderIndices = OutClothPatternFacade.GetRenderIndices();
+			if (bRemeshSim)
+			{
+				const float TotalArea2D = Algo::Accumulate(ClothFacade.GetSimIndices2D(), 0.0f, [&ClothFacade](float PreviousSum, const FIntVector3& Tri)
+					{
+						const TArrayView<const FVector2f>& Vertices2D = ClothFacade.GetSimPosition2D();
+						const float TriArea = FVector2f::CrossProduct(Vertices2D[Tri[1]] - Vertices2D[Tri[0]], Vertices2D[Tri[2]] - Vertices2D[Tri[0]]);
+						return PreviousSum + FMath::Abs(TriArea);
+					});
 
-			const int32 TInd = TriangleIDs[PatternTriangleIndex];
-			const FIndex3i SourceTri = DynamicMesh.GetTriangle(TInd);
+				if (TotalArea2D < UE_SMALL_NUMBER)
+				{
+					FClothDataflowTools::LogAndToastWarning(*this,
+						LOCTEXT("SimMeshAreaZeroHeadline", "2D Sim Mesh Has Zero Area"),
+						LOCTEXT("SimMeshAreaZeroDetails", "The input Cloth Collection has a 2D Sim mesh with zero area. Sim remeshing will not be performed"));
+				}
+				else
+				{
+					Private::EmptySimSelections(OutputClothCollection);
+					Private::RemeshSimMesh(ClothCollection, IterationsSim, SmoothingSim, SimDensityMapName, FVector2f(DensityMapSim.Low / 100.0f, DensityMapSim.High / 100.0f), OutputClothCollection);
+					Private::RebuildTopologyDependentSimData(ClothCollection, OutputClothCollection);
+				}
+			}
 
-			RenderIndices[PatternTriangleIndex][0] = VertexOffset + SourceToDestVertexMap[SourceTri[0]];
-			RenderIndices[PatternTriangleIndex][1] = VertexOffset + SourceToDestVertexMap[SourceTri[1]];
-			RenderIndices[PatternTriangleIndex][2] = VertexOffset + SourceToDestVertexMap[SourceTri[2]];
+			if (bRemeshRender)
+			{
+				Private::EmptyRenderSelections(OutputClothCollection);
+				Private::RemeshRenderMesh(ClothCollection,
+					RemeshMethodRender,
+					bRemeshRenderSeams,
+					RenderSeamRemeshIterations,
+					IterationsRender,
+					SmoothingRender,
+					RenderDensityMapName, 
+					FVector2f(DensityMapRender.Low / 100.0f, DensityMapRender.High / 100.0f), 
+					TargetPercentRender,
+					OutputClothCollection);
+			}
 		}
 
-		FCollectionClothConstFacade InClothFacade(ClothCollection);
-		FCollectionClothRenderPatternConstFacade InPatternFacade = InClothFacade.GetRenderPattern(SourceMaterialID);
-		OutClothPatternFacade.SetRenderMaterialPathName(InPatternFacade.GetRenderMaterialPathName());
+		SetValue(Context, MoveTemp(*OutputClothCollection), &Collection);
 	}
+
 }
 
 
+
+
+///////////////////////////////////////////////////////////////////////////////////////////////////////////////////////////////
+
+
+FChaosClothAssetRemeshNode::FChaosClothAssetRemeshNode(const UE::Dataflow::FNodeParameters& InParam, FGuid InGuid)
+	: FDataflowNode(InParam, InGuid)
+{
+	RegisterInputConnection(&Collection);
+	RegisterOutputConnection(&Collection, &Collection);
+	RegisterInputConnection(&DensityMapSim.StringValue, GET_MEMBER_NAME_CHECKED(FChaosClothAssetConnectableIStringValue, StringValue));
+	RegisterInputConnection(&DensityMapRender.StringValue, GET_MEMBER_NAME_CHECKED(FChaosClothAssetConnectableIStringValue, StringValue));
+}
+
 void FChaosClothAssetRemeshNode::Evaluate(UE::Dataflow::FContext& Context, const FDataflowOutput* Out) const
 {
 	if (Out->IsA<FManagedArrayCollection>(&Collection))
@@ -1720,8 +1872,8 @@ void FChaosClothAssetRemeshNode::Evaluate(UE::Dataflow::FContext& Context, const
 		const TSharedRef<FManagedArrayCollection> ClothCollection = MakeShared<FManagedArrayCollection>(MoveTemp(InCollection));
 		FCollectionClothFacade ClothFacade(ClothCollection);
 
-		const FString SimDensityMapName = GetValue<FString>(Context, &DensityMapSim.StringValue);
-		const FString RenderDensityMapName = GetValue<FString>(Context, &DensityMapRender.StringValue);
+		const FString SimDensityMapName = GetValue<FString>(Context, &DensityMapSim.StringValue, FString());
+		const FString RenderDensityMapName = GetValue<FString>(Context, &DensityMapRender.StringValue, FString());
 
 		// Copy collection to output
 		TSharedPtr<FManagedArrayCollection> OutputClothCollection = MakeShared<FManagedArrayCollection>();
@@ -1771,16 +1923,25 @@ void FChaosClothAssetRemeshNode::Evaluate(UE::Dataflow::FContext& Context, const
 				}
 				else
 				{
-					EmptySimSelections(OutputClothCollectionRef);
-					RemeshSimMesh(ClothCollection, SimDensityMapName, OutputClothCollectionRef);
-					RebuildTopologyDependentSimData(ClothCollection, OutputClothCollectionRef);
+					Private::EmptySimSelections(OutputClothCollectionRef);
+					Private::RemeshSimMesh(ClothCollection, IterationsSim, SmoothingSim, SimDensityMapName, FVector2f(TargetPercentSim / 100.0f, TargetPercentSim / 50.0f), OutputClothCollectionRef);
+					Private::RebuildTopologyDependentSimData(ClothCollection, OutputClothCollectionRef);
 				}
 			}
 
 			if (bRemeshRender)
 			{
-				EmptyRenderSelections(OutputClothCollectionRef);
-				RemeshRenderMesh(ClothCollection, RenderDensityMapName, OutputClothCollectionRef);
+				Private::EmptyRenderSelections(OutputClothCollectionRef);
+				Private::RemeshRenderMesh(ClothCollection,
+					RemeshMethodRender,
+					bRemeshRenderSeams,
+					RenderSeamRemeshIterations,
+					IterationsRender,
+					SmoothingRender,
+					RenderDensityMapName,
+					FVector2f(TargetPercentRender / 100.0f, TargetPercentRender / 100.0f),
+					TargetPercentRender,
+					OutputClothCollectionRef);
 			}
 		}
 
diff --git a/Engine/Plugins/ChaosClothAssetEditor/Source/ChaosClothAssetDataflowNodes/Public/ChaosClothAsset/RemeshNode.h b/Engine/Plugins/ChaosClothAssetEditor/Source/ChaosClothAssetDataflowNodes/Public/ChaosClothAsset/RemeshNode.h
--- a/Engine/Plugins/ChaosClothAssetEditor/Source/ChaosClothAssetDataflowNodes/Public/ChaosClothAsset/RemeshNode.h
+++ b/Engine/Plugins/ChaosClothAssetEditor/Source/ChaosClothAssetDataflowNodes/Public/ChaosClothAsset/RemeshNode.h
@@ -5,6 +5,7 @@
 #include "Dataflow/DataflowNode.h"
 #include "GeometryCollection/ManagedArrayCollection.h"
 #include "ChaosClothAsset/ConnectableValue.h"
+#include "ChaosClothAsset/WeightedValue.h"
 #include "RemeshNode.generated.h"
 
 #if UE_ENABLE_INCLUDE_ORDER_DEPRECATED_IN_5_5
@@ -20,12 +21,85 @@ enum class EChaosClothAssetRemeshMethod : uint8
 	Simplify
 };
 
-
 /** Remesh the cloth surface(s) to get the specified mesh resolution(s).
  *  NOTE: Weight Maps, Skinning Data, Self Collision Spheres, and Long Range Attachment Constraints will be reconstructed on the output mesh, however all other Selections will be removed
  */
 USTRUCT(Meta = (DataflowCloth))
-struct FChaosClothAssetRemeshNode : public FDataflowNode
+struct FChaosClothAssetRemeshNode_v2 : public FDataflowNode
+{
+	GENERATED_USTRUCT_BODY()
+	DATAFLOW_NODE_DEFINE_INTERNAL(FChaosClothAssetRemeshNode_v2, "Remesh", "Cloth", "Cloth Remesh")
+	DATAFLOW_NODE_RENDER_TYPE("SurfaceRender", FName("FClothCollection"), "Collection")
+
+public:
+
+	FChaosClothAssetRemeshNode_v2(const UE::Dataflow::FNodeParameters& InParam, FGuid InGuid = FGuid::NewGuid());
+
+private:
+
+	UPROPERTY(Meta = (Dataflowinput, DataflowOutput, DataflowPassthrough = "Collection"))
+	FManagedArrayCollection Collection;
+
+	UPROPERTY(EditAnywhere, Category = "Sim Mesh")
+	bool bRemeshSim = true;
+
+	/**
+	 * Range of target mesh resolutions, as a percentage of input triangle mesh resolution. A value of 50 on all vertices should roughly halve the total number of triangles.
+	 * If a valid vertex weight map is specified, it will use vertex weights to interpolate between the Lo and Hi values. Otherwise it will use the Lo value on all vertices.
+	 */
+	UPROPERTY(EditAnywhere, Category = "Sim Mesh", Meta = (EditCondition = "bRemeshSim"))
+	FChaosClothAssetWeightedValueNonAnimatable DensityMapSim = { 100.f, 200.f, TEXT("DensityMapSim") };
+
+	UPROPERTY(EditAnywhere, Category = "Sim Mesh", meta = (UIMin = "0", UIMax = "20", ClampMin = "0", ClampMax = "200", EditCondition = "bRemeshSim"))
+	int32 IterationsSim = 10;
+
+	UPROPERTY(EditAnywhere, Category = "Sim Mesh", meta = (UIMin = "0.0", UIMax = "1.0", ClampMin = "0.0", ClampMax = "1.0", EditCondition = "bRemeshSim"))
+	double SmoothingSim = 0.25;
+
+	UPROPERTY(EditAnywhere, Category = "Render Mesh")
+	bool bRemeshRender = false;
+
+	UPROPERTY(EditAnywhere, Category = "Render Mesh", meta = (UIMin = "0.0", UIMax = "1.0", ClampMin = "0.0", ClampMax = "1.0", EditCondition = "bRemeshRender"))
+	EChaosClothAssetRemeshMethod RemeshMethodRender = EChaosClothAssetRemeshMethod::Remesh;
+
+	/**
+	 * Range of target mesh resolutions when using the Remesh method, as a percentage of input triangle mesh resolution. A value of 50 on all vertices should roughly halve the total number of triangles.
+	 * If a valid vertex weight map is specified, it will use vertex weights to interpolate between the Lo and Hi values. Otherwise it will use the Lo value on all vertices.
+	 */
+	UPROPERTY(EditAnywhere, Category = "Render Mesh", Meta = (EditCondition = "bRemeshRender && RemeshMethodRender == EChaosClothAssetRemeshMethod::Remesh"))
+	FChaosClothAssetWeightedValueNonAnimatable DensityMapRender = { 100.f, 200.f, TEXT("DensityMapRender") };
+
+	/**
+	 * Target mesh resolution when using the Simplify method, as a percentage of input triangle mesh resolution. A value of 50 should roughly halve the total number of triangles.
+	 */
+	UPROPERTY(EditAnywhere, Category = "Render Mesh", meta = (UIMin = "1", UIMax = "200", ClampMin = "1", EditCondition = "bRemeshRender && RemeshMethodRender == EChaosClothAssetRemeshMethod::Simplify"))
+	int32 TargetPercentRender = 100;
+
+	UPROPERTY(EditAnywhere, Category = "Render Mesh", meta = (UIMin = "0", UIMax = "20", ClampMin = "0", ClampMax = "100", EditCondition = "bRemeshRender && RemeshMethodRender == EChaosClothAssetRemeshMethod::Remesh"))
+	int32 IterationsRender = 10;
+
+	UPROPERTY(EditAnywhere, Category = "Render Mesh", meta = (UIMin = "0.0", UIMax = "1.0", ClampMin = "0.0", ClampMax = "1.0", EditCondition = "bRemeshRender && RemeshMethodRender == EChaosClothAssetRemeshMethod::Remesh"))
+	double SmoothingRender = 0.25;
+
+	/** If checked, attempt to find matching vertices along Render mesh boundaries and remesh these separately */
+	UPROPERTY(EditAnywhere, Category = "Render Mesh", meta = (EditCondition = "bRemeshRender"))
+	bool bRemeshRenderSeams = false;
+
+	/** Number of remesh iterations over the Render mesh seams */
+	UPROPERTY(EditAnywhere, Category = "Render Mesh", meta = (UIMin = "0", UIMax = "20", ClampMin = "0", ClampMax = "100", EditCondition = "bRemeshRender && bRemeshRenderSeams"))
+	int32 RenderSeamRemeshIterations = 1;
+
+	virtual void Evaluate(UE::Dataflow::FContext & Context, const FDataflowOutput * Out) const override;
+
+};
+
+
+
+/** Remesh the cloth surface(s) to get the specified mesh resolution(s).
+ *  NOTE: Weight Maps, Skinning Data, Self Collision Spheres, and Long Range Attachment Constraints will be reconstructed on the output mesh, however all other Selections will be removed
+ */
+USTRUCT(Meta = (DataflowCloth, Deprecated = "5.6"))
+struct UE_DEPRECATED(5.6, "Use the newer version of this node instead.") FChaosClothAssetRemeshNode : public FDataflowNode
 {
 	GENERATED_USTRUCT_BODY()
 	DATAFLOW_NODE_DEFINE_INTERNAL(FChaosClothAssetRemeshNode, "Remesh", "Cloth", "Cloth Remesh")
@@ -80,21 +154,4 @@ public:
 	FChaosClothAssetRemeshNode(const UE::Dataflow::FNodeParameters& InParam, FGuid InGuid = FGuid::NewGuid());
 
 	virtual void Evaluate(UE::Dataflow::FContext& Context, const FDataflowOutput* Out) const override;
-
-private:
-
-	void EmptySimSelections(const TSharedRef<FManagedArrayCollection>& ClothCollection) const;
-
-	void RemeshSimMesh(const TSharedRef<const FManagedArrayCollection>& ClothCollection,
-		const FString& DensityMapName,
-		const TSharedRef<FManagedArrayCollection>& OutClothCollection) const;
-
-	void RebuildTopologyDependentSimData(const TSharedRef<const FManagedArrayCollection>& ClothCollection,
-		const TSharedRef<FManagedArrayCollection>& OutClothCollection) const;
-
-	void EmptyRenderSelections(const TSharedRef<FManagedArrayCollection>& ClothCollection) const;
-
-	void RemeshRenderMesh(const TSharedRef<const FManagedArrayCollection>& ClothCollection,
-		const FString& DensityMapName,
-		const TSharedRef<FManagedArrayCollection>& OutClothCollection) const;
 };
