FChaosClothAssetSimulationClothVertexFaceSpringConfigNode: make zero-volume tet culling optional in the Tetrahedralize approach

diff --git a/Engine/Plugins/ChaosClothAssetEditor/Source/ChaosClothAssetDataflowNodes/Private/ChaosClothAsset/SimulationClothVertexFaceSpringConfigNode.cpp b/Engine/Plugins/ChaosClothAssetEditor/Source/ChaosClothAssetDataflowNodes/Private/ChaosClothAsset/SimulationClothVertexFaceSpringConfigNode.cpp
--- a/Engine/Plugins/ChaosClothAssetEditor/Source/ChaosClothAssetDataflowNodes/Private/ChaosClothAsset/SimulationClothVertexFaceSpringConfigNode.cpp
+++ b/Engine/Plugins/ChaosClothAssetEditor/Source/ChaosClothAssetDataflowNodes/Private/ChaosClothAsset/SimulationClothVertexFaceSpringConfigNode.cpp
@@ -105,7 +105,7 @@ namespace UE::Chaos::ClothAsset::Private
 		}		
 	};
 
-	uint32 GetTypeHash(const FVertexFaceConstraint& Constraint)
+	static uint32 GetTypeHash(const FVertexFaceConstraint& Constraint)
 	{
 		uint32 Hash = ::GetTypeHash(Constraint.SourceVertex);
 		Hash = HashCombineFast(Hash, ::GetTypeHash(Constraint.TargetFace));
@@ -113,7 +113,7 @@ namespace UE::Chaos::ClothAsset::Private
 		return Hash;
 	}
 
-	void AppendConstraintsSourceToClosestTarget(const TSet<int32>& SourceVertices, const TSet<int32>& TargetFaces, const FCollectionClothConstFacade& Cloth, TSet<FVertexFaceConstraint>& Constraints)
+	static void AppendConstraintsSourceToClosestTarget(const TSet<int32>& SourceVertices, const TSet<int32>& TargetFaces, const FCollectionClothConstFacade& Cloth, TSet<FVertexFaceConstraint>& Constraints)
 	{
 		using namespace UE::Geometry;
 		TConstArrayView<FIntVector3> Elements = Cloth.GetSimIndices3D();
@@ -152,7 +152,7 @@ namespace UE::Chaos::ClothAsset::Private
 		}
 	}
 
-	void AppendConstraintsSourceToRayIntersectionTarget(const TSet<int32>& SourceVertices, const TSet<int32>& TargetFaces, bool bFlipRay, float MaxRadius, const FCollectionClothConstFacade& Cloth, TSet<FVertexFaceConstraint>& Constraints)
+	static void AppendConstraintsSourceToRayIntersectionTarget(const TSet<int32>& SourceVertices, const TSet<int32>& TargetFaces, bool bFlipRay, float MaxRadius, const FCollectionClothConstFacade& Cloth, TSet<FVertexFaceConstraint>& Constraints)
 	{
 		using namespace UE::Geometry;
 		TConstArrayView<FIntVector3> Elements = Cloth.GetSimIndices3D();
@@ -192,7 +192,7 @@ namespace UE::Chaos::ClothAsset::Private
 		}
 	}
 
-	void AppendConstraintsAllWithinRadius(const TSet<int32>& SourceVertices, const TSet<int32>& TargetFaces, float Radius, int32 DisableNeighborDistance, const FCollectionClothConstFacade& Cloth, TSet<FVertexFaceConstraint>& Constraints)
+	static void AppendConstraintsAllWithinRadius(const TSet<int32>& SourceVertices, const TSet<int32>& TargetFaces, float Radius, int32 DisableNeighborDistance, const FCollectionClothConstFacade& Cloth, TSet<FVertexFaceConstraint>& Constraints)
 	{
 		using namespace UE::Geometry;
 
@@ -255,7 +255,13 @@ namespace UE::Chaos::ClothAsset::Private
 	}
 
 
-	void AppendConstraintsTetrahedralize(const TSet<int32>& SourceVertices, const TSet<int32>& TargetFaces, float Radius, int32 DisableNeighborDistance, const FCollectionClothConstFacade& Cloth, TSet<FVertexFaceConstraint>& Constraints)
+	static void AppendConstraintsTetrahedralize(const TSet<int32>& SourceVertices, 
+		const TSet<int32>& TargetFaces, 
+		float Radius, 
+		int32 DisableNeighborDistance, 
+		const FCollectionClothConstFacade& Cloth, 
+		bool bSkipZeroVolumeTets,
+		TSet<FVertexFaceConstraint>& Constraints)
 	{
 		// Speed up vertex-triangle queries by constructing this map up front
 		TMap<int32, TArray<int32>> VertexToTriangleMap;
@@ -275,16 +281,32 @@ namespace UE::Chaos::ClothAsset::Private
 			}
 		}
 
-		// Check if the three input vertices correspond to a single triangle on the Cloth sim mesh
-		auto IsSimTri = [](const FIntVector3& TestIndices, const TMap<int32, TArray<int32>>& VertexToTriangleMap, const FCollectionClothConstFacade& Cloth, int32& OutClothTriIndex)
+		// Check if the three input vertices correspond to a single triangle (with three distinct vertices) on the Cloth sim mesh
+		auto IsValidSimTri = [](const FIntVector3& TestIndices, const TMap<int32, TArray<int32>>& VertexToTriangleMap, const FCollectionClothConstFacade& Cloth, int32& OutClothTriIndex)
 		{
-			check(TestIndices[0] != TestIndices[1] && TestIndices[0] != TestIndices[2] && TestIndices[1] != TestIndices[2]);
+			if (TestIndices[0] == TestIndices[1] || TestIndices[0] == TestIndices[2] || TestIndices[1] == TestIndices[2])
+			{
+				return false;
+			}
 
+			// Check if the test triangle vertex is in the vertex-to-triangle map. 
+			// If not, it means no triangles containing this vertex are in the TargetFaces set, which means our test triangle is not in the TargetFaces set either.
+			if (!VertexToTriangleMap.Find(TestIndices[0]))
+			{
+				return false;
+			}
+
+			// Get triangles to check -- we'll use triangles incident on the first vertex
 			const TArray<int32> CandidateTriangles = VertexToTriangleMap[TestIndices[0]];
+
 			for (const int32 CandidateTriIndex : CandidateTriangles)
 			{
 				const FIntVector3& Tri = Cloth.GetSimIndices3D()[CandidateTriIndex];
-				check(Tri[0] != Tri[1] && Tri[0] != Tri[2] && Tri[1] != Tri[2]);
+
+				if (Tri[0] == Tri[1] || Tri[0] == Tri[2] || Tri[1] == Tri[2])
+				{
+					continue;
+				}
 
 				if ((Tri[0] == TestIndices[0] || Tri[0] == TestIndices[1] || Tri[0] == TestIndices[2]) &&
 					(Tri[1] == TestIndices[0] || Tri[1] == TestIndices[1] || Tri[1] == TestIndices[2]) &&
@@ -302,11 +324,7 @@ namespace UE::Chaos::ClothAsset::Private
 		const FFaceSetMeshAdapter MeshAdapter(TargetFaces, Cloth);
 
 		// Throw all source vertices and target triangle vertices into the tet mesher
-		TSet<int32> AllVertexIndices;
-		for (const int32 SourceVertex : SourceVertices)
-		{
-			AllVertexIndices.Add(SourceVertex);
-		}
+		TSet<int32> AllVertexIndices(SourceVertices);
 		for (const int32 TargetFace : TargetFaces)
 		{
 			if (MeshAdapter.IsTriangle(TargetFace))
@@ -337,7 +355,7 @@ namespace UE::Chaos::ClothAsset::Private
 					// For each tet face, if it corresponds to a triangle in the sim mesh, use it to create a constraint
 					const FIntVector3 TetFace(Tet[TetFaceIndex], Tet[(TetFaceIndex + 1) % 4], Tet[(TetFaceIndex + 2) % 4]);
 					int32 SimMeshTriIndex = -1;
-					if (IsSimTri(TetFace, VertexToTriangleMap, Cloth, SimMeshTriIndex))
+					if (IsValidSimTri(TetFace, VertexToTriangleMap, Cloth, SimMeshTriIndex))
 					{
 						const int32 VertIndex = Tet[(TetFaceIndex + 3) % 4];
 						check(0 <= VertIndex && VertIndex < Cloth.GetNumSimVertices3D());
@@ -353,14 +371,17 @@ namespace UE::Chaos::ClothAsset::Private
 						}
 
 						// Check if the vertex is coplanar to the triangle and skip this contraint if it is
-						const FVector3f VertPos = Cloth.GetSimPosition3D()[VertIndex];
-						const FVector3f A = Cloth.GetSimPosition3D()[SimMeshTri[0]] - VertPos;
-						const FVector3f B = Cloth.GetSimPosition3D()[SimMeshTri[1]] - VertPos;
-						const FVector3f C = Cloth.GetSimPosition3D()[SimMeshTri[2]] - VertPos;
-						const float Vol = FVector3f::DotProduct(A, FVector3f::CrossProduct(B, C));
-						if (FMath::Abs(Vol) < UE_KINDA_SMALL_NUMBER)
+						if (bSkipZeroVolumeTets)
 						{
-							continue;
+							const FVector3f VertPos = Cloth.GetSimPosition3D()[VertIndex];
+							const FVector3f A = Cloth.GetSimPosition3D()[SimMeshTri[0]] - VertPos;
+							const FVector3f B = Cloth.GetSimPosition3D()[SimMeshTri[1]] - VertPos;
+							const FVector3f C = Cloth.GetSimPosition3D()[SimMeshTri[2]] - VertPos;
+							const float Vol = FVector3f::DotProduct(A, FVector3f::CrossProduct(B, C));
+							if (FMath::Abs(Vol) < UE_KINDA_SMALL_NUMBER)
+							{
+								continue;
+							}
 						}
 
 						Constraints.Add({ VertIndex, SimMeshTriIndex, FVertexFaceConstraint::EIntersectionType::Closest });
@@ -372,7 +393,7 @@ namespace UE::Chaos::ClothAsset::Private
 	}
 
 
-	void StoreConstraintsEnds(const TSet<FVertexFaceConstraint>& Constraints, const FCollectionClothConstFacade& Cloth, TArray<int32>& SourceVertices, TArray<FIntVector>& TargetVertices, TArray<FVector3f>& TargetWeights)
+	static void StoreConstraintsEnds(const TSet<FVertexFaceConstraint>& Constraints, const FCollectionClothConstFacade& Cloth, TArray<int32>& SourceVertices, TArray<FIntVector>& TargetVertices, TArray<FVector3f>& TargetWeights)
 	{
 		using namespace UE::Geometry;
 
@@ -445,7 +466,7 @@ namespace UE::Chaos::ClothAsset::Private
 		}
 	}
 
-	void CalculateRestLengthsFromPositions(const FCollectionClothConstFacade& Cloth, const TArray<int32>& SourceVertices, const TArray<FIntVector>& TargetVertices, const TArray<FVector3f>& TargetWeights, const float RestLengthScale, TArray<float>& RestLengths)
+	static void CalculateRestLengthsFromPositions(const FCollectionClothConstFacade& Cloth, const TArray<int32>& SourceVertices, const TArray<FIntVector>& TargetVertices, const TArray<FVector3f>& TargetWeights, const float RestLengthScale, TArray<float>& RestLengths)
 	{
 		TConstArrayView<FVector3f> Positions = Cloth.GetSimPosition3D();
 		const int32 NumConstraints = SourceVertices.Num();
@@ -462,7 +483,7 @@ namespace UE::Chaos::ClothAsset::Private
 		}
 	}
 
-	void CalculateRestLengthsFromThickness(const ::Chaos::Softs::FPBDFlatWeightMapView& ThicknessView, const TArray<int32>& SourceVertices, const TArray<FIntVector>& TargetVertices, const TArray<FVector3f>& TargetWeights, TArray<float>& RestLengths)
+	static void CalculateRestLengthsFromThickness(const ::Chaos::Softs::FPBDFlatWeightMapView& ThicknessView, const TArray<int32>& SourceVertices, const TArray<FIntVector>& TargetVertices, const TArray<FVector3f>& TargetWeights, TArray<float>& RestLengths)
 	{
 		const int32 NumConstraints = SourceVertices.Num();
 		check(TargetVertices.Num() == NumConstraints);
@@ -721,6 +742,7 @@ TArray<FChaosClothAssetSimulationClothVertexFaceSpringConfigNode::FConstructionS
 		ConstructionSetData[Index].MaxRayLength = ConstructionSets[Index].MaxRayLength;
 		ConstructionSetData[Index].Radius = ConstructionSets[Index].Radius;
 		ConstructionSetData[Index].DisableNeighborDistance = ConstructionSets[Index].DisableNeighborDistance;
+		ConstructionSetData[Index].bSkipZeroVolumeTets = ConstructionSets[Index].bSkipZeroVolumeTets;
 	}
 	return ConstructionSetData;
 }
@@ -763,7 +785,7 @@ void FChaosClothAssetSimulationClothVertexFaceSpringConfigNode::CreateConstraint
 					Private::AppendConstraintsAllWithinRadius(SourceSet, TargetSet, Data.Radius, Data.DisableNeighborDistance, ClothFacade, Constraints);
 					break;
 				case EChaosClothAssetClothVertexFaceSpringConstructionMethod::Tetrahedralize:
-					Private::AppendConstraintsTetrahedralize(SourceSet, TargetSet, Data.Radius, Data.DisableNeighborDistance, ClothFacade, Constraints);
+					Private::AppendConstraintsTetrahedralize(SourceSet, TargetSet, Data.Radius, Data.DisableNeighborDistance, ClothFacade, Data.bSkipZeroVolumeTets, Constraints);
 					break;
 				default:
 					checkNoEntry();
diff --git a/Engine/Plugins/ChaosClothAssetEditor/Source/ChaosClothAssetDataflowNodes/Public/ChaosClothAsset/SimulationClothVertexFaceSpringConfigNode.h b/Engine/Plugins/ChaosClothAssetEditor/Source/ChaosClothAssetDataflowNodes/Public/ChaosClothAsset/SimulationClothVertexFaceSpringConfigNode.h
--- a/Engine/Plugins/ChaosClothAssetEditor/Source/ChaosClothAssetDataflowNodes/Public/ChaosClothAsset/SimulationClothVertexFaceSpringConfigNode.h
+++ b/Engine/Plugins/ChaosClothAssetEditor/Source/ChaosClothAssetDataflowNodes/Public/ChaosClothAsset/SimulationClothVertexFaceSpringConfigNode.h
@@ -55,6 +55,11 @@ struct FChaosClothAssetSimulationClothVertexFaceSpringConstructionSet
 	/** Do not consider vertices within this N-ring of connectivity distance.*/
 	UPROPERTY(EditAnywhere, Category = "Cloth Vertex Face Spring Construction", Meta = (EditCondition = "ConstructionMethod == EChaosClothAssetClothVertexFaceSpringConstructionMethod::AllWithinRadius", EditConditionHides, ClampMin = "1", UIMax = "5"))
 	int32 DisableNeighborDistance = 2;
+
+	/** Cull zero-volume tets when doing Tetrahedralize.*/
+	UPROPERTY(EditAnywhere, Category = "Cloth Vertex Face Spring Construction", Meta = (EditCondition = "ConstructionMethod == EChaosClothAssetClothVertexFaceSpringConstructionMethod::Tetrahedralize", EditConditionHides))
+	bool bSkipZeroVolumeTets = false;
+
 };
 
 /** Node for creating vertex-face constraints and setting their simulation properties.*/
@@ -106,6 +111,7 @@ private:
 		float MaxRayLength;
 		float Radius;
 		int32 DisableNeighborDistance = 2;
+		bool bSkipZeroVolumeTets;
 	};
 	TArray<FConstructionSetData> GetConstructionSetData(UE::Dataflow::FContext& Context) const;
 
