Chaos Cloth: allow merging cloth collections if one collection's reference skeleton is a subset of the other.

diff --git a/Engine/Plugins/ChaosClothAssetEditor/Source/ChaosClothAssetDataflowNodes/Private/ChaosClothAsset/MergeClothCollectionsNode.cpp b/Engine/Plugins/ChaosClothAssetEditor/Source/ChaosClothAssetDataflowNodes/Private/ChaosClothAsset/MergeClothCollectionsNode.cpp
--- a/Engine/Plugins/ChaosClothAssetEditor/Source/ChaosClothAssetDataflowNodes/Private/ChaosClothAsset/MergeClothCollectionsNode.cpp
+++ b/Engine/Plugins/ChaosClothAssetEditor/Source/ChaosClothAssetDataflowNodes/Private/ChaosClothAsset/MergeClothCollectionsNode.cpp
@@ -413,8 +413,35 @@ namespace UE::Chaos::ClothAsset::Private
 		}
 	}
 
-	static bool AreSkeletalMeshesCompatible(const FDataflowNode& DataflowNode, const FCollectionClothConstFacade& Cloth1, const FCollectionClothConstFacade& Cloth2)
+	static void RemapBoneIndices(TArrayView<TArray<int32>> BoneIndices, const TArray<int32>& Remap)
 	{
+		for (TArray<int32>& Array : BoneIndices)
+		{
+			for (int32& Index : Array)
+			{
+				if (Index == INDEX_NONE)
+				{
+					continue;
+				}
+				if (ensure(Remap.IsValidIndex(Index)))
+				{
+					Index = Remap[Index];
+				}
+			}
+		}
+
+	}
+
+	static void RemapBones(FCollectionClothFacade& Cloth, const TArray<int32>& Remap, int32 SimVertex3DOffset, int32 RenderVertexOffset)
+	{
+		RemapBoneIndices(Cloth.GetSimBoneIndices().RightChop(SimVertex3DOffset), Remap);
+		RemapBoneIndices(Cloth.GetRenderBoneIndices().RightChop(RenderVertexOffset), Remap);
+	}
+
+	static bool AreSkeletalMeshesCompatible(const FDataflowNode& DataflowNode, FCollectionClothFacade& Cloth1, const FCollectionClothConstFacade& Cloth2, TArray<int32>& OtherBoneRemap)
+	{
+		OtherBoneRemap.Reset();
+
 		/** Disallow merging cloth facades with incompatible ref skeletons. */
 		const FString& SkeletalMeshPathName1 = Cloth1.GetSkeletalMeshPathName();
 		const FString& SkeletalMeshPathName2 = Cloth2.GetSkeletalMeshPathName();
@@ -440,54 +467,67 @@ namespace UE::Chaos::ClothAsset::Private
 
 		const FReferenceSkeleton& RefSkeleton1 = SkeletalMesh1->GetRefSkeleton();
 		const FReferenceSkeleton& RefSkeleton2 = SkeletalMesh2->GetRefSkeleton();
-		if (RefSkeleton1.GetNum() != RefSkeleton2.GetNum())
-		{
-			const FText Details = FText::Format(
-				LOCTEXT(
-					"IncompatibleSkeletalMeshesNumBonesDetails",
-					"Cloth collections failed to merge due to incompatible Skeletal Meshes, \"{0}\" and \"{1}\". RefSkeleton Bone counts {2} != {3}."),
-				FText::FromString(SkeletalMeshPathName1),
-				FText::FromString(SkeletalMeshPathName2),
-				RefSkeleton1.GetNum(),
-				RefSkeleton2.GetNum());
-
-			FClothDataflowTools::LogAndToastWarning(DataflowNode, ErrorHeadline, Details);
-			return false;
-		}
 
-		const TArray<FMeshBoneInfo>& RefBoneInfo1 = RefSkeleton1.GetRefBoneInfo();
-		const TArray<FMeshBoneInfo>& RefBoneInfo2 = RefSkeleton2.GetRefBoneInfo();
-		const TArray<FTransform>& RefBonePose1 = RefSkeleton1.GetRefBonePose();
-		const TArray<FTransform>& RefBonePose2 = RefSkeleton2.GetRefBonePose();
-		for (int32 BoneIndex = 0; BoneIndex < RefSkeleton1.GetNum(); ++BoneIndex)
+		const FReferenceSkeleton& MergedRefSkeleton = RefSkeleton1.GetNum() >= RefSkeleton2.GetNum() ? RefSkeleton1 : RefSkeleton2;
+		const FReferenceSkeleton& RemapRefSkeleton = RefSkeleton1.GetNum() >= RefSkeleton2.GetNum() ? RefSkeleton2 : RefSkeleton1;
+		const FString& MergedSkeletalMeshPath = RefSkeleton1.GetNum() >= RefSkeleton2.GetNum() ? SkeletalMeshPathName1 : SkeletalMeshPathName2;
+
+		const TArray<FMeshBoneInfo>& RemapBoneInfo = RemapRefSkeleton.GetRefBoneInfo();
+		const TArray<FTransform>& MergedBonePose = MergedRefSkeleton.GetRefBonePose();
+		const TArray<FTransform>& RemapBonePose = RemapRefSkeleton.GetRefBonePose();
+		TArray<int32> RemapIndices;
+		RemapIndices.SetNumUninitialized(RemapRefSkeleton.GetNum());
+		bool bAnyRemap = false;
+		for (int32 BoneIndex = 0; BoneIndex < RemapRefSkeleton.GetNum(); ++BoneIndex)
 		{
-			if (!(RefBoneInfo1[BoneIndex] == RefBoneInfo2[BoneIndex]))
+			const int32 MergedBoneIndex = MergedRefSkeleton.FindBoneIndex(RemapBoneInfo[BoneIndex].Name);
+			if(MergedBoneIndex == INDEX_NONE)
 			{
 				const FText Details = FText::Format(
 					LOCTEXT(
 						"IncompatibleSkeletalMeshesRefBoneInfoDetails",
-						"Cloth collections failed to merge due to incompatible Skeletal Meshes, \"{0}\" and \"{1}\". RefBoneInfos are mismatched at BoneIndex {2}."),
+						"Cloth collections failed to merge due to incompatible Skeletal Meshes, \"{0}\" and \"{1}\". Could not find bone \"{2}\" in \"{3}\"."),
 					FText::FromString(SkeletalMeshPathName1),
 					FText::FromString(SkeletalMeshPathName2),
-					BoneIndex);
+					FText::FromName(RemapBoneInfo[BoneIndex].Name),
+					FText::FromString(MergedSkeletalMeshPath));
 
 				FClothDataflowTools::LogAndToastWarning(DataflowNode, ErrorHeadline, Details);
 				return false;
 			}
-			if (!RefBonePose1[BoneIndex].Equals(RefBonePose2[BoneIndex]))
+			if (!RemapBonePose[BoneIndex].Equals(MergedBonePose[MergedBoneIndex]))
 			{
 				const FText Details = FText::Format(
 					LOCTEXT(
 						"IncompatibleSkeletalMeshesRefBonePoseDetails",
-						"Cloth collections failed to merge due to incompatible Skeletal Meshes, \"{0}\" and \"{1}\". RefBonePoses are mismatched at BoneIndex {2}."),
+						"Cloth collections failed to merge due to incompatible Skeletal Meshes, \"{0}\" and \"{1}\". RefBonePoses are mismatched for bone \"{2}\"."),
 					FText::FromString(SkeletalMeshPathName1),
 					FText::FromString(SkeletalMeshPathName2),
-					BoneIndex);
+					FText::FromName(RemapBoneInfo[BoneIndex].Name));
 
 				FClothDataflowTools::LogAndToastWarning(DataflowNode, ErrorHeadline, Details);
 				return false;
 			}
+			RemapIndices[BoneIndex] = MergedBoneIndex;
+			if (BoneIndex != MergedBoneIndex)
+			{
+				bAnyRemap = true;
+			}
+		}
+
+		if (bAnyRemap)
+		{
+			if (&RemapRefSkeleton == &RefSkeleton1)
+			{
+				// Remap here since Cloth1 is writable.
+				RemapBones(Cloth1, RemapIndices, 0, 0);
+			}
+			else
+			{
+				OtherBoneRemap = MoveTemp(RemapIndices);
+			}
 		}
+		Cloth1.SetSkeletalMeshPathName(MergedSkeletalMeshPath);
 
 		return true;
 	}
@@ -610,10 +650,29 @@ void FChaosClothAssetMergeClothCollectionsNode_v2::Evaluate(UE::Dataflow::FConte
 
 			// Append cloth
 			const FCollectionClothConstFacade OtherClothFacade(OtherClothCollection);
-			if (OtherClothFacade.IsValid() && Private::AreSkeletalMeshesCompatible(*this, ClothFacade, OtherClothFacade))
+			if (OtherClothFacade.IsValid())
 			{
-				ClothFacade.Append(OtherClothFacade);
-				bAreAnyValid = true;
+
+				TArray<int32> OtherBoneRemap;
+				if (Private::AreSkeletalMeshesCompatible(*this, ClothFacade, OtherClothFacade, OtherBoneRemap))
+				{
+					if (!GroupNameOffsets.Find(ClothCollectionGroup::SimVertices3D))
+					{
+						GroupNameOffsets.Add(ClothCollectionGroup::SimVertices3D) = ClothCollection->NumElements(ClothCollectionGroup::SimVertices3D);
+					}
+					if (!GroupNameOffsets.Find(ClothCollectionGroup::RenderVertices))
+					{
+						GroupNameOffsets.Add(ClothCollectionGroup::RenderVertices) = ClothCollection->NumElements(ClothCollectionGroup::RenderVertices);
+					}
+
+					ClothFacade.Append(OtherClothFacade);
+					bAreAnyValid = true;
+
+					if (!OtherBoneRemap.IsEmpty())
+					{
+						Private::RemapBones(ClothFacade, OtherBoneRemap, GroupNameOffsets[ClothCollectionGroup::SimVertices3D], GroupNameOffsets[ClothCollectionGroup::RenderVertices]);
+					}
+				}
 			}
 
 			// Append selections (with offsets)
@@ -925,7 +984,8 @@ void FChaosClothAssetMergeClothCollectionsNode::Evaluate(UE::Dataflow::FContext&
 
 			// Append cloth
 			const FCollectionClothConstFacade OtherClothFacade(OtherClothCollection);
-			if (OtherClothFacade.IsValid() && Private::AreSkeletalMeshesCompatible(*this, ClothFacade, OtherClothFacade))
+			TArray<int32> RemapBoneIndices;
+			if (OtherClothFacade.IsValid() && Private::AreSkeletalMeshesCompatible(*this, ClothFacade, OtherClothFacade, RemapBoneIndices))
 			{
 				ClothFacade.Append(OtherClothFacade);
 				bAreAnyValid = true;
