Dataflow Editor: use FRenderCollection's vertex normals in the construction view rendering
- Also update cloth 2D Sim mesh rendering to use a constant normal direction

diff --git a/Engine/Plugins/ChaosClothAssetEditor/Source/ChaosClothAssetDataflowNodes/Private/ChaosClothAsset/ChaosClothAssetDataflowNodesModule.cpp b/Engine/Plugins/ChaosClothAssetEditor/Source/ChaosClothAssetDataflowNodes/Private/ChaosClothAsset/ChaosClothAssetDataflowNodesModule.cpp
--- a/Engine/Plugins/ChaosClothAssetEditor/Source/ChaosClothAssetDataflowNodes/Private/ChaosClothAsset/ChaosClothAssetDataflowNodesModule.cpp
+++ b/Engine/Plugins/ChaosClothAssetEditor/Source/ChaosClothAssetDataflowNodes/Private/ChaosClothAsset/ChaosClothAssetDataflowNodesModule.cpp
@@ -233,7 +233,21 @@ PRAGMA_ENABLE_DEPRECATION_WARNINGS
 						}
 
 						Indices = ClothFacade.GetSimIndices2D();
-						Normals = ClothFacade.GetSimNormal();		// TODO: Should we hard code a 2D normal?
+
+						FVector3f Normal2D(0.0f, 0.0f, 0.0f);
+						for (const FIntVector3& Tri : ClothFacade.GetSimIndices2D())
+						{
+							const FVector2f AC = ClothFacade.GetSimPosition2D()[Tri[2]] - ClothFacade.GetSimPosition2D()[Tri[0]];
+							const FVector2f AB = ClothFacade.GetSimPosition2D()[Tri[1]] - ClothFacade.GetSimPosition2D()[Tri[0]];
+							const float Det = AB[0] * AC[1] - AC[0] * AB[1];
+							if (FMath::Abs(Det) < UE_SMALL_NUMBER)
+							{
+								continue;
+							}
+							Normal2D = Det > 0 ? FVector3f(0.0f, 0.0f, 1.0f) : FVector3f(0.0f, 0.0f, -1.0f);
+							break;
+						}
+						Normals.Init(Normal2D, ClothFacade.GetNumSimVertices2D());
 					}
 					else if (State.GetViewMode().GetName() == FClothRenderViewMode::Name)
 					{
diff --git a/Engine/Plugins/Experimental/Dataflow/Source/DataflowAssetTools/Private/Dataflow/CollectionRenderingPatternUtility.cpp b/Engine/Plugins/Experimental/Dataflow/Source/DataflowAssetTools/Private/Dataflow/CollectionRenderingPatternUtility.cpp
--- a/Engine/Plugins/Experimental/Dataflow/Source/DataflowAssetTools/Private/Dataflow/CollectionRenderingPatternUtility.cpp
+++ b/Engine/Plugins/Experimental/Dataflow/Source/DataflowAssetTools/Private/Dataflow/CollectionRenderingPatternUtility.cpp
@@ -66,8 +66,6 @@ namespace UE::Dataflow
 						Indices[TriangleIndex].Z - StartVertices)
 					);
 				}
-				FMeshNormals::QuickComputeVertexNormals(DynamicMesh);
-
 				DynamicMesh.EnableAttributes();
 
 				// Build Remmaping indices back into the colleciton. 
@@ -76,6 +74,13 @@ namespace UE::Dataflow
 					UE::Geometry::FNonManifoldMappingSupport::AttachNonManifoldVertexMappingData(Remapping, DynamicMesh);
 				}
 
+				FDynamicMeshNormalOverlay* const NormalOverlay = DynamicMesh.Attributes()->PrimaryNormals();
+				NormalOverlay->CreateFromPredicate([](int ParentVID, int TriIDA, int TriIDB) { return true; }, 0.f);
+				for (int ElemIdx : NormalOverlay->ElementIndicesItr())
+				{
+					NormalOverlay->SetElement(ElemIdx, Normals[ElemIdx]);
+				}
+
 				DynamicMesh.Attributes()->EnablePrimaryColors();
 				DynamicMesh.Attributes()->PrimaryColors()->CreateFromPredicate([](int ParentVID, int TriIDA, int TriIDB) {return true; }, 0.f);
 				DynamicMesh.EnableVertexColors(FVector3f::Zero());
