Fix Chaos Cloth Merge Collections: when main cloth facade fails to merge (due to incompatible skeletons), do not try to append selection sets or weight maps since their corresponding cloth data does not exist.

diff --git a/Engine/Plugins/ChaosClothAssetEditor/Source/ChaosClothAssetDataflowNodes/Private/ChaosClothAsset/MergeClothCollectionsNode.cpp b/Engine/Plugins/ChaosClothAssetEditor/Source/ChaosClothAssetDataflowNodes/Private/ChaosClothAsset/MergeClothCollectionsNode.cpp
--- a/Engine/Plugins/ChaosClothAssetEditor/Source/ChaosClothAssetDataflowNodes/Private/ChaosClothAsset/MergeClothCollectionsNode.cpp
+++ b/Engine/Plugins/ChaosClothAssetEditor/Source/ChaosClothAssetDataflowNodes/Private/ChaosClothAsset/MergeClothCollectionsNode.cpp
@@ -75,6 +75,7 @@ namespace UE::Chaos::ClothAsset::Private
 
 	/** Build weight maps for each properties if necessary */
 	static FString BuildWeightMaps(const FDataflowNode& DataflowNode,
+		bool bAppendedCloth,
 		const FCollectionClothConstFacade& InClothFacade, FCollectionClothFacade& OutClothFacade,
 		const FVector2f& InPropertyBounds, const FVector2f& OutPropertyBounds,
 		const FVector2f& PropertyBounds, const FString& PropertyName,
@@ -118,7 +119,7 @@ namespace UE::Chaos::ClothAsset::Private
 			TArrayView<float> WeightMap = OutClothFacade.GetWeightMap(FName(WeightMapName));
 
 			const int32 InNumVertices = InClothFacade.GetNumSimVertices3D();
-			const int32 OutNumVertices = OutClothFacade.GetNumSimVertices3D() - InNumVertices;
+			const int32 OutNumVertices = bAppendedCloth ? OutClothFacade.GetNumSimVertices3D() - InNumVertices : 0;
 			check(OutNumVertices >= 0);
 			
 			FillWeightMap(WeightMap.Left(OutNumVertices), OutClothFacade.GetWeightMap(FName(OutWeightMapName)).Left(OutNumVertices), PropertyBounds, OutPropertyBounds);
@@ -191,6 +192,7 @@ namespace UE::Chaos::ClothAsset::Private
 	}
 
 	static void UpdateSpringConstraintWeights(
+		bool bAppendedCloth,
 		const ::Chaos::Softs::FEmbeddedSpringFacade& InSpringFacade,
 		::Chaos::Softs::FEmbeddedSpringFacade& OutSpringFacade,
 		const FConstraintMergeData& ConstraintMergeData,
@@ -231,7 +233,7 @@ namespace UE::Chaos::ClothAsset::Private
 				InWeightMap = InConstraintFacade.GetDampingConst();
 			}
 			const int32 InNumSprings = InConstraintFacade.GetNumSprings();
-			const int32 OutNumSprings = OutConstraintFacade.GetNumSprings() - InNumSprings;
+			const int32 OutNumSprings = bAppendedCloth ? OutConstraintFacade.GetNumSprings() - InNumSprings : 0;
 			check(OutNumSprings >= 0);
 			FillWeightMap(WeightMap.Left(OutNumSprings), WeightMap.Left(OutNumSprings), PropertyBounds, OutPropertyBounds);
 			FillWeightMap(WeightMap.Right(InNumSprings), InWeightMap, PropertyBounds, InPropertyBounds);
@@ -259,7 +261,7 @@ namespace UE::Chaos::ClothAsset::Private
 				InWeightMap = InConstraintFacade.GetDampingConst();
 			}
 			const int32 InNumSprings = InConstraintFacade.GetNumSprings();
-			const int32 OutNumSprings = OutConstraintFacade.GetNumSprings() - InNumSprings;
+			const int32 OutNumSprings = bAppendedCloth ? OutConstraintFacade.GetNumSprings() - InNumSprings : 0;
 			check(OutNumSprings >= 0);
 			FillWeightMap(WeightMap.Left(OutNumSprings), WeightMap.Left(OutNumSprings), PropertyBounds, OutPropertyBounds);
 			FillWeightMap(WeightMap.Right(InNumSprings), InWeightMap, PropertyBounds, InPropertyBounds);
@@ -287,7 +289,7 @@ namespace UE::Chaos::ClothAsset::Private
 				InWeightMap = InConstraintFacade.GetDampingConst();
 			}
 			const int32 InNumSprings = InConstraintFacade.GetNumSprings();
-			const int32 OutNumSprings = OutConstraintFacade.GetNumSprings() - InNumSprings;
+			const int32 OutNumSprings = bAppendedCloth ? OutConstraintFacade.GetNumSprings() - InNumSprings : 0;
 			check(OutNumSprings >= 0);
 			FillWeightMap(WeightMap.Left(OutNumSprings), WeightMap.Left(OutNumSprings), PropertyBounds, OutPropertyBounds);
 			FillWeightMap(WeightMap.Right(InNumSprings), InWeightMap, PropertyBounds, InPropertyBounds);
@@ -341,6 +343,7 @@ namespace UE::Chaos::ClothAsset::Private
 
 	/** Append input properties to the output property facade and add potential weight maps */
 	static void AppendInputProperties(const FDataflowNode& DataflowNode,
+		bool bAppendedCloth,
 		const FCollectionClothConstFacade& InClothFacade,
 			  FCollectionClothFacade& OutClothFacade,
 		const ::Chaos::Softs::FCollectionPropertyConstFacade& InPropertyFacade,
@@ -381,12 +384,12 @@ namespace UE::Chaos::ClothAsset::Private
 
 					if (IsSpringConstraintProperty(FName(*InPropertyKey)))
 					{
-						UpdateSpringConstraintWeights(InSpringFacade, OutSpringFacade, ConstraintMergeData, InPropertyBounds, OutPropertyBounds, PropertyBounds, InPropertyKey);
+						UpdateSpringConstraintWeights(bAppendedCloth, InSpringFacade, OutSpringFacade, ConstraintMergeData, InPropertyBounds, OutPropertyBounds, PropertyBounds, InPropertyKey);
 					}
 					else
 					{
 						// We keep the string value to be the one in the output if defined
-						const FString WeightMapName = BuildWeightMaps(DataflowNode, InClothFacade, OutClothFacade,
+						const FString WeightMapName = BuildWeightMaps(DataflowNode, bAppendedCloth, InClothFacade, OutClothFacade,
 							InPropertyBounds, OutPropertyBounds, PropertyBounds, InPropertyKey,
 							InPropertyFacade.GetStringValue(InKeyIndex), OutPropertyFacade.GetStringValue(OutKeyIndex), MergedPropertyMaps);
 
@@ -650,6 +653,7 @@ void FChaosClothAssetMergeClothCollectionsNode_v2::Evaluate(UE::Dataflow::FConte
 
 			// Append cloth
 			const FCollectionClothConstFacade OtherClothFacade(OtherClothCollection);
+			bool bAppendedCloth = false;
 			if (OtherClothFacade.IsValid())
 			{
 
@@ -667,6 +671,7 @@ void FChaosClothAssetMergeClothCollectionsNode_v2::Evaluate(UE::Dataflow::FConte
 
 					ClothFacade.Append(OtherClothFacade);
 					bAreAnyValid = true;
+					bAppendedCloth = true;
 
 					if (!OtherBoneRemap.IsEmpty())
 					{
@@ -676,7 +681,7 @@ void FChaosClothAssetMergeClothCollectionsNode_v2::Evaluate(UE::Dataflow::FConte
 			}
 
 			// Append selections (with offsets)
-			if (OtherSelectionFacade.IsValid())
+			if (bAppendedCloth && OtherSelectionFacade.IsValid())
 			{
 				constexpr bool bUpdateExistingSelections = true; // Want last one wins.
 				SelectionFacade.AppendWithOffsets(OtherSelectionFacade, bUpdateExistingSelections, GroupNameOffsets);
@@ -685,7 +690,7 @@ void FChaosClothAssetMergeClothCollectionsNode_v2::Evaluate(UE::Dataflow::FConte
 
 			// Append springs (with offsets)
 			ConstraintMergeData.ResetOtherData();
-			if (OtherEmbeddedSpringFacade.IsValid())
+			if (bAppendedCloth && OtherEmbeddedSpringFacade.IsValid())
 			{
 				for (int32 ConstraintIndex = 0; ConstraintIndex < OtherEmbeddedSpringFacade.GetNumSpringConstraints(); ++ConstraintIndex)
 				{
@@ -791,7 +796,7 @@ void FChaosClothAssetMergeClothCollectionsNode_v2::Evaluate(UE::Dataflow::FConte
 				}
 				else
 				{
-					Private::AppendInputProperties(*this, OtherClothFacade, ClothFacade, OtherPropertyFacade, PropertyFacade, OtherEmbeddedSpringFacade, SpringFacade, ConstraintMergeData);
+					Private::AppendInputProperties(*this, bAppendedCloth, OtherClothFacade, ClothFacade, OtherPropertyFacade, PropertyFacade, OtherEmbeddedSpringFacade, SpringFacade, ConstraintMergeData);
 				}
 				bAreAnyValid = true;
 			}
@@ -985,14 +990,35 @@ void FChaosClothAssetMergeClothCollectionsNode::Evaluate(UE::Dataflow::FContext&
 			// Append cloth
 			const FCollectionClothConstFacade OtherClothFacade(OtherClothCollection);
 			TArray<int32> RemapBoneIndices;
-			if (OtherClothFacade.IsValid() && Private::AreSkeletalMeshesCompatible(*this, ClothFacade, OtherClothFacade, RemapBoneIndices))
+			bool bAppendedCloth = false;
+			if (OtherClothFacade.IsValid())
 			{
-				ClothFacade.Append(OtherClothFacade);
-				bAreAnyValid = true;
+
+				TArray<int32> OtherBoneRemap;
+				if (Private::AreSkeletalMeshesCompatible(*this, ClothFacade, OtherClothFacade, OtherBoneRemap))
+				{
+					if (!GroupNameOffsets.Find(ClothCollectionGroup::SimVertices3D))
+					{
+						GroupNameOffsets.Add(ClothCollectionGroup::SimVertices3D) = ClothCollection->NumElements(ClothCollectionGroup::SimVertices3D);
+					}
+					if (!GroupNameOffsets.Find(ClothCollectionGroup::RenderVertices))
+					{
+						GroupNameOffsets.Add(ClothCollectionGroup::RenderVertices) = ClothCollection->NumElements(ClothCollectionGroup::RenderVertices);
+					}
+
+					ClothFacade.Append(OtherClothFacade);
+					bAreAnyValid = true;
+					bAppendedCloth = true;
+
+					if (!OtherBoneRemap.IsEmpty())
+					{
+						Private::RemapBones(ClothFacade, OtherBoneRemap, GroupNameOffsets[ClothCollectionGroup::SimVertices3D], GroupNameOffsets[ClothCollectionGroup::RenderVertices]);
+					}
+				}
 			}
 
 			// Append selections (with offsets)
-			if (OtherSelectionFacade.IsValid())
+			if (bAppendedCloth && OtherSelectionFacade.IsValid())
 			{
 				constexpr bool bUpdateExistingSelections = true; // Want last one wins.
 				SelectionFacade.AppendWithOffsets(OtherSelectionFacade, bUpdateExistingSelections, GroupNameOffsets);
@@ -1000,7 +1026,7 @@ void FChaosClothAssetMergeClothCollectionsNode::Evaluate(UE::Dataflow::FContext&
 			}
 			// Append springs (with offsets)
 			ConstraintMergeData.ResetOtherData();
-			if (OtherEmbeddedSpringFacade.IsValid())
+			if (bAppendedCloth && OtherEmbeddedSpringFacade.IsValid())
 			{
 				if (OtherClothFacade.HasUserDefinedAttribute<int32>(TEXT("VertexSpringConstraintIndex"), ClothCollectionGroup::Lods))
 				{
@@ -1059,7 +1085,7 @@ void FChaosClothAssetMergeClothCollectionsNode::Evaluate(UE::Dataflow::FContext&
 				}
 				else
 				{
-					Private::AppendInputProperties(*this, OtherClothFacade, ClothFacade, OtherPropertyFacade, PropertyFacade, OtherEmbeddedSpringFacade, SpringFacade, ConstraintMergeData);
+					Private::AppendInputProperties(*this, bAppendedCloth, OtherClothFacade, ClothFacade, OtherPropertyFacade, PropertyFacade, OtherEmbeddedSpringFacade, SpringFacade, ConstraintMergeData);
 				}
 				bAreAnyValid = true;
 			}
