Convert the Chaos Cloth Asset Editor viewports to use the new viewport toolbars.

diff --git a/Engine/Plugins/ChaosClothAssetEditor/Source/ChaosClothAssetEditor/Private/ChaosClothAsset/ClothEditor3DViewportClient.cpp b/Engine/Plugins/ChaosClothAssetEditor/Source/ChaosClothAssetEditor/Private/ChaosClothAsset/ClothEditor3DViewportClient.cpp
--- a/Engine/Plugins/ChaosClothAssetEditor/Source/ChaosClothAssetEditor/Private/ChaosClothAsset/ClothEditor3DViewportClient.cpp
+++ b/Engine/Plugins/ChaosClothAssetEditor/Source/ChaosClothAssetEditor/Private/ChaosClothAsset/ClothEditor3DViewportClient.cpp
@@ -18,6 +18,7 @@
 #include "EngineUtils.h"
 #include "Animation/SkeletalMeshActor.h"
 #include "AssetViewerSettings.h"
+#include "ClothEditorCommands.h"
 #include "Editor/EditorPerProjectUserSettings.h"
 #include "Transforms/TransformGizmoDataBinder.h"
 #include "BaseBehaviors/ClickDragBehavior.h"
@@ -207,7 +208,7 @@ bool FChaosClothAssetEditor3DViewportClient::IsSimulationEnabled() const
 	return false;
 }
 
-void FChaosClothAssetEditor3DViewportClient::SetLODModel(int32 LODIndex)
+void FChaosClothAssetEditor3DViewportClient::SetLODLevel(int32 LODIndex)
 {
 	if (ClothEdMode)
 	{
@@ -215,7 +216,7 @@ void FChaosClothAssetEditor3DViewportClient::SetLODModel(int32 LODIndex)
 	}
 }
 
-bool FChaosClothAssetEditor3DViewportClient::IsLODModelSelected(int32 LODIndex) const
+bool FChaosClothAssetEditor3DViewportClient::IsLODSelected(int32 LODIndex) const
 {
 	if (ClothEdMode)
 	{
@@ -224,7 +225,7 @@ bool FChaosClothAssetEditor3DViewportClient::IsLODModelSelected(int32 LODIndex)
 	return false;
 }
 
-int32 FChaosClothAssetEditor3DViewportClient::GetLODModel() const
+int32 FChaosClothAssetEditor3DViewportClient::GetCurrentLOD() const
 {
 	if (ClothEdMode)
 	{
@@ -233,7 +234,7 @@ int32 FChaosClothAssetEditor3DViewportClient::GetLODModel() const
 	return INDEX_NONE;
 }
 
-int32 FChaosClothAssetEditor3DViewportClient::GetNumLODs() const
+int32 FChaosClothAssetEditor3DViewportClient::GetLODCount() const
 {
 	if (ClothEdMode)
 	{
@@ -242,6 +243,12 @@ int32 FChaosClothAssetEditor3DViewportClient::GetNumLODs() const
 	return 0;
 }
 
+void FChaosClothAssetEditor3DViewportClient::FillLODCommands(TArray<TSharedPtr<FUICommandInfo>>& Commands)
+{
+	Commands.Add(FChaosClothAssetEditorCommands::Get().LODAuto);
+	Commands.Add(FChaosClothAssetEditorCommands::Get().LOD0);
+}
+
 void FChaosClothAssetEditor3DViewportClient::Draw(const FSceneView* View, FPrimitiveDrawInterface* PDI)
 {
 	FEditorViewportClient::Draw(View, PDI);
diff --git a/Engine/Plugins/ChaosClothAssetEditor/Source/ChaosClothAssetEditor/Private/ChaosClothAsset/ClothEditor3DViewportClient.h b/Engine/Plugins/ChaosClothAssetEditor/Source/ChaosClothAssetEditor/Private/ChaosClothAsset/ClothEditor3DViewportClient.h
--- a/Engine/Plugins/ChaosClothAssetEditor/Source/ChaosClothAssetEditor/Private/ChaosClothAsset/ClothEditor3DViewportClient.h
+++ b/Engine/Plugins/ChaosClothAssetEditor/Source/ChaosClothAssetEditor/Private/ChaosClothAsset/ClothEditor3DViewportClient.h
@@ -6,6 +6,7 @@
 #include "Dataflow/DataflowNodeParameters.h"
 #include "GeometryCollection/ManagedArrayCollection.h"
 #include "InputBehaviorSet.h"
+#include "IPreviewLODController.h"
 #include "BaseBehaviors/BehaviorTargetInterfaces.h"
 
 class UChaosClothComponent;
@@ -29,7 +30,8 @@ class CHAOSCLOTHASSETEDITOR_API FChaosClothAssetEditor3DViewportClient : public
 	public TSharedFromThis<FChaosClothAssetEditor3DViewportClient>, 
 	public IClickBehaviorTarget,
 	public IClickDragBehaviorTarget,
-	public IInputBehaviorSource
+	public IInputBehaviorSource,
+	public IPreviewLODController
 {
 public:
 
@@ -64,11 +66,15 @@ public:
 	void SetEnableSimulation(bool bEnable);
 	bool IsSimulationEnabled() const;
 
+	// IPreviewLODController interface
 	// LODIndex == INDEX_NONE is LOD Auto
-	void SetLODModel(int32 LODIndex);
-	bool IsLODModelSelected(int32 LODIndex) const;
-	int32 GetLODModel() const;
-	int32 GetNumLODs() const;
+	virtual void SetLODLevel(int32 LODIndex) override;
+	virtual int32 GetLODCount() const override;
+	virtual int32 GetCurrentLOD() const override;
+	virtual bool IsLODSelected(int32 LODIndex) const override;
+	virtual int32 GetAutoLODStartingIndex() const override { return 1; }
+	virtual void FillLODCommands(TArray<TSharedPtr<FUICommandInfo>>& Commands) override;
+	// ~IPreviewLODController
 
 	FBox PreviewBoundingBox() const;
 
diff --git a/Engine/Plugins/ChaosClothAssetEditor/Source/ChaosClothAssetEditor/Private/ChaosClothAsset/ClothEditorCommands.cpp b/Engine/Plugins/ChaosClothAssetEditor/Source/ChaosClothAssetEditor/Private/ChaosClothAsset/ClothEditorCommands.cpp
--- a/Engine/Plugins/ChaosClothAssetEditor/Source/ChaosClothAssetEditor/Private/ChaosClothAsset/ClothEditorCommands.cpp
+++ b/Engine/Plugins/ChaosClothAssetEditor/Source/ChaosClothAssetEditor/Private/ChaosClothAsset/ClothEditorCommands.cpp
@@ -57,19 +57,19 @@ void FChaosClothAssetEditorCommands::RegisterCommands()
 	UI_COMMAND(SetConstructionMode3D, "3D Sim", "Switches the viewport to 3D simulation mesh view", EUserInterfaceActionType::RadioButton, FInputChord());
 	UI_COMMAND(SetConstructionModeRender, "Render", "Switches the viewport to render mesh view", EUserInterfaceActionType::RadioButton, FInputChord());
 
-	UI_COMMAND(TogglePreviewWireframe, "TogglePreviewWireframe", "Toggle preview wireframe", EUserInterfaceActionType::ToggleButton, FInputChord());
-	UI_COMMAND(ToggleConstructionViewWireframe, "ToggleConstructionViewWireframe", "Toggle construction view wireframe", EUserInterfaceActionType::ToggleButton, FInputChord());
+	UI_COMMAND(TogglePreviewWireframe, "Preview Wireframe", "Toggle preview wireframe", EUserInterfaceActionType::ToggleButton, FInputChord());
+	UI_COMMAND(ToggleConstructionViewWireframe, "Construction View Wireframe", "Toggle construction view wireframe", EUserInterfaceActionType::ToggleButton, FInputChord());
 
-	UI_COMMAND(ToggleConstructionViewSeams, "ToggleConstructionViewSeams", "Toggle construction view seams", EUserInterfaceActionType::ToggleButton, FInputChord());
-	UI_COMMAND(ToggleConstructionViewSeamsCollapse, "ToggleConstructionViewSeamsCollapse", "Toggle drawing one line between seams in construction view", EUserInterfaceActionType::ToggleButton, FInputChord());
-	UI_COMMAND(ToggleConstructionViewSurfaceNormals, "ToggleConstructionViewSurfaceNormals", "Toggle construction view surface normals", EUserInterfaceActionType::ToggleButton, FInputChord());
+	UI_COMMAND(ToggleConstructionViewSeams, "Show Seams", "Display seam information (not available for non-manifold meshes).", EUserInterfaceActionType::ToggleButton, FInputChord());
+	UI_COMMAND(ToggleConstructionViewSeamsCollapse, "Collapse Seam Lines", "Display a single line connecting each seam, rather than all stitches.", EUserInterfaceActionType::ToggleButton, FInputChord());
+	UI_COMMAND(ToggleConstructionViewSurfaceNormals, "Show Normals", "Display surface normals", EUserInterfaceActionType::ToggleButton, FInputChord());
 
-	UI_COMMAND(TogglePatternColor, "TogglePatternColor", "Toggle Pattern color in the construction view", EUserInterfaceActionType::ToggleButton, FInputChord());
-	UI_COMMAND(ToggleMeshStats, "ToggleMeshStats", "Toggle mesh information display", EUserInterfaceActionType::ToggleButton, FInputChord());
+	UI_COMMAND(TogglePatternColor, "Color Patterns", "Display each Pattern in a different color", EUserInterfaceActionType::ToggleButton, FInputChord());
+	UI_COMMAND(ToggleMeshStats, "Mesh Stats", "Show mesh stats in the viewport.", EUserInterfaceActionType::ToggleButton, FInputChord());
 	
-	UI_COMMAND(SoftResetSimulation, "SoftResetSimulation", "Soft reset simulation", EUserInterfaceActionType::ToggleButton, FInputChord());
-	UI_COMMAND(HardResetSimulation, "HardResetClothSimulation", "Hard reset simulation", EUserInterfaceActionType::ToggleButton, FInputChord(EModifierKey::Control | EModifierKey::Alt, EKeys::C));
-	UI_COMMAND(ToggleSimulationSuspended, "ToggleSimulationSuspended", "Toggle simulation suspended", EUserInterfaceActionType::ToggleButton, FInputChord());
+	UI_COMMAND(SoftResetSimulation, "Soft Reset Simulation", "Soft reset the cloth simulation", EUserInterfaceActionType::Button, FInputChord());
+	UI_COMMAND(HardResetSimulation, "Hard Reset Simulation", "Hard reset the cloth simulation", EUserInterfaceActionType::Button, FInputChord(EModifierKey::Control | EModifierKey::Alt, EKeys::C));
+	UI_COMMAND(ToggleSimulationSuspended, "Toggle Simulation", "Toggle the simulation of the cloth", EUserInterfaceActionType::ToggleButton, FInputChord());
 
 	UI_COMMAND(LODAuto, "LOD Auto", "Automatically select LOD", EUserInterfaceActionType::RadioButton, FInputChord());
 	UI_COMMAND(LOD0, "LOD 0", "Force select LOD 0", EUserInterfaceActionType::RadioButton, FInputChord());
diff --git a/Engine/Plugins/ChaosClothAssetEditor/Source/ChaosClothAssetEditor/Private/ChaosClothAsset/ClothEditorViewportToolbarSections.cpp b/Engine/Plugins/ChaosClothAssetEditor/Source/ChaosClothAssetEditor/Private/ChaosClothAsset/ClothEditorViewportToolbarSections.cpp
new file mode 100644
--- /dev/null
+++ b/Engine/Plugins/ChaosClothAssetEditor/Source/ChaosClothAssetEditor/Private/ChaosClothAsset/ClothEditorViewportToolbarSections.cpp
@@ -0,0 +1,138 @@
+﻿#include "ClothEditorViewportToolbarSections.h"
+
+#include "ClothEditorRestSpaceViewportClient.h"
+#include "ChaosClothAsset/ClothEditorCommands.h"
+#include "EditorModeManager.h"
+#include "SEditorViewport.h"
+#include "ViewportToolbar/UnrealEdViewportToolbar.h"
+#include "Widgets/Layout/SBorder.h"
+#include "Widgets/Input/SSpinBox.h"
+
+#define LOCTEXT_NAMESPACE "ClothEditorViewportToolbarSections"
+
+namespace UE::Chaos::ClothAsset
+{
+
+FToolMenuEntry CreateDynamicLightIntensityItem()
+{
+	return FToolMenuEntry::InitDynamicEntry("DynamicLightIntensity", FNewToolMenuSectionDelegate::CreateLambda([](FToolMenuSection& Section)
+	{
+		UUnrealEdViewportToolbarContext* Context = Section.FindContext<UUnrealEdViewportToolbarContext>();
+		if (!Context)
+		{
+			return;
+		}
+		
+		TSharedPtr<SEditorViewport> Viewport = Context->Viewport.Pin();
+		if (!Viewport)
+		{
+			return;
+		}
+		
+		TSharedPtr<FChaosClothEditorRestSpaceViewportClient> ViewportClient = StaticCastSharedPtr<FChaosClothEditorRestSpaceViewportClient>(Viewport->GetViewportClient());
+		
+		constexpr float IntensityMin = 0.f;
+		constexpr float IntensityMax = 20.f;
+		
+		TSharedRef<SWidget> Widget = SNew(SBox)
+			.HAlign(HAlign_Right)
+			[
+				SNew(SBox)
+				.Padding(FMargin(4.0f, 0.0f, 0.0f, 0.0f))
+				.WidthOverride(100.0f)
+				[
+					SNew(SBorder)
+					.BorderImage(FAppStyle::Get().GetBrush("Menu.WidgetBorder"))
+					.Padding(FMargin(1.0f))
+					[
+						SNew(SSpinBox<float>)
+						.Style(&FAppStyle::Get(), "Menu.SpinBox")
+						.Font(FAppStyle::GetFontStyle(TEXT("MenuItem.Font")))
+						.MinValue(IntensityMin)
+						.MaxValue(IntensityMax)
+						.Value_Lambda([WeakClient = ViewportClient.ToWeakPtr()]
+						{
+							if (TSharedPtr<FChaosClothEditorRestSpaceViewportClient> Client = WeakClient.Pin())
+							{
+								return Client->GetCameraPointLightIntensity();
+							}
+							return 0.f;
+						})
+						.OnValueChanged_Lambda([WeakClient = ViewportClient.ToWeakPtr()](float Intensity)
+						{
+							if (TSharedPtr<FChaosClothEditorRestSpaceViewportClient> Client = WeakClient.Pin())
+							{
+								Client->SetCameraPointLightIntensity(Intensity);
+							}
+						}) 
+						.IsEnabled_Lambda([WeakClient = ViewportClient.ToWeakPtr()]
+						{
+							if (TSharedPtr<FChaosClothEditorRestSpaceViewportClient> Client = WeakClient.Pin())
+							{
+								return Client->GetConstructionViewMode() == EClothPatternVertexType::Render;
+							}
+							return false;
+						})
+					]
+				]
+			];
+			
+		Section.AddEntry(FToolMenuEntry::InitWidget("LightIntensity", Widget, LOCTEXT("LightIntensityLabel", "Render Light Intensity")));
+	}));
+}
+
+FToolMenuEntry CreateDynamicSimulationMenuItem()
+{
+	return FToolMenuEntry::InitDynamicEntry("DynamicSimulationMenu", FNewToolMenuSectionDelegate::CreateLambda([](FToolMenuSection& Section)
+	{
+		UUnrealEdViewportToolbarContext* Context = Section.FindContext<UUnrealEdViewportToolbarContext>();
+		if (!Context)
+		{
+			return;
+		}
+		
+		const TAttribute<FText> SimLabel = TAttribute<FText>::CreateLambda([WeakViewport = Context->Viewport]
+		{
+			if (TSharedPtr<SEditorViewport> Viewport = WeakViewport.Pin())
+			{
+				const TSharedPtr<FEditorViewportClient> ViewportClient = Viewport->GetViewportClient();
+				check(ViewportClient.IsValid());
+				const FEditorModeTools* const EditorModeTools = ViewportClient->GetModeTools();
+				UChaosClothAssetEditorMode* const ClothEdMode = Cast<UChaosClothAssetEditorMode>(EditorModeTools->GetActiveScriptableMode(UChaosClothAssetEditorMode::EM_ChaosClothAssetEditorModeId));
+				if (ClothEdMode)
+				{
+					switch (ClothEdMode->GetConstructionViewMode())
+					{
+					case EClothPatternVertexType::Sim2D:
+						return LOCTEXT("ConstructionViewMenuTitle_Sim2D", "2D Sim");
+					case EClothPatternVertexType::Sim3D:
+						return LOCTEXT("ConstructionViewMenuTitle_Sim3D", "3D Sim");
+					case EClothPatternVertexType::Render:
+						return LOCTEXT("ConstructionViewMenuTitle_Render", "Render");
+					}
+				}
+			}
+			
+			return LOCTEXT("ConstructionViewMenuTitle_Default", "View"); 
+		});
+		
+		Section.AddSubMenu(
+			"SimMode",
+			SimLabel,
+			LOCTEXT("SimulationMenuTooltip", "Change the simulation view."),
+			FNewToolMenuDelegate::CreateLambda([](UToolMenu* InMenu)
+			{
+				FToolMenuSection& Section = InMenu->FindOrAddSection("ConstructionViewModeMenuSection");
+				Section.AddMenuEntry(FChaosClothAssetEditorCommands::Get().SetConstructionMode2D);
+				Section.AddMenuEntry(FChaosClothAssetEditorCommands::Get().SetConstructionMode3D);
+				Section.AddMenuEntry(FChaosClothAssetEditorCommands::Get().SetConstructionModeRender);
+			})
+		);
+	}));
+}
+
+
+
+}
+
+#undef LOCTEXT_NAMESPACE
\ No newline at end of file
diff --git a/Engine/Plugins/ChaosClothAssetEditor/Source/ChaosClothAssetEditor/Private/ChaosClothAsset/ClothEditorViewportToolbarSections.h b/Engine/Plugins/ChaosClothAssetEditor/Source/ChaosClothAssetEditor/Private/ChaosClothAsset/ClothEditorViewportToolbarSections.h
new file mode 100644
--- /dev/null
+++ b/Engine/Plugins/ChaosClothAssetEditor/Source/ChaosClothAssetEditor/Private/ChaosClothAsset/ClothEditorViewportToolbarSections.h
@@ -0,0 +1,13 @@
+﻿// Copyright Epic Games, Inc. All Rights Reserved.
+
+#pragma once
+
+#include "CoreMinimal.h"
+#include "ToolMenus.h"
+
+// Common toolbar functions shared by the various cloth editor viewport toolbars.
+namespace UE::Chaos::ClothAsset
+{
+	FToolMenuEntry CreateDynamicLightIntensityItem();
+	FToolMenuEntry CreateDynamicSimulationMenuItem();
+}
diff --git a/Engine/Plugins/ChaosClothAssetEditor/Source/ChaosClothAssetEditor/Private/ChaosClothAsset/SClothEditor3DViewport.cpp b/Engine/Plugins/ChaosClothAssetEditor/Source/ChaosClothAssetEditor/Private/ChaosClothAsset/SClothEditor3DViewport.cpp
--- a/Engine/Plugins/ChaosClothAssetEditor/Source/ChaosClothAssetEditor/Private/ChaosClothAsset/SClothEditor3DViewport.cpp
+++ b/Engine/Plugins/ChaosClothAssetEditor/Source/ChaosClothAssetEditor/Private/ChaosClothAsset/SClothEditor3DViewport.cpp
@@ -1,6 +1,8 @@
 // Copyright Epic Games, Inc. All Rights Reserved.
 
 #include "ChaosClothAsset/SClothEditor3DViewport.h"
+
+#include "ClothEditorViewportToolbarSections.h"
 #include "ChaosClothAsset/ClothEditor3DViewportClient.h"
 #include "ChaosClothAsset/SClothEditor3DViewportToolBar.h"
 #include "ChaosClothAsset/ClothEditorCommands.h"
@@ -9,8 +11,11 @@
 #include "ChaosClothAsset/SClothAnimationScrubPanel.h"
 #include "Animation/AnimSingleNodeInstance.h"
 #include "EditorModeTools.h"
+#include "ToolMenus.h"
+#include "PreviewProfileController.h"
 #include "Widgets/Layout/SBorder.h"
 #include "Components/SkeletalMeshComponent.h"
+#include "ViewportToolbar/UnrealEdViewportToolbar.h"
 
 #define LOCTEXT_NAMESPACE "SChaosClothAssetEditor3DViewport"
 
@@ -60,6 +65,10 @@ TWeakPtr<const UE::Chaos::ClothAsset::FChaosClothPreviewScene> SChaosClothAssetE
 	return ClothViewportClient->GetClothPreviewScene();
 }
 
+TSharedPtr<IPreviewProfileController> SChaosClothAssetEditor3DViewport::CreatePreviewProfileController()
+{
+	return MakeShared<FPreviewProfileController>();	
+}
 
 void SChaosClothAssetEditor3DViewport::BindCommands()
 {
@@ -123,13 +132,13 @@ void SChaosClothAssetEditor3DViewport::BindCommands()
 		FExecuteAction::CreateLambda([this]()
 		{
 			TSharedPtr<FChaosClothAssetEditor3DViewportClient> ClothViewportClient = StaticCastSharedPtr<FChaosClothAssetEditor3DViewportClient>(Client);
-			ClothViewportClient->SetLODModel(INDEX_NONE);
+			ClothViewportClient->SetLODLevel(INDEX_NONE);
 		}),
 		FCanExecuteAction(),
 		FIsActionChecked::CreateLambda([this]()
 		{
 			TSharedPtr<FChaosClothAssetEditor3DViewportClient> ClothViewportClient = StaticCastSharedPtr<FChaosClothAssetEditor3DViewportClient>(Client);
-			return ClothViewportClient->IsLODModelSelected(INDEX_NONE);
+			return ClothViewportClient->IsLODSelected(INDEX_NONE);
 		}
 		));
 
@@ -138,13 +147,13 @@ void SChaosClothAssetEditor3DViewport::BindCommands()
 		FExecuteAction::CreateLambda([this]()
 		{
 			TSharedPtr<FChaosClothAssetEditor3DViewportClient> ClothViewportClient = StaticCastSharedPtr<FChaosClothAssetEditor3DViewportClient>(Client);
-			ClothViewportClient->SetLODModel(0);
+			ClothViewportClient->SetLODLevel(0);
 		}),
 		FCanExecuteAction(),
 		FIsActionChecked::CreateLambda([this]()
 		{
 			TSharedPtr<FChaosClothAssetEditor3DViewportClient> ClothViewportClient = StaticCastSharedPtr<FChaosClothAssetEditor3DViewportClient>(Client);
-			return ClothViewportClient->IsLODModelSelected(0);
+			return ClothViewportClient->IsLODSelected(0);
 		}
 		));
 
@@ -158,6 +167,92 @@ TSharedPtr<SWidget> SChaosClothAssetEditor3DViewport::MakeViewportToolbar()
 		.CommandList(ToolkitCommandList);
 }
 
+TSharedPtr<SWidget> SChaosClothAssetEditor3DViewport::BuildViewportToolbar()
+{
+	using namespace UE::Chaos::ClothAsset;
+
+	const FName ViewportToolbarName = "ChaosClothEditor.3DViewportToolbar";
+	
+	if (!UToolMenus::Get()->IsMenuRegistered(ViewportToolbarName))
+	{
+		UToolMenu* ToolbarMenu = UToolMenus::Get()->RegisterMenu(ViewportToolbarName, NAME_None, EMultiBoxType::SlimHorizontalToolBar);
+		ToolbarMenu->StyleName = "ViewportToolbar";
+		
+		ToolbarMenu->FindOrAddSection("Left");
+		
+		{
+			FToolMenuSection& RightSection = ToolbarMenu->FindOrAddSection("Right");
+			RightSection.Alignment = EToolMenuSectionAlign::Last;
+			
+			RightSection.AddEntry(UE::UnrealEd::CreateCameraSubmenu(UE::UnrealEd::FViewportCameraMenuOptions().ShowAll()));
+			
+			{
+				RightSection.AddEntry(UE::UnrealEd::CreateViewModesSubmenu());
+				UToolMenu* ViewModesMenu = UToolMenus::Get()->ExtendMenu(UToolMenus::JoinMenuPaths(ViewportToolbarName, "ViewModes"));
+				FToolMenuSection& ViewSection = ViewModesMenu->FindOrAddSection("Cloth", LOCTEXT("ClothViewModeSection", "Cloth"));
+				
+				FToolMenuEntry& WireframeEntry = ViewSection.AddMenuEntry(FChaosClothAssetEditorCommands::Get().TogglePreviewWireframe);
+				WireframeEntry.SetShowInToolbarTopLevel(true);
+				WireframeEntry.ToolBarData.ResizeParams.ClippingPriority = 2000;
+			}
+			
+			{
+				// LOD
+				RightSection.AddDynamicEntry("DynamicLOD", FNewToolMenuSectionDelegate::CreateLambda([](FToolMenuSection& Section)
+				{
+					if (UUnrealEdViewportToolbarContext* Context = Section.FindContext<UUnrealEdViewportToolbarContext>())
+					{
+						if (TSharedPtr<SEditorViewport> EditorViewport = Context->Viewport.Pin())
+						{
+							if (TSharedPtr<FEditorViewportClient> Client = EditorViewport->GetViewportClient())
+							{
+								TSharedPtr<FChaosClothAssetEditor3DViewportClient> ClothClient = StaticCastSharedPtr<FChaosClothAssetEditor3DViewportClient>(Client);
+								Section.AddEntry(UE::UnrealEd::CreatePreviewLODSelectionSubmenu(ClothClient));
+							}
+						}
+					}
+				}));
+			}
+			
+			{
+				RightSection.AddEntry(UE::UnrealEd::CreateAssetViewerProfileSubmenu());
+				UToolMenu* AssetViewerMenu = UToolMenus::Get()->ExtendMenu(UToolMenus::JoinMenuPaths(ViewportToolbarName, "AssetViewerProfile"));
+				FToolMenuSection& SimulationSection = AssetViewerMenu->FindOrAddSection("SimulationControls", LOCTEXT("SimulationControlsSection", "Simulation Playback Controls"));
+				
+				// Simulation Controls
+				FToolMenuEntryToolBarData ToolBarData;
+				ToolBarData.BlockGroupName = "SimulationControlGroup";
+				ToolBarData.LabelOverride = FText::GetEmpty();
+				ToolBarData.ResizeParams.ClippingPriority = 2000;
+				
+				FToolMenuEntry& HardReset = SimulationSection.AddMenuEntry(FChaosClothAssetEditorCommands::Get().HardResetSimulation);
+				HardReset.ToolBarData = ToolBarData;
+				HardReset.SetShowInToolbarTopLevel(true);
+				
+				FToolMenuEntry& SoftReset = SimulationSection.AddMenuEntry(FChaosClothAssetEditorCommands::Get().SoftResetSimulation);
+				SoftReset.ToolBarData = ToolBarData;
+				SoftReset.SetShowInToolbarTopLevel(true);
+				
+				FToolMenuEntry& ToggleSimulation = SimulationSection.AddMenuEntry(FChaosClothAssetEditorCommands::Get().ToggleSimulationSuspended);
+				ToggleSimulation.ToolBarData = ToolBarData;
+				ToggleSimulation.SetShowInToolbarTopLevel(true);
+			}
+		}
+	}
+	
+	FToolMenuContext Context;
+	{
+		Context.AppendCommandList(GetCommandList());
+		Context.AppendCommandList(ToolkitCommandList);
+		Context.AddExtender(GetExtenders());
+		
+		UUnrealEdViewportToolbarContext* ContextObject = UE::UnrealEd::CreateViewportToolbarDefaultContext(SharedThis(this));
+		Context.AddObject(ContextObject);
+	}
+	
+	return UToolMenus::Get()->GenerateWidget(ViewportToolbarName, Context);
+}
+
 bool SChaosClothAssetEditor3DViewport::IsVisible() const
 {
 	// Intentionally not calling SEditorViewport::IsVisible because it will return false if our simulation is more than 250ms.
diff --git a/Engine/Plugins/ChaosClothAssetEditor/Source/ChaosClothAssetEditor/Private/ChaosClothAsset/SClothEditor3DViewport.h b/Engine/Plugins/ChaosClothAssetEditor/Source/ChaosClothAssetEditor/Private/ChaosClothAsset/SClothEditor3DViewport.h
--- a/Engine/Plugins/ChaosClothAssetEditor/Source/ChaosClothAssetEditor/Private/ChaosClothAsset/SClothEditor3DViewport.h
+++ b/Engine/Plugins/ChaosClothAssetEditor/Source/ChaosClothAssetEditor/Private/ChaosClothAsset/SClothEditor3DViewport.h
@@ -32,6 +32,7 @@ public:
 	// SAssetEditorViewport
 	virtual void BindCommands() override;
 	virtual TSharedPtr<SWidget> MakeViewportToolbar() override;
+	virtual TSharedPtr<SWidget> BuildViewportToolbar() override;
 	virtual bool IsVisible() const override;
 
 	virtual void OnFocusViewportToSelection() override;
@@ -50,6 +51,8 @@ private:
 	TWeakPtr<UE::Chaos::ClothAsset::FChaosClothPreviewScene> GetPreviewScene();
 	TWeakPtr<const UE::Chaos::ClothAsset::FChaosClothPreviewScene> GetPreviewScene() const;
 
+	virtual TSharedPtr<IPreviewProfileController> CreatePreviewProfileController() override;
+
 	float GetViewMinInput() const;
 	float GetViewMaxInput() const;
 	EVisibility GetAnimControlVisibility() const;
diff --git a/Engine/Plugins/ChaosClothAssetEditor/Source/ChaosClothAssetEditor/Private/ChaosClothAsset/SClothEditor3DViewportToolBar.cpp b/Engine/Plugins/ChaosClothAssetEditor/Source/ChaosClothAssetEditor/Private/ChaosClothAsset/SClothEditor3DViewportToolBar.cpp
--- a/Engine/Plugins/ChaosClothAssetEditor/Source/ChaosClothAssetEditor/Private/ChaosClothAsset/SClothEditor3DViewportToolBar.cpp
+++ b/Engine/Plugins/ChaosClothAssetEditor/Source/ChaosClothAssetEditor/Private/ChaosClothAsset/SClothEditor3DViewportToolBar.cpp
@@ -190,7 +190,7 @@ FText SChaosClothAssetEditor3DViewportToolBar::GetLODMenuLabel() const
 	TSharedRef<FChaosClothAssetEditor3DViewportClient> ViewportClient = 
 		StaticCastSharedPtr<FChaosClothAssetEditor3DViewportClient>(ChaosClothAssetEditor3DViewportPtr.Pin()->GetViewportClient()).ToSharedRef();
 
-	const int32 LODSelectionType = ViewportClient->GetLODModel();
+	const int32 LODSelectionType = ViewportClient->GetCurrentLOD();
 	if (LODSelectionType >= 0)
 	{
 		const FString TitleLabel = FString::Printf(TEXT("LOD %d"), LODSelectionType);
@@ -209,7 +209,7 @@ TSharedRef<SWidget> SChaosClothAssetEditor3DViewportToolBar::MakeLODMenu() const
 	FMenuBuilder MenuBuilder(bInShouldCloseWindowAfterMenuSelection, CommandList);
 
 	MenuBuilder.PushCommandList(CommandList.ToSharedRef());
-	const int32 NumLODs = ViewportClient->GetNumLODs();
+	const int32 NumLODs = ViewportClient->GetLODCount();
 
 	MenuBuilder.BeginSection("ClothAssetPreviewLODs", LOCTEXT("ShowLOD_PreviewLabel", "Preview LODs"));
 	{
@@ -220,9 +220,9 @@ TSharedRef<SWidget> SChaosClothAssetEditor3DViewportToolBar::MakeLODMenu() const
 		{
 			FString TitleLabel = FString::Printf(TEXT("LOD %d"), LODIndex);
 
-			FUIAction Action(FExecuteAction::CreateSP(ViewportClient, &FChaosClothAssetEditor3DViewportClient::SetLODModel, LODIndex),
+			FUIAction Action(FExecuteAction::CreateSP(ViewportClient, &FChaosClothAssetEditor3DViewportClient::SetLODLevel, LODIndex),
 				FCanExecuteAction(),
-				FIsActionChecked::CreateSP(ViewportClient, &FChaosClothAssetEditor3DViewportClient::IsLODModelSelected, LODIndex));
+				FIsActionChecked::CreateSP(ViewportClient, &FChaosClothAssetEditor3DViewportClient::IsLODSelected, LODIndex));
 
 			MenuBuilder.AddMenuEntry(FText::FromString(TitleLabel), FText::GetEmpty(), FSlateIcon(), Action, NAME_None, EUserInterfaceActionType::RadioButton);
 		}
diff --git a/Engine/Plugins/ChaosClothAssetEditor/Source/ChaosClothAssetEditor/Private/ChaosClothAsset/SClothEditorRestSpaceViewport.cpp b/Engine/Plugins/ChaosClothAssetEditor/Source/ChaosClothAssetEditor/Private/ChaosClothAsset/SClothEditorRestSpaceViewport.cpp
--- a/Engine/Plugins/ChaosClothAssetEditor/Source/ChaosClothAssetEditor/Private/ChaosClothAsset/SClothEditorRestSpaceViewport.cpp
+++ b/Engine/Plugins/ChaosClothAssetEditor/Source/ChaosClothAssetEditor/Private/ChaosClothAsset/SClothEditorRestSpaceViewport.cpp
@@ -1,12 +1,19 @@
 // Copyright Epic Games, Inc. All Rights Reserved.
 
 #include "ChaosClothAsset/SClothEditorRestSpaceViewport.h"
+
+#include "ClothEditorViewportToolbarSections.h"
 #include "ChaosClothAsset/ClothEditorRestSpaceViewportClient.h"
 #include "SViewportToolBar.h"
 #include "ChaosClothAsset/SClothEditorRestSpaceViewportToolBar.h"
 #include "ChaosClothAsset/ClothEditorCommands.h"
 #include "ChaosClothAsset/ClothEditorMode.h"
 #include "EditorModeManager.h"
+#include "ToolMenus.h"
+#include "Widgets/Input/SSpinBox.h"
+#include "Widgets/Layout/SBorder.h"
+#include "ViewportToolbar/UnrealEdViewportToolbar.h"
+#include "Widgets/Text/SRichTextBlock.h"
 
 #define LOCTEXT_NAMESPACE "SChaosClothAssetEditorRestSpaceViewport"
 
@@ -294,6 +301,115 @@ TSharedPtr<SWidget> SChaosClothAssetEditorRestSpaceViewport::MakeViewportToolbar
 		.RestSpaceViewportClient(RestSpaceViewportClient);
 }
 
+TSharedPtr<SWidget> SChaosClothAssetEditorRestSpaceViewport::BuildViewportToolbar()
+{
+	using namespace UE::Chaos::ClothAsset;
+
+	const FName ToolbarName = "ChaosClothEditor.RestSpaceViewportToolbar";
+	
+	if (!UToolMenus::Get()->IsMenuRegistered(ToolbarName))
+	{
+		UToolMenu* ToolbarMenu = UToolMenus::Get()->RegisterMenu(ToolbarName, NAME_None, EMultiBoxType::SlimHorizontalToolBar);
+		ToolbarMenu->StyleName = "ViewportToolbar";
+		
+		ToolbarMenu->FindOrAddSection("Left");
+		
+		{
+			FToolMenuSection& RightSection = ToolbarMenu->FindOrAddSection("Right");
+			RightSection.Alignment = EToolMenuSectionAlign::Last;
+			
+			RightSection.AddEntry(UE::UnrealEd::CreateCameraSubmenu(UE::UnrealEd::FViewportCameraMenuOptions().ShowAll()));
+			
+			{
+				// View Modes
+				RightSection.AddEntry(UE::UnrealEd::CreateViewModesSubmenu());
+				UToolMenu* ViewModesMenu = UToolMenus::Get()->ExtendMenu(UToolMenus::JoinMenuPaths(ToolbarName, "ViewModes"));
+				FToolMenuSection& ViewSection = ViewModesMenu->FindOrAddSection("Cloth", LOCTEXT("ClothViewModeSection", "Cloth"));
+				
+				FToolMenuEntry& WireframeEntry = ViewSection.AddMenuEntry(FChaosClothAssetEditorCommands::Get().ToggleConstructionViewWireframe);
+				WireframeEntry.SetShowInToolbarTopLevel(true);
+				WireframeEntry.ToolBarData.ResizeParams.ClippingPriority = 2000;
+				
+				ViewSection.AddEntry(CreateDynamicLightIntensityItem());
+			}
+			
+			{
+				// Show Menu
+				FToolMenuEntry& ShowEntry = RightSection.AddSubMenu(
+					"Show",
+					LOCTEXT("ShowSubmenuLabel", "Show"),
+					LOCTEXT("ShowSubmenuTooltip", "Show flags related to the current viewport"),
+					FNewToolMenuDelegate::CreateLambda(
+						[](UToolMenu* InMenu) -> void
+						{
+							FToolMenuSection& Section = InMenu->FindOrAddSection("Cloth", LOCTEXT("ClothShowMenu", "Chaos Cloth"));
+							Section.AddMenuEntry(FChaosClothAssetEditorCommands::Get().ToggleMeshStats);
+							Section.AddMenuEntry(FChaosClothAssetEditorCommands::Get().ToggleConstructionViewSeams);
+							Section.AddMenuEntry(FChaosClothAssetEditorCommands::Get().ToggleConstructionViewSeamsCollapse);
+							Section.AddMenuEntry(FChaosClothAssetEditorCommands::Get().TogglePatternColor);
+							Section.AddMenuEntry(FChaosClothAssetEditorCommands::Get().ToggleConstructionViewSurfaceNormals);
+						}
+					),
+					false,
+					FSlateIcon(FAppStyle::GetAppStyleSetName(), "Level.VisibleHighlightIcon16x")
+				);
+				
+				ShowEntry.ToolBarData.LabelOverride = FText();
+				ShowEntry.ToolBarData.ResizeParams.ClippingPriority = 800;
+			}
+			
+			RightSection.AddEntry(CreateDynamicSimulationMenuItem());
+		}
+	}
+	
+	FToolMenuContext Context;
+	{
+		Context.AppendCommandList(GetCommandList());
+		Context.AddExtender(GetExtenders());
+		
+		UUnrealEdViewportToolbarContext* ContextObject = UE::UnrealEd::CreateViewportToolbarDefaultContext(SharedThis(this));
+		Context.AddObject(ContextObject);
+	}
+	
+	return UToolMenus::Get()->GenerateWidget(ToolbarName, Context);
+}
+
+void SChaosClothAssetEditorRestSpaceViewport::PopulateViewportOverlays(TSharedRef<SOverlay> Overlay)
+{
+	Overlay->AddSlot()
+	.Padding(TAttribute<FMargin>::CreateLambda([]
+	{
+		if (UE::UnrealEd::ShowOldViewportToolbars())
+		{
+			return FMargin(4.f, 36.f, 0.f, 0.f);
+		}
+		return FMargin(4.f, 4.f, 0.f, 0.f); 
+	}))
+	[
+		SNew(SRichTextBlock)
+		.DecoratorStyleSet(&FAppStyle::Get())
+		.Text(this, &SChaosClothAssetEditorRestSpaceViewport::GetDisplayString)
+		.TextStyle(&FAppStyle::Get().GetWidgetStyle<FTextBlockStyle>("AnimViewport.MessageText"))
+	];
+}
+
+FText SChaosClothAssetEditorRestSpaceViewport::GetDisplayString() const
+{
+	if (const FEditorModeTools* const EditorModeTools = RestSpaceViewportClient->GetModeTools())
+	{
+		if (UChaosClothAssetEditorMode* const ClothEdMode = Cast<UChaosClothAssetEditorMode>(EditorModeTools->GetActiveScriptableMode(UChaosClothAssetEditorMode::EM_ChaosClothAssetEditorModeId)))
+		{
+			if (ClothEdMode->IsMeshStatsActive())
+			{
+				const int TriangleCount = ClothEdMode->GetConstructionViewTriangleCount();
+				const int VertexCount = ClothEdMode->GetConstructionViewVertexCount();
+				const FText MeshStats = FText::Format(LOCTEXT("RestSpaceMeshStats", "Tris: {0}, Verts: {1}"), TriangleCount, VertexCount);
+				return MeshStats;
+			}
+		}
+	}
+	return FText();
+}
 
 void SChaosClothAssetEditorRestSpaceViewport::OnFocusViewportToSelection()
 {
diff --git a/Engine/Plugins/ChaosClothAssetEditor/Source/ChaosClothAssetEditor/Private/ChaosClothAsset/SClothEditorRestSpaceViewport.h b/Engine/Plugins/ChaosClothAssetEditor/Source/ChaosClothAssetEditor/Private/ChaosClothAsset/SClothEditorRestSpaceViewport.h
--- a/Engine/Plugins/ChaosClothAssetEditor/Source/ChaosClothAssetEditor/Private/ChaosClothAsset/SClothEditorRestSpaceViewport.h
+++ b/Engine/Plugins/ChaosClothAssetEditor/Source/ChaosClothAssetEditor/Private/ChaosClothAsset/SClothEditorRestSpaceViewport.h
@@ -24,6 +24,8 @@ public:
 	// SEditorViewport
 	virtual void BindCommands() override;
 	virtual TSharedPtr<SWidget> MakeViewportToolbar() override;
+	virtual TSharedPtr<SWidget> BuildViewportToolbar() override;
+	virtual void PopulateViewportOverlays(TSharedRef<SOverlay> Overlay) override;
 	virtual void OnFocusViewportToSelection() override;
 	virtual bool IsVisible() const override;
 
@@ -35,6 +37,8 @@ public:
 private:
 
 	UChaosClothAssetEditorMode* GetEdMode() const;
+	
+	FText GetDisplayString() const; 
 
 	TSharedPtr<UE::Chaos::ClothAsset::FChaosClothEditorRestSpaceViewportClient> RestSpaceViewportClient;
 
diff --git a/Engine/Plugins/ChaosClothAssetEditor/Source/ChaosClothAssetEditor/Private/ChaosClothAsset/SClothEditorRestSpaceViewportToolBar.cpp b/Engine/Plugins/ChaosClothAssetEditor/Source/ChaosClothAssetEditor/Private/ChaosClothAsset/SClothEditorRestSpaceViewportToolBar.cpp
--- a/Engine/Plugins/ChaosClothAssetEditor/Source/ChaosClothAssetEditor/Private/ChaosClothAsset/SClothEditorRestSpaceViewportToolBar.cpp
+++ b/Engine/Plugins/ChaosClothAssetEditor/Source/ChaosClothAssetEditor/Private/ChaosClothAsset/SClothEditorRestSpaceViewportToolBar.cpp
@@ -39,21 +39,7 @@ void SChaosClothAssetEditorRestSpaceViewportToolBar::Construct(const FArguments&
 		.BorderImage(FAppStyle::Get().GetBrush("EditorViewportToolBar.Background"))
 		.Cursor(EMouseCursor::Default)
 		[
-			SNew(SVerticalBox)
-				+ SVerticalBox::Slot()
-				.AutoHeight()
-			[
-				SAssignNew( MainBoxPtr, SHorizontalBox )
-			]
-			+ SVerticalBox::Slot()
-			.Padding(FMargin(4.0f, 3.0f, 0.0f, 0.0f))
-			[
-				// Display text (e.g., item being previewed)
-				SNew(SRichTextBlock)
-					.DecoratorStyleSet(&FAppStyle::Get())
-					.Text(this, &SChaosClothAssetEditorRestSpaceViewportToolBar::GetDisplayString)
-					.TextStyle(&FAppStyle::Get().GetWidgetStyle<FTextBlockStyle>("AnimViewport.MessageText"))
-			]
+			SAssignNew( MainBoxPtr, SHorizontalBox )
 		]
 	];
 
@@ -247,24 +233,6 @@ TSharedRef<SWidget> SChaosClothAssetEditorRestSpaceViewportToolBar::MakeToolBar(
 	return ToolbarBuilder.MakeWidget();
 }
 
-FText SChaosClothAssetEditorRestSpaceViewportToolBar::GetDisplayString() const
-{
-	if (const FEditorModeTools* const EditorModeTools = RestSpaceViewportClient->GetModeTools())
-	{
-		if (UChaosClothAssetEditorMode* const ClothEdMode = Cast<UChaosClothAssetEditorMode>(EditorModeTools->GetActiveScriptableMode(UChaosClothAssetEditorMode::EM_ChaosClothAssetEditorModeId)))
-		{
-			if (ClothEdMode->IsMeshStatsActive())
-			{
-				const int TriangleCount = ClothEdMode->GetConstructionViewTriangleCount();
-				const int VertexCount = ClothEdMode->GetConstructionViewVertexCount();
-				const FText MeshStats = FText::Format(LOCTEXT("RestSpaceMeshStats", "Tris: {0}, Verts: {1}"), TriangleCount, VertexCount);
-				return MeshStats;
-			}
-		}
-	}
-	return FText();
-}
-
 FText SChaosClothAssetEditorRestSpaceViewportToolBar::GetViewModeMenuLabel() const
 {
 	FText Label = LOCTEXT("ConstructionViewMenuTitle_Default", "View");
diff --git a/Engine/Plugins/ChaosClothAssetEditor/Source/ChaosClothAssetEditor/Private/ChaosClothAsset/SClothEditorRestSpaceViewportToolBar.h b/Engine/Plugins/ChaosClothAssetEditor/Source/ChaosClothAssetEditor/Private/ChaosClothAsset/SClothEditorRestSpaceViewportToolBar.h
--- a/Engine/Plugins/ChaosClothAssetEditor/Source/ChaosClothAssetEditor/Private/ChaosClothAsset/SClothEditorRestSpaceViewportToolBar.h
+++ b/Engine/Plugins/ChaosClothAssetEditor/Source/ChaosClothAssetEditor/Private/ChaosClothAsset/SClothEditorRestSpaceViewportToolBar.h
@@ -36,8 +36,6 @@ private:
 	TSharedRef<SWidget> MakeDisplayToolBar(const TSharedPtr<FExtender> InExtenders);
 	TSharedRef<SWidget> MakeToolBar(const TSharedPtr<FExtender> InExtenders);
 
-	FText GetDisplayString() const;
-
 	FText GetViewModeMenuLabel() const;
 	const FSlateBrush* GetViewModeMenuLabelIcon() const;
 	TSharedRef<SWidget> GenerateViewModeMenuContent() const;
