Fix missing copyright header on ClothEditorViewportToolbarSections.cpp

diff --git a/Engine/Plugins/ChaosClothAssetEditor/Source/ChaosClothAssetEditor/Private/ChaosClothAsset/ClothEditorViewportToolbarSections.cpp b/Engine/Plugins/ChaosClothAssetEditor/Source/ChaosClothAssetEditor/Private/ChaosClothAsset/ClothEditorViewportToolbarSections.cpp
--- a/Engine/Plugins/ChaosClothAssetEditor/Source/ChaosClothAssetEditor/Private/ChaosClothAsset/ClothEditorViewportToolbarSections.cpp
+++ b/Engine/Plugins/ChaosClothAssetEditor/Source/ChaosClothAssetEditor/Private/ChaosClothAsset/ClothEditorViewportToolbarSections.cpp
@@ -1,4 +1,6 @@
-﻿#include "ClothEditorViewportToolbarSections.h"
+﻿// Copyright Epic Games, Inc. All Rights Reserved.
+
+#include "ClothEditorViewportToolbarSections.h"
 
 #include "ClothEditorRestSpaceViewportClient.h"
 #include "ChaosClothAsset/ClothEditorCommands.h"
