Objects with DataflowInstances can now be duplicated correctly, as the Owner field will now be correctly set during duplication.

Calling UpdateOwnerAsset on the DataflowInstance belonging to a Cloth Asset now works as intended. Previously it failed an assertion due to the incorrect overload of Execute being overridden in the Cloth Asset terminal node.

diff --git a/Engine/Plugins/ChaosClothAssetEditor/Source/ChaosClothAssetDataflowNodes/Public/ChaosClothAsset/TerminalNode.h b/Engine/Plugins/ChaosClothAssetEditor/Source/ChaosClothAssetDataflowNodes/Public/ChaosClothAsset/TerminalNode.h
--- a/Engine/Plugins/ChaosClothAssetEditor/Source/ChaosClothAssetDataflowNodes/Public/ChaosClothAsset/TerminalNode.h
+++ b/Engine/Plugins/ChaosClothAssetEditor/Source/ChaosClothAssetDataflowNodes/Public/ChaosClothAsset/TerminalNode.h
@@ -50,7 +50,7 @@ public:
 private:
 	//~ Begin FDataflowNode interface
 	virtual void SetAssetValue(TObjectPtr<UObject> Asset, UE::Dataflow::FContext& Context) const override;
-	virtual void Evaluate(UE::Dataflow::FContext& Context, const FDataflowOutput* Out) const override {}
+	virtual void Evaluate(UE::Dataflow::FContext& Context) const override {}
 	virtual TArray<UE::Dataflow::FPin> AddPins() override;
 	virtual bool CanAddPin() const override { return true; }
 	virtual bool CanRemovePin() const override { return CollectionLods.Num() > NumInitialCollectionLods; }
diff --git a/Engine/Source/Runtime/Experimental/Dataflow/Engine/Public/Dataflow/DataflowInstance.h b/Engine/Source/Runtime/Experimental/Dataflow/Engine/Public/Dataflow/DataflowInstance.h
--- a/Engine/Source/Runtime/Experimental/Dataflow/Engine/Public/Dataflow/DataflowInstance.h
+++ b/Engine/Source/Runtime/Experimental/Dataflow/Engine/Public/Dataflow/DataflowInstance.h
@@ -145,7 +145,8 @@ private:
 #endif
 
 	// Owning Object
-	TObjectPtr<UObject> Owner = nullptr;
+	UPROPERTY()
+	TObjectPtr<UObject> Owner;
 };
 
 UINTERFACE(MinimalAPI)
