Dataflow cloth : add interactor name for MaxDistance  in SimulationMaxDistanceConfig nodes

diff --git a/Engine/Plugins/ChaosClothAssetEditor/Source/ChaosClothAssetDataflowNodes/Public/ChaosClothAsset/SimulationMaxDistanceConfigNode.h b/Engine/Plugins/ChaosClothAssetEditor/Source/ChaosClothAssetDataflowNodes/Public/ChaosClothAsset/SimulationMaxDistanceConfigNode.h
--- a/Engine/Plugins/ChaosClothAssetEditor/Source/ChaosClothAssetDataflowNodes/Public/ChaosClothAsset/SimulationMaxDistanceConfigNode.h
+++ b/Engine/Plugins/ChaosClothAssetEditor/Source/ChaosClothAssetDataflowNodes/Public/ChaosClothAsset/SimulationMaxDistanceConfigNode.h
@@ -21,7 +21,7 @@ public:
 	 * are interpolated with the per particle weight to make the final value used for the simulation.
 	 * Otherwise all particles are considered to have a zero weight, and only the Low value is meaningful.
 	 */
-	UPROPERTY(EditAnywhere, Category = "MaxDistance Properties", Meta = (UIMin = "0", UIMax = "100", ClampMin = "0"))
+	UPROPERTY(EditAnywhere, Category = "MaxDistance Properties", Meta = (UIMin = "0", UIMax = "100", ClampMin = "0", InteractorName = "MaxDistance"))
 	FChaosClothAssetWeightedValue MaxDistance = { true, 0.f, 100.f, TEXT("MaxDistance") };
 
 	/**
