Cloth Paint Tool: fix gradient mode so that "low" vertices are selected with plain drag, and "high" vertices are selected with CTRL-drag. Also fix an issue where selected vertices don't change color when overwritten with a new selection.

diff --git a/Engine/Plugins/ChaosClothAssetEditor/Source/ChaosClothAssetEditorTools/Private/ChaosClothAsset/ClothWeightMapPaintTool.cpp b/Engine/Plugins/ChaosClothAssetEditor/Source/ChaosClothAssetEditorTools/Private/ChaosClothAsset/ClothWeightMapPaintTool.cpp
--- a/Engine/Plugins/ChaosClothAssetEditor/Source/ChaosClothAssetEditorTools/Private/ChaosClothAsset/ClothWeightMapPaintTool.cpp
+++ b/Engine/Plugins/ChaosClothAssetEditor/Source/ChaosClothAssetEditorTools/Private/ChaosClothAsset/ClothWeightMapPaintTool.cpp
@@ -1534,13 +1534,15 @@ void UClothEditorWeightMapPaintTool::OnSelectionModified()
 	{
 		const FGroupTopologySelection NewSelection = PolygonSelectionMechanic->GetActiveSelection();
 
-		const bool bSelectingLowValueGradientVertices = GetCtrlToggle();
+		const bool bSelectingLowValueGradientVertices = !GetCtrlToggle();
 		if (bSelectingLowValueGradientVertices)
 		{
+			HighValueGradientVertexSelection.Remove(NewSelection);
 			LowValueGradientVertexSelection = NewSelection;
 		}
 		else
 		{
+			LowValueGradientVertexSelection.Remove(NewSelection);
 			HighValueGradientVertexSelection = NewSelection;
 		}
 
@@ -1829,10 +1831,10 @@ void UClothEditorWeightMapPaintTool::Render(IToolsContextRenderAPI* RenderAPI)
 		const FTransform Transform = DynamicMeshComponent->GetComponentTransform();
 		GradientSelectionRenderer.SetTransform(Transform);
 			
-		GradientSelectionRenderer.SetPointParameters(FLinearColor::Red, 1.0);
+		GradientSelectionRenderer.SetPointParameters(FLinearColor::Green, 1.0);
 		PolygonSelectionMechanic->UMeshTopologySelectionMechanic::GetTopologySelector()->DrawSelection(LowValueGradientVertexSelection, &GradientSelectionRenderer, &RenderCameraState);
 
-		GradientSelectionRenderer.SetPointParameters(FLinearColor::Green, 1.0);
+		GradientSelectionRenderer.SetPointParameters(FLinearColor::Red, 1.0);
 		PolygonSelectionMechanic->UMeshTopologySelectionMechanic::GetTopologySelector()->DrawSelection(HighValueGradientVertexSelection, &GradientSelectionRenderer, &RenderCameraState);
 
 		// Now the current unsaved selection
diff --git a/Engine/Plugins/ChaosClothAssetEditor/Source/ChaosClothAssetEditorTools/Private/ChaosClothAsset/ClothWeightMapPaintTool.h b/Engine/Plugins/ChaosClothAssetEditor/Source/ChaosClothAssetEditorTools/Private/ChaosClothAsset/ClothWeightMapPaintTool.h
--- a/Engine/Plugins/ChaosClothAssetEditor/Source/ChaosClothAssetEditorTools/Private/ChaosClothAsset/ClothWeightMapPaintTool.h
+++ b/Engine/Plugins/ChaosClothAssetEditor/Source/ChaosClothAssetEditorTools/Private/ChaosClothAsset/ClothWeightMapPaintTool.h
@@ -171,15 +171,15 @@ public:
 		"SubToolType == EClothEditorWeightMapPaintInteractionType::Brush || SubToolType == EClothEditorWeightMapPaintInteractionType::Fill"))
 	double Strength = 1.0;
 
-	/** The Gradient upper limit value */
+	/** The Gradient lower limit value */
 	UPROPERTY(EditAnywhere, Category = Gradient, meta = (UIMin = 0, ClampMin = 0, UIMax = 1, ClampMax = 1,
 		HideEditConditionToggle, EditConditionHides, EditCondition = "SubToolType == EClothEditorWeightMapPaintInteractionType::Gradient"))
-	double GradientHighValue = 1.0;
+	double GradientLowValue = 0.0;
 
-	/** The Gradient lower limit value */
+	/** The Gradient upper limit value */
 	UPROPERTY(EditAnywhere, Category = Gradient, meta = (UIMin = 0, ClampMin = 0, UIMax = 1, ClampMax = 1,
 		HideEditConditionToggle, EditConditionHides, EditCondition = "SubToolType == EClothEditorWeightMapPaintInteractionType::Gradient"))
-	double GradientLowValue = 0.0;
+	double GradientHighValue = 1.0;
 
 
 	/** The Region affected by the current operation will be bounded by edge angles larger than this threshold */
