Chaos Cloth Asset - Fixed crash when visualizing patterns while in the render construction view.

#rnx

diff --git a/Engine/Plugins/ChaosClothAssetEditor/Source/ChaosClothAssetEditor/Private/ChaosClothAsset/ClothEditorMode.cpp b/Engine/Plugins/ChaosClothAssetEditor/Source/ChaosClothAssetEditor/Private/ChaosClothAsset/ClothEditorMode.cpp
--- a/Engine/Plugins/ChaosClothAssetEditor/Source/ChaosClothAssetEditor/Private/ChaosClothAsset/ClothEditorMode.cpp
+++ b/Engine/Plugins/ChaosClothAssetEditor/Source/ChaosClothAssetEditor/Private/ChaosClothAsset/ClothEditorMode.cpp
@@ -861,7 +861,7 @@ void UChaosClothAssetEditorMode::ReinitializeDynamicMeshComponents()
 		UE::Geometry::FMeshNormals::InitializeMeshToPerTriangleNormals(&LodMesh);
 	}
 
-	if (bPatternColors)
+	if (bPatternColors && ConstructionViewMode != EClothPatternVertexType::Render)
 	{
 		checkf(LodMesh.TriangleCount() == ClothFacade.GetNumSimFaces(), TEXT("Expected to have the same number of faces in LodMesh and ClothCollection sim mesh"));
 
