Chaos Cloth Asset:
- Fixed crash in render code after using the CopySimMeshToRenderMesh method or node caused by a missing root bone binding.
- Added missing binding to root bone on the sim mesh to usd importer.
- Added DeleteSelection utility method.
- Added preemptive deprecated warning on the RemoveDegenerateTriangles method (will be derprecated in a follow up change).

#rnx

diff --git a/Engine/Plugins/ChaosClothAsset/Source/ChaosClothAsset/Private/ChaosClothAsset/ClothGeometryTools.cpp b/Engine/Plugins/ChaosClothAsset/Source/ChaosClothAsset/Private/ChaosClothAsset/ClothGeometryTools.cpp
--- a/Engine/Plugins/ChaosClothAsset/Source/ChaosClothAsset/Private/ChaosClothAsset/ClothGeometryTools.cpp
+++ b/Engine/Plugins/ChaosClothAsset/Source/ChaosClothAsset/Private/ChaosClothAsset/ClothGeometryTools.cpp
@@ -505,6 +505,20 @@ namespace UE::Chaos::ClothAsset
 		}
 	}
 
+	void FClothGeometryTools::DeleteSelections(const TSharedRef<FManagedArrayCollection>& ClothCollection, const FName Group)
+	{
+		FCollectionClothSelectionFacade ClothSelectionFacade(ClothCollection);
+		const TArray<FName> Names = ClothSelectionFacade.GetNames();
+
+		for (const FName Name : Names)
+		{
+			if (Group == NAME_None || ClothSelectionFacade.GetSelectionGroup(Name) == Group)
+			{
+				ClothSelectionFacade.RemoveSelectionSet(Name);
+			}
+		}
+	}
+
 	void FClothGeometryTools::CopySimMeshToRenderMesh(const TSharedRef<FManagedArrayCollection>& ClothCollection, const FString& RenderMaterialPathName, bool bSingleRenderPattern)
 	{
 		FCollectionClothFacade ClothFacade(ClothCollection);
@@ -599,6 +613,11 @@ namespace UE::Chaos::ClothAsset
 			RenderBoneIndices[VertexIndex] = SimBoneIndices[VertexIndex3D];
 			RenderBoneWeights[VertexIndex] = SimBoneWeights[VertexIndex3D];
 		}
+
+		// Bind to root bone
+		constexpr bool bBindSimMesh = false;
+		constexpr bool bBindRenderMesh = true;
+		BindMeshToRootBone(ClothCollection, bBindSimMesh, bBindRenderMesh);
 	}
 
 	void FClothGeometryTools::ReverseMesh(
diff --git a/Engine/Plugins/ChaosClothAsset/Source/ChaosClothAsset/Public/ChaosClothAsset/ClothGeometryTools.h b/Engine/Plugins/ChaosClothAsset/Source/ChaosClothAsset/Public/ChaosClothAsset/ClothGeometryTools.h
--- a/Engine/Plugins/ChaosClothAsset/Source/ChaosClothAsset/Public/ChaosClothAsset/ClothGeometryTools.h
+++ b/Engine/Plugins/ChaosClothAsset/Source/ChaosClothAsset/Public/ChaosClothAsset/ClothGeometryTools.h
@@ -30,6 +30,9 @@ namespace UE::Chaos::ClothAsset
 		/** Remove all tethers. */
 		static void DeleteTethers(const TSharedRef<FManagedArrayCollection>& ClothCollection);
 
+		/** Remove all selections, or the selections from a specific group if specified. */
+		static void DeleteSelections(const TSharedRef<FManagedArrayCollection>& ClothCollection, const FName Group = NAME_None);
+
 		/** Turn the sim mesh portion of this ClothCollection into a render mesh. */
 		static void CopySimMeshToRenderMesh(const TSharedRef<FManagedArrayCollection>& ClothCollection, const FString& RenderMaterialPathName, bool bSingleRenderPattern);
 
diff --git a/Engine/Plugins/ChaosClothAssetEditor/Source/ChaosClothAssetDataflowNodes/Private/ChaosClothAsset/USDImportNode.cpp b/Engine/Plugins/ChaosClothAssetEditor/Source/ChaosClothAssetDataflowNodes/Private/ChaosClothAsset/USDImportNode.cpp
--- a/Engine/Plugins/ChaosClothAssetEditor/Source/ChaosClothAssetDataflowNodes/Private/ChaosClothAsset/USDImportNode.cpp
+++ b/Engine/Plugins/ChaosClothAssetEditor/Source/ChaosClothAssetDataflowNodes/Private/ChaosClothAsset/USDImportNode.cpp
@@ -3,6 +3,7 @@
 #include "ChaosClothAsset/USDImportNode.h"
 #include "ChaosClothAsset/ClothAsset.h"
 #include "ChaosClothAsset/ClothDataflowTools.h"
+#include "ChaosClothAsset/ClothGeometryTools.h"
 #include "ChaosClothAsset/CollectionClothFacade.h"
 #include "AssetRegistry/AssetData.h"
 #include "AssetRegistry/AssetRegistryModule.h"
@@ -703,6 +704,7 @@ bool FChaosClothAssetUSDImportNode::ImportFromFile(
 					TArray<FVector3f> OutDrapedPositions3D;
 					TArray<int32> OutIndices;
 
+PRAGMA_DISABLE_DEPRECATION_WARNINGS
 					bool bHasRepairedTriangles = FClothDataflowTools::RemoveDegenerateTriangles(
 						TriangleToVertexIndex,
 						RestPositions2D,
@@ -714,6 +716,7 @@ bool FChaosClothAssetUSDImportNode::ImportFromFile(
 
 					// Remove duplicate triangles
 					bHasRepairedTriangles = FClothDataflowTools::RemoveDuplicateTriangles(OutTriangleToVertexIndex) || bHasRepairedTriangles;
+PRAGMA_ENABLE_DEPRECATION_WARNINGS
 
 					// Add the new pattern
 					const int32 SimPatternIndex = ClothFacade.AddSimPattern();
@@ -879,6 +882,11 @@ bool FChaosClothAssetUSDImportNode::ImportFromFile(
 			FCollectionClothSeamFacade Seam = ClothFacade.AddGetSeam();
 			Seam.Initialize(SeamStitches[SeamIndex]);
 		}
+
+		// Bind to root bone
+		constexpr bool bBindSimMesh = true;
+		constexpr bool bBindRenderMesh = false;
+		FClothGeometryTools::BindMeshToRootBone(OutClothCollection, bBindSimMesh, bBindRenderMesh);
 	}
 	return true;
 
