Update Chaos Cloth proxy deformer to accept pattern-level selections.

diff --git a/Engine/Plugins/ChaosClothAsset/Source/ChaosClothAsset/Private/ChaosClothAsset/ClothGeometryTools.cpp b/Engine/Plugins/ChaosClothAsset/Source/ChaosClothAsset/Private/ChaosClothAsset/ClothGeometryTools.cpp
--- a/Engine/Plugins/ChaosClothAsset/Source/ChaosClothAsset/Private/ChaosClothAsset/ClothGeometryTools.cpp
+++ b/Engine/Plugins/ChaosClothAsset/Source/ChaosClothAsset/Private/ChaosClothAsset/ClothGeometryTools.cpp
@@ -1370,6 +1370,15 @@ PRAGMA_ENABLE_DEPRECATION_WARNINGS
 			}
 		};	
 
+		auto AppendElementRange = [&OutSelectionSet](const int32 StartIndex, const int32 NumElements)
+		{
+			OutSelectionSet.Reserve(OutSelectionSet.Num() + NumElements);
+			for (int32 ElemIndex = StartIndex; ElemIndex < StartIndex + NumElements; ++ElemIndex)
+			{
+				OutSelectionSet.Add(ElemIndex);
+			}
+		};
+
 		if (OrigSelectionGroup == ClothCollectionGroup::SimVertices2D)
 		{
 			if (GroupName == ClothCollectionGroup::SimFaces)
@@ -1430,6 +1439,57 @@ PRAGMA_ENABLE_DEPRECATION_WARNINGS
 				return true;
 			}
 		}
+		else if (OrigSelectionGroup == ClothCollectionGroup::SimPatterns)
+		{
+			const int32 NumSimPatterns = ClothFacade.GetNumSimPatterns();
+			if (GroupName == ClothCollectionGroup::SimFaces)
+			{
+				OutSelectionSet.Reset();
+				for (const int32 PatternIndex : OrigSelectionSet)
+				{
+					if (PatternIndex >= 0 && PatternIndex < NumSimPatterns)
+					{
+						const FCollectionClothSimPatternConstFacade PatternFacade = ClothFacade.GetSimPattern(PatternIndex);
+						AppendElementRange(PatternFacade.GetSimFacesOffset(), PatternFacade.GetNumSimFaces());
+					}
+				}
+				return true;
+			}
+			else if (GroupName == ClothCollectionGroup::SimVertices2D)
+			{
+				OutSelectionSet.Reset();
+				for (const int32 PatternIndex : OrigSelectionSet)
+				{
+					if (PatternIndex >= 0 && PatternIndex < NumSimPatterns)
+					{
+						const FCollectionClothSimPatternConstFacade PatternFacade = ClothFacade.GetSimPattern(PatternIndex);
+						AppendElementRange(PatternFacade.GetSimVertices2DOffset(), PatternFacade.GetNumSimVertices2D());
+					}
+				}
+				return true;
+			}
+			else if (GroupName == ClothCollectionGroup::SimVertices3D)
+			{
+				OutSelectionSet.Reset();
+				for (const int32 PatternIndex : OrigSelectionSet)
+				{
+					if (PatternIndex >= 0 && PatternIndex < NumSimPatterns)
+					{
+						const FCollectionClothSimPatternConstFacade PatternFacade = ClothFacade.GetSimPattern(PatternIndex);
+						OutSelectionSet.Reserve(OutSelectionSet.Num() + PatternFacade.GetNumSimVertices2D());
+
+						const TConstArrayView<FIntVector> SimIndices3D = PatternFacade.GetSimIndices3D();
+						for (const FIntVector& Indices : SimIndices3D)
+						{
+							OutSelectionSet.Add(Indices[0]);
+							OutSelectionSet.Add(Indices[1]);
+							OutSelectionSet.Add(Indices[2]);
+						}
+					}
+				}
+				return true;
+			}
+		}
 		else if (OrigSelectionGroup == ClothCollectionGroup::RenderVertices)
 		{
 			if (GroupName == ClothCollectionGroup::RenderFaces)
@@ -1446,6 +1506,36 @@ PRAGMA_ENABLE_DEPRECATION_WARNINGS
 				return true;
 			}
 		}
+		else if (OrigSelectionGroup == ClothCollectionGroup::RenderPatterns)
+		{
+			const int32 NumRenderPatterns = ClothFacade.GetNumRenderPatterns();
+			if (GroupName == ClothCollectionGroup::RenderFaces)
+			{
+				OutSelectionSet.Reset();
+				for (const int32 PatternIndex : OrigSelectionSet)
+				{
+					if (PatternIndex >= 0 && PatternIndex < NumRenderPatterns)
+					{
+						const FCollectionClothRenderPatternConstFacade PatternFacade = ClothFacade.GetRenderPattern(PatternIndex);
+						AppendElementRange(PatternFacade.GetRenderFacesOffset(), PatternFacade.GetNumRenderFaces());
+					}
+				}
+				return true;
+			}
+			else if (GroupName == ClothCollectionGroup::RenderVertices)
+			{
+				OutSelectionSet.Reset();
+				for (const int32 PatternIndex : OrigSelectionSet)
+				{
+					if (PatternIndex >= 0 && PatternIndex < NumRenderPatterns)
+					{
+						const FCollectionClothRenderPatternConstFacade PatternFacade = ClothFacade.GetRenderPattern(PatternIndex);
+						AppendElementRange(PatternFacade.GetRenderVerticesOffset(), PatternFacade.GetNumRenderVertices());
+					}
+				}
+				return true;
+			}
+		}
 
 		return false;
 	}
diff --git a/Engine/Plugins/ChaosClothAssetEditor/Source/ChaosClothAssetDataflowNodes/Private/ChaosClothAsset/ProxyDeformerNode.cpp b/Engine/Plugins/ChaosClothAssetEditor/Source/ChaosClothAssetDataflowNodes/Private/ChaosClothAsset/ProxyDeformerNode.cpp
--- a/Engine/Plugins/ChaosClothAssetEditor/Source/ChaosClothAssetDataflowNodes/Private/ChaosClothAsset/ProxyDeformerNode.cpp
+++ b/Engine/Plugins/ChaosClothAssetEditor/Source/ChaosClothAssetDataflowNodes/Private/ChaosClothAsset/ProxyDeformerNode.cpp
@@ -4,6 +4,7 @@
 #include "ChaosClothAsset/ClothCollectionAttribute.h"
 #include "ChaosClothAsset/ClothCollectionGroup.h"
 #include "ChaosClothAsset/ClothDataflowTools.h"
+#include "ChaosClothAsset/ClothGeometryTools.h"
 #include "ChaosClothAsset/CollectionClothFacade.h"
 #include "ChaosClothAsset/CollectionClothSelectionFacade.h"
 #include "Dataflow/DataflowInputOutput.h"
@@ -129,153 +130,45 @@ namespace UE::Chaos::ClothAsset::Private
 		}
 	};
 
-	static FPointWeightMap SelectionToPointWeightMap(const FCollectionClothConstFacade& ClothFacade, const FCollectionClothSelectionConstFacade& SelectionFacade, const FName& SelectionName)
+	static FPointWeightMap SelectionToPointWeightMap(const TSharedRef<const FManagedArrayCollection>& ClothCollection, const FCollectionClothConstFacade& ClothFacade, const FName& SelectionName)
 	{
-		constexpr float SelectedValue = 1.f;
-		if (const TSet<int32>* SelectionSet = SelectionFacade.IsValid() ? SelectionFacade.FindSelectionSet(SelectionName) : nullptr)
+		constexpr float SelectedValue = 1.f; 
+		
+		TSet<int32> SimVertex3DSelectionSet;
+		if (FClothGeometryTools::ConvertSelectionToNewGroupType(ClothCollection, SelectionName, ClothCollectionGroup::SimVertices3D, SimVertex3DSelectionSet))
 		{
 			FPointWeightMap PointWeightMap;
-
-			const FName SelectionGroup = SelectionFacade.GetSelectionGroup(SelectionName);
-
-			if (SelectionGroup == ClothCollectionGroup::SimVertices3D)
+			PointWeightMap.Initialize(ClothFacade.GetNumSimVertices3D());  // Init to zero (unselected)
+			for (const int32 VertexIndex : SimVertex3DSelectionSet)
 			{
-				PointWeightMap.Initialize(ClothFacade.GetNumSimVertices3D());  // Init to zero (unselected)
-				for (const int32 VertexIndex : *SelectionSet)
-				{
-					PointWeightMap[VertexIndex] = SelectedValue;
-				}
-				return PointWeightMap;
-			}
-			else if (SelectionGroup == ClothCollectionGroup::SimVertices2D)
-			{
-				PointWeightMap.Initialize(ClothFacade.GetNumSimVertices3D());  // Init to zero (unselected)
-				const TConstArrayView<int32> Vertex2DTo3D = ClothFacade.GetSimVertex3DLookup();
-				for (const int32 VertexIndex : *SelectionSet)
-				{
-					PointWeightMap[Vertex2DTo3D[VertexIndex]] = SelectedValue;
-				}
-				return PointWeightMap;
-			}
-			else if (SelectionGroup == ClothCollectionGroup::SimFaces)
-			{
-				PointWeightMap.Initialize(ClothFacade.GetNumSimVertices3D());  // Init to zero (unselected)
-				const TConstArrayView<FIntVector3> SimIndices3D = ClothFacade.GetSimIndices3D();
-				for (const int32 FaceIndex : *SelectionSet)
-				{
-					PointWeightMap[SimIndices3D[FaceIndex][0]] = SelectedValue;
-					PointWeightMap[SimIndices3D[FaceIndex][1]] = SelectedValue;
-					PointWeightMap[SimIndices3D[FaceIndex][2]] = SelectedValue;
-				}
-				return PointWeightMap;
+				PointWeightMap[VertexIndex] = SelectedValue;
 			}
+			return PointWeightMap;
 		}
 		// Invalid or no selection found, all points are dynamic 
 		return FPointWeightMap(ClothFacade.GetNumSimVertices3D(), SelectedValue);
 	}
 
-	static TArray<ClothingMeshUtils::FMeshToMeshFilterSet> SelectionsToMeshToMeshFilterSets_v2(const FCollectionClothConstFacade& ClothFacade, const FCollectionClothSelectionConstFacade& SelectionFacade, const TArray<TPair<FName, FName>>& SelectionNames)
+	static TArray<ClothingMeshUtils::FMeshToMeshFilterSet> SelectionsToMeshToMeshFilterSets_v2(const TSharedRef<const FManagedArrayCollection>& ClothCollection, const TArray<TPair<FName, FName>>& SelectionNames)
 	{
-		auto GetSimFaceSelection = [&ClothFacade](const FName& SelectionGroup, const TSet<int32>& SelectionSet) -> TSet<int32>
-			{
-				TSet<int32> SimFaceSelection;
-				if (SelectionGroup == ClothCollectionGroup::SimVertices2D)
-				{
-					const TConstArrayView<FIntVector3> SimIndices2D = ClothFacade.GetSimIndices2D();
-					SimFaceSelection.Reserve(SelectionSet.Num());
-
-					for (int32 FaceIndex = 0; FaceIndex < SimIndices2D.Num(); ++FaceIndex)
-					{
-						const FIntVector3& Indices = SimIndices2D[FaceIndex];
-
-						if (SelectionSet.Contains(Indices[0]) &&
-							SelectionSet.Contains(Indices[1]) &&
-							SelectionSet.Contains(Indices[2]))
-						{
-							SimFaceSelection.Add(FaceIndex);
-						}
-					}
-				}
-				else if (SelectionGroup == ClothCollectionGroup::SimVertices3D)
-				{
-					const TConstArrayView<FIntVector3> SimIndices3D = ClothFacade.GetSimIndices3D();
-					SimFaceSelection.Reserve(SelectionSet.Num());
-
-					for (int32 FaceIndex = 0; FaceIndex < SimIndices3D.Num(); ++FaceIndex)
-					{
-						const FIntVector3& Indices = SimIndices3D[FaceIndex];
-
-						if (SelectionSet.Contains(Indices[0]) &&
-							SelectionSet.Contains(Indices[1]) &&
-							SelectionSet.Contains(Indices[2]))
-						{
-							SimFaceSelection.Add(FaceIndex);
-						}
-					}
-				}
-				else if (SelectionGroup == ClothCollectionGroup::SimFaces)
-				{
-					SimFaceSelection = SelectionSet;
-				}
-				return SimFaceSelection;
-			};
-
-		auto GetRenderVertexSelection = [&ClothFacade](const FName& SelectionGroup, const TSet<int32>& SelectionSet) -> TSet<int32>
-			{
-				TSet<int32> RenderVertexSelection;
-				if (SelectionGroup == ClothCollectionGroup::RenderVertices)
-				{
-					RenderVertexSelection = SelectionSet;
-				}
-				else if (SelectionGroup == ClothCollectionGroup::RenderFaces)
-				{
-					RenderVertexSelection.Reserve(SelectionSet.Num() * 3);
-					const TConstArrayView<FIntVector3> RenderIndices = ClothFacade.GetRenderIndices();
-					for (const int32 FaceIndex : SelectionSet)
-					{
-						RenderVertexSelection.Add(RenderIndices[FaceIndex][0]);
-						RenderVertexSelection.Add(RenderIndices[FaceIndex][1]);
-						RenderVertexSelection.Add(RenderIndices[FaceIndex][2]);
-					}
-				}
-				return RenderVertexSelection;
-			};
-
 		// Fill up the MeshToMeshFilterSets
 		TArray<ClothingMeshUtils::FMeshToMeshFilterSet> MeshToMeshFilterSets;
-
-		if (SelectionFacade.IsValid())
+		MeshToMeshFilterSets.Reserve(SelectionNames.Num());
+		for (const TPair<FName, FName>& SelectionName : SelectionNames)
 		{
-			MeshToMeshFilterSets.Reserve(SelectionNames.Num());
-
-			for (const TPair<FName,FName>& SelectionName : SelectionNames)
+			TSet<int32> SimFaceSelection;
+			TSet<int32> RenderVertexSelection;
+			if (FClothGeometryTools::ConvertSelectionToNewGroupType(ClothCollection, SelectionName.Get<0>(), ClothCollectionGroup::RenderVertices, RenderVertexSelection)
+				&& FClothGeometryTools::ConvertSelectionToNewGroupType(ClothCollection, SelectionName.Get<1>(), ClothCollectionGroup::SimFaces, SimFaceSelection))
 			{
-				if (const TSet<int32>* RenderSelectionSet = SelectionFacade.FindSelectionSet(SelectionName.Get<0>()))
+				if (!SimFaceSelection.Num() && !RenderVertexSelection.Num())
 				{
-					if (const TSet<int32>* SimSelectionSet = SelectionFacade.FindSelectionSet(SelectionName.Get<1>()))
-					{
-						if (RenderSelectionSet->Num() || SimSelectionSet->Num())
-						{
-							const FName RenderSelectionGroup = SelectionFacade.GetSelectionGroup(SelectionName.Get<0>());
-							const FName SimSelectionGroup = SelectionFacade.GetSelectionGroup(SelectionName.Get<1>());
-
-							// Retrieve the sim face selection
-							TSet<int32> SimFaceSelection = GetSimFaceSelection(SimSelectionGroup, *SimSelectionSet);
-
-							// Retrieve the render vertex selection
-							TSet<int32> RenderVertexSelection = GetRenderVertexSelection(RenderSelectionGroup, *RenderSelectionSet);
-
-							if (!SimFaceSelection.Num() && !RenderVertexSelection.Num())
-							{
-								continue;  // Nothing selected
-							}
-
-							ClothingMeshUtils::FMeshToMeshFilterSet& MeshToMeshFilterSet = MeshToMeshFilterSets.AddDefaulted_GetRef();
-							MeshToMeshFilterSet.SourceTriangles = MoveTemp(SimFaceSelection);
-							MeshToMeshFilterSet.TargetVertices = MoveTemp(RenderVertexSelection);
-						}
-					}
+					continue;  // Nothing selected
 				}
+
+				ClothingMeshUtils::FMeshToMeshFilterSet& MeshToMeshFilterSet = MeshToMeshFilterSets.AddDefaulted_GetRef();
+				MeshToMeshFilterSet.SourceTriangles = MoveTemp(SimFaceSelection);
+				MeshToMeshFilterSet.TargetVertices = MoveTemp(RenderVertexSelection);
 			}
 		}
 
@@ -283,122 +176,48 @@ namespace UE::Chaos::ClothAsset::Private
 	}
 	
 	
-	static TArray<ClothingMeshUtils::FMeshToMeshFilterSet> SelectionsToMeshToMeshFilterSets(const FCollectionClothConstFacade& ClothFacade, const FCollectionClothSelectionConstFacade& SelectionFacade, const TArray<FName>& SelectionNames)
+	static TArray<ClothingMeshUtils::FMeshToMeshFilterSet> SelectionsToMeshToMeshFilterSets(const TSharedRef<const FManagedArrayCollection>& ClothCollection, const TArray<FName>& SelectionNames)
 	{
-		auto GetSimFaceSelection = [&ClothFacade](const FName& SelectionGroup, const TSet<int32>& SelectionSet) -> TSet<int32>
-		{
-			TSet<int32> SimFaceSelection;
-			if (SelectionGroup == ClothCollectionGroup::SimVertices2D)
-			{
-				const TConstArrayView<FIntVector3> SimIndices2D = ClothFacade.GetSimIndices2D();
-				SimFaceSelection.Reserve(SelectionSet.Num());
-
-				for (int32 FaceIndex = 0; FaceIndex < SimIndices2D.Num(); ++FaceIndex)
-				{
-					const FIntVector3& Indices = SimIndices2D[FaceIndex];
-
-					if (SelectionSet.Contains(Indices[0]) &&
-						SelectionSet.Contains(Indices[1]) &&
-						SelectionSet.Contains(Indices[2]))
-					{
-						SimFaceSelection.Add(FaceIndex);
-					}
-				}
-			}
-			else if (SelectionGroup == ClothCollectionGroup::SimVertices3D)
-			{
-				const TConstArrayView<FIntVector3> SimIndices3D = ClothFacade.GetSimIndices3D();
-				SimFaceSelection.Reserve(SelectionSet.Num());
-
-				for (int32 FaceIndex = 0; FaceIndex < SimIndices3D.Num(); ++FaceIndex)
-				{
-					const FIntVector3& Indices = SimIndices3D[FaceIndex];
-
-					if (SelectionSet.Contains(Indices[0]) &&
-						SelectionSet.Contains(Indices[1]) &&
-						SelectionSet.Contains(Indices[2]))
-					{
-						SimFaceSelection.Add(FaceIndex);
-					}
-				}
-			}
-			else if (SelectionGroup == ClothCollectionGroup::SimFaces)
-			{
-				SimFaceSelection = SelectionSet;
-			}
-			return SimFaceSelection;
-		};
-
-		auto GetRenderVertexSelection = [&ClothFacade](const FName& SelectionGroup, const TSet<int32>& SelectionSet) -> TSet<int32>
-		{
-			TSet<int32> RenderVertexSelection;
-			if (SelectionGroup == ClothCollectionGroup::RenderVertices)
-			{
-				RenderVertexSelection = SelectionSet;
-			}
-			else if (SelectionGroup == ClothCollectionGroup::RenderFaces)
-			{
-				RenderVertexSelection.Reserve(SelectionSet.Num() * 3);
-				const TConstArrayView<FIntVector3> RenderIndices = ClothFacade.GetRenderIndices();
-				for (const int32 FaceIndex : SelectionSet)
-				{
-					RenderVertexSelection.Add(RenderIndices[FaceIndex][0]);
-					RenderVertexSelection.Add(RenderIndices[FaceIndex][1]);
-					RenderVertexSelection.Add(RenderIndices[FaceIndex][2]);
-				}
-			}
-			return RenderVertexSelection;
-		};
-
 		// Fill up the MeshToMeshFilterSets
 		TArray<ClothingMeshUtils::FMeshToMeshFilterSet> MeshToMeshFilterSets;
-
-		if (SelectionFacade.IsValid())
 		{
 			MeshToMeshFilterSets.Reserve(SelectionNames.Num());
 
 			for (const FName& SelectionName : SelectionNames)
 			{
-				if (const TSet<int32>* SelectionSet = SelectionFacade.FindSelectionSet(SelectionName))
+				TSet<int32> SimFaceSelection;
+				bool bRenderVertexSelectionIsSecondary = true;
+				if (!FClothGeometryTools::ConvertSelectionToNewGroupType(ClothCollection, SelectionName, ClothCollectionGroup::SimFaces, SimFaceSelection))
 				{
+					// Try using secondary selection
+					bRenderVertexSelectionIsSecondary = false;
 PRAGMA_DISABLE_DEPRECATION_WARNINGS
-					if (const TSet<int32>* SecondarySelectionSet = SelectionFacade.FindSelectionSecondarySet(SelectionName))
+					if(!FClothGeometryTools::ConvertSelectionToNewGroupType(ClothCollection, SelectionName, ClothCollectionGroup::SimFaces, true, SimFaceSelection))
 PRAGMA_ENABLE_DEPRECATION_WARNINGS
 					{
-						if (SelectionSet->Num() && SecondarySelectionSet->Num())
-						{
-							FName SelectionGroup = SelectionFacade.GetSelectionGroup(SelectionName);
+						// Failed to get SimFaces set
+						continue;
+					}
+				}
+
+				if (SimFaceSelection.IsEmpty())
+				{
+					continue;  // Nothing selected
+				}
+
+				TSet<int32> RenderVertexSelection;
 PRAGMA_DISABLE_DEPRECATION_WARNINGS
-							FName SelectionSecondaryGroup = SelectionFacade.GetSelectionSecondaryGroup(SelectionName);
+				FClothGeometryTools::ConvertSelectionToNewGroupType(ClothCollection, SelectionName, ClothCollectionGroup::RenderVertices, bRenderVertexSelectionIsSecondary, RenderVertexSelection);
 PRAGMA_ENABLE_DEPRECATION_WARNINGS
-							// Retrieve the sim face selection
-							TSet<int32> SimFaceSelection = GetSimFaceSelection(SelectionGroup, *SelectionSet);
-							if (!SimFaceSelection.Num())
-							{
-								// Try swapping the selections
-								Swap(SelectionSet, SecondarySelectionSet);
-								Swap(SelectionGroup, SelectionSecondaryGroup);
-
-								SimFaceSelection = GetSimFaceSelection(SelectionGroup, *SelectionSet);
-							}
-							if (!SimFaceSelection.Num())
-							{
-								continue;  // Nothing selected on the simulation side
-							}
-
-							// Retrieve the render vertex selection
-							TSet<int32> RenderVertexSelection = GetRenderVertexSelection(SelectionSecondaryGroup, *SecondarySelectionSet);
-							if (!RenderVertexSelection.Num())
-							{
-								continue;  // Nothing selected on the render side
-							}
-
-							ClothingMeshUtils::FMeshToMeshFilterSet& MeshToMeshFilterSet = MeshToMeshFilterSets.AddDefaulted_GetRef();
-							MeshToMeshFilterSet.SourceTriangles = MoveTemp(SimFaceSelection);
-							MeshToMeshFilterSet.TargetVertices = MoveTemp(RenderVertexSelection);
-						}
-					}
+
+				if (RenderVertexSelection.IsEmpty())
+				{
+					continue;  // Nothing selected
 				}
+
+				ClothingMeshUtils::FMeshToMeshFilterSet& MeshToMeshFilterSet = MeshToMeshFilterSets.AddDefaulted_GetRef();
+				MeshToMeshFilterSet.SourceTriangles = MoveTemp(SimFaceSelection);
+				MeshToMeshFilterSet.TargetVertices = MoveTemp(RenderVertexSelection);
 			}
 		}
 
@@ -453,7 +272,7 @@ void FChaosClothAssetProxyDeformerNode_v2::Evaluate(UE::Dataflow::FContext& Cont
 			DeformerMappingDataGenerator.RenderNormals = ClothFacade.GetRenderNormal();
 			DeformerMappingDataGenerator.RenderIndices = ClothFacade.GetRenderIndices();
 			DeformerMappingDataGenerator.PointWeightMap = FPointWeightMap(); // _v2 no longer compute SkinningBlend transitions
-			DeformerMappingDataGenerator.MeshToMeshFilterSet = Private::SelectionsToMeshToMeshFilterSets_v2(ClothFacade, SelectionFacade, GetSelectionFilterNames(Context));
+			DeformerMappingDataGenerator.MeshToMeshFilterSet = Private::SelectionsToMeshToMeshFilterSets_v2(ClothCollection, GetSelectionFilterNames(Context));
 			DeformerMappingDataGenerator.RenderDeformerPositionBaryCoordsAndDist = ClothFacade.GetRenderDeformerPositionBaryCoordsAndDist();
 			DeformerMappingDataGenerator.RenderDeformerNormalBaryCoordsAndDist = ClothFacade.GetRenderDeformerNormalBaryCoordsAndDist();
 			DeformerMappingDataGenerator.RenderDeformerTangentBaryCoordsAndDist = ClothFacade.GetRenderDeformerTangentBaryCoordsAndDist();
@@ -659,8 +478,8 @@ void FChaosClothAssetProxyDeformerNode::Evaluate(UE::Dataflow::FContext& Context
 			DeformerMappingDataGenerator.RenderPositions = ClothFacade.GetRenderPosition();
 			DeformerMappingDataGenerator.RenderNormals = ClothFacade.GetRenderNormal();
 			DeformerMappingDataGenerator.RenderIndices = ClothFacade.GetRenderIndices();
-			DeformerMappingDataGenerator.PointWeightMap = Private::SelectionToPointWeightMap(ClothFacade, SelectionFacade, SimVertexSelectionName);
-			DeformerMappingDataGenerator.MeshToMeshFilterSet = Private::SelectionsToMeshToMeshFilterSets(ClothFacade, SelectionFacade, GetSelectionFilterNames(Context));
+			DeformerMappingDataGenerator.PointWeightMap = Private::SelectionToPointWeightMap(ClothCollection, ClothFacade, SimVertexSelectionName);
+			DeformerMappingDataGenerator.MeshToMeshFilterSet = Private::SelectionsToMeshToMeshFilterSets(ClothCollection, GetSelectionFilterNames(Context));
 			DeformerMappingDataGenerator.RenderDeformerPositionBaryCoordsAndDist = ClothFacade.GetRenderDeformerPositionBaryCoordsAndDist();
 			DeformerMappingDataGenerator.RenderDeformerNormalBaryCoordsAndDist = ClothFacade.GetRenderDeformerNormalBaryCoordsAndDist();
 			DeformerMappingDataGenerator.RenderDeformerTangentBaryCoordsAndDist = ClothFacade.GetRenderDeformerTangentBaryCoordsAndDist();
