Dataflow, UpdateClothFromDynamicMeshNode: enable setting cloth render mesh UVs from DynamicMesh UVs

diff --git a/Engine/Plugins/ChaosClothAssetEditor/Source/ChaosClothAssetDataflowNodes/Private/ChaosClothAsset/ClothCollectionToDynamicMeshNode.cpp b/Engine/Plugins/ChaosClothAssetEditor/Source/ChaosClothAssetDataflowNodes/Private/ChaosClothAsset/ClothCollectionToDynamicMeshNode.cpp
--- a/Engine/Plugins/ChaosClothAssetEditor/Source/ChaosClothAssetDataflowNodes/Private/ChaosClothAsset/ClothCollectionToDynamicMeshNode.cpp
+++ b/Engine/Plugins/ChaosClothAssetEditor/Source/ChaosClothAssetDataflowNodes/Private/ChaosClothAsset/ClothCollectionToDynamicMeshNode.cpp
@@ -113,6 +113,39 @@ void FChaosClothAssetUpdateClothFromDynamicMeshNode::Evaluate(UE::Dataflow::FCon
 						}
 					}
 				}
+				if (bCopyUVsToRenderUVs)
+				{
+					if (const UE::Geometry::FDynamicMeshAttributeSet* const AttributeSet = DynMesh.Attributes())
+					{
+						if (const UE::Geometry::FDynamicMeshUVOverlay* const UVOverlay = AttributeSet->GetUVLayer(UVChannelIndex))
+						{
+							TArrayView<TArray<FVector2f>> RenderUVs = ClothFacade.GetRenderUVs();
+							TConstArrayView<FIntVector> RenderTriangles = ClothFacade.GetRenderIndices();
+							for (int32 VertexID : DynMesh.VertexIndicesItr())
+							{
+								const int32 RenderMeshVertexID = ClothMapping.GetOriginalVertexID(VertexID);
+								if (RenderUVs.IsValidIndex(RenderMeshVertexID) && RenderUVs[RenderMeshVertexID].Num() > UVChannelIndex)
+								{
+									// Get the UV values corresponding to VertexID in the dynamic mesh
+									UVOverlay->EnumerateVertexElements(VertexID,
+										[&RenderUVs, &ClothMapping, &RenderTriangles, RenderMeshVertexID, UVChannelIndex=this->UVChannelIndex](int32 TriangleID, int32 ElementID, const FVector2f& UVValue)->bool
+										{
+											const int32 RenderMeshTriID = ClothMapping.GetOriginalTriangleID(TriangleID);
+											const FIntVector& RenderTri = RenderTriangles[RenderMeshTriID];
+											for (int32 LocalIndex = 0; LocalIndex < 3; ++LocalIndex)
+											{
+												if (RenderTri[LocalIndex] == RenderMeshVertexID)
+												{ 
+													RenderUVs[RenderTri[LocalIndex]][UVChannelIndex] = UVValue;
+												}
+											}
+											return true;
+										});
+								}
+							}
+						}
+					}
+				}
 				if (bCopyToSim3DPositions)
 				{
 					TArrayView<FVector3f> SimPositions = ClothFacade.GetSimPosition3D();
diff --git a/Engine/Plugins/ChaosClothAssetEditor/Source/ChaosClothAssetDataflowNodes/Public/ChaosClothAsset/ClothCollectionToDynamicMeshNode.h b/Engine/Plugins/ChaosClothAssetEditor/Source/ChaosClothAssetDataflowNodes/Public/ChaosClothAsset/ClothCollectionToDynamicMeshNode.h
--- a/Engine/Plugins/ChaosClothAssetEditor/Source/ChaosClothAssetDataflowNodes/Public/ChaosClothAsset/ClothCollectionToDynamicMeshNode.h
+++ b/Engine/Plugins/ChaosClothAssetEditor/Source/ChaosClothAssetDataflowNodes/Public/ChaosClothAsset/ClothCollectionToDynamicMeshNode.h
@@ -77,6 +77,10 @@ private:
 	UPROPERTY(EditAnywhere, Category = "DynamicMesh To Cloth Collection")
 	bool bCopyToRenderPositions = false;
 
+	/** Copy DynamicMesh UVs to Render UVs */
+	UPROPERTY(EditAnywhere, Category = "DynamicMesh To Cloth Collection")
+	bool bCopyUVsToRenderUVs = false;
+
 	/** Copy DynamicMesh Vertex Positions to Sim3D Positions */
 	UPROPERTY(EditAnywhere, Category = "DynamicMesh To Cloth Collection")
 	bool bCopyToSim3DPositions = false;
@@ -85,8 +89,8 @@ private:
 	UPROPERTY(EditAnywhere, Category = "DynamicMesh To Cloth Collection")
 	bool bCopyUVsToSim2DPositions = false;
 
-	/** Which UV Channel to use at Sim2D Positions */
-	UPROPERTY(EditAnywhere, Category = "DynamicMesh To Cloth Collection", meta = (EditCondition = "bCopyUVsToSim2DPositions", EditConditionHides))
+	/** Which UV Channel to use at Sim2D Positions or Render UVs */
+	UPROPERTY(EditAnywhere, Category = "DynamicMesh To Cloth Collection", meta = (EditCondition = "bCopyUVsToSim2DPositions || bCopyUVsToRenderUVs", EditConditionHides))
 	int32 UVChannelIndex = 0;
 };
 
