SClothAnimationScrubPanel: guard against null pointer dereference

diff --git a/Engine/Plugins/ChaosClothAssetEditor/Source/ChaosClothAssetEditor/Private/ChaosClothAsset/SClothAnimationScrubPanel.cpp b/Engine/Plugins/ChaosClothAssetEditor/Source/ChaosClothAssetEditor/Private/ChaosClothAsset/SClothAnimationScrubPanel.cpp
--- a/Engine/Plugins/ChaosClothAssetEditor/Source/ChaosClothAssetEditor/Private/ChaosClothAsset/SClothAnimationScrubPanel.cpp
+++ b/Engine/Plugins/ChaosClothAssetEditor/Source/ChaosClothAssetEditor/Private/ChaosClothAsset/SClothAnimationScrubPanel.cpp
@@ -248,17 +248,20 @@ FReply SClothAnimationScrubPanel::OnClick_PreviewPlaybackMode()
 		
 		if (UAnimSingleNodeInstance* const PreviewInstance = GetPreviewAnimationInstance())
 		{
-			// If we paused due to hitting the end point, start playing again when entering loop mode
-			const float CurrentTime = PreviewInstance->GetCurrentTime();
-			const float PlayRate = PreviewInstance->GetPlayRate();
-			const float AssetPlayLength = PreviewInstance->CurrentAsset->GetPlayLength();
-			if (PlayRate < 0.0 && CurrentTime <= 0.0)
-			{
-				PreviewInstance->SetPlaying(true);
-			}
-			else if (PlayRate > 0.0 && CurrentTime >= AssetPlayLength)
+			if (PreviewInstance->CurrentAsset)
 			{
-				PreviewInstance->SetPlaying(true);
+				// If we paused due to hitting the end point, start playing again when entering loop mode
+				const float CurrentTime = PreviewInstance->GetCurrentTime();
+				const float PlayRate = PreviewInstance->GetPlayRate();
+				const float AssetPlayLength = PreviewInstance->CurrentAsset->GetPlayLength();
+				if (PlayRate < 0.0 && CurrentTime <= 0.0)
+				{
+					PreviewInstance->SetPlaying(true);
+				}
+				else if (PlayRate > 0.0 && CurrentTime >= AssetPlayLength)
+				{
+					PreviewInstance->SetPlaying(true);
+				}
 			}
 		}
 	}
@@ -313,19 +316,22 @@ void SClothAnimationScrubPanel::OnTickPlayback(double InCurrentTime, float InDel
 {
 	if (UAnimSingleNodeInstance* const PreviewInstance = GetPreviewAnimationInstance())
 	{
-		const float CurrentTime = PreviewInstance->GetCurrentTime();
-		const float PlayRate = PreviewInstance->GetPlayRate();
-		const float AssetPlayLength = PreviewInstance->CurrentAsset->GetPlayLength();
-
-		if (PreviewPlaybackMode == EClothPreviewPlaybackMode::PingPong)
+		if (PreviewInstance->CurrentAsset)
 		{
-			if (PlayRate < 0.0 && CurrentTime <= 0.0)
-			{
-				PreviewInstance->SetReverse(!PreviewInstance->IsReverse());
-			}
-			else if (PlayRate > 0.0 && CurrentTime >= AssetPlayLength)
+			const float CurrentTime = PreviewInstance->GetCurrentTime();
+			const float PlayRate = PreviewInstance->GetPlayRate();
+			const float AssetPlayLength = PreviewInstance->CurrentAsset->GetPlayLength();
+
+			if (PreviewPlaybackMode == EClothPreviewPlaybackMode::PingPong)
 			{
-				PreviewInstance->SetReverse(!PreviewInstance->IsReverse());
+				if (PlayRate < 0.0 && CurrentTime <= 0.0)
+				{
+					PreviewInstance->SetReverse(!PreviewInstance->IsReverse());
+				}
+				else if (PlayRate > 0.0 && CurrentTime >= AssetPlayLength)
+				{
+					PreviewInstance->SetReverse(!PreviewInstance->IsReverse());
+				}
 			}
 		}
 	}
