Fix logging errors from new cloth dataflow nodes due to uninitialized properties.

diff --git a/Engine/Plugins/ChaosClothAssetEditor/Source/ChaosClothAssetDataflowNodes/Public/ChaosClothAsset/ClothCollectionQueryNode.h b/Engine/Plugins/ChaosClothAssetEditor/Source/ChaosClothAssetDataflowNodes/Public/ChaosClothAsset/ClothCollectionQueryNode.h
--- a/Engine/Plugins/ChaosClothAssetEditor/Source/ChaosClothAssetDataflowNodes/Public/ChaosClothAsset/ClothCollectionQueryNode.h
+++ b/Engine/Plugins/ChaosClothAssetEditor/Source/ChaosClothAssetDataflowNodes/Public/ChaosClothAsset/ClothCollectionQueryNode.h
@@ -33,17 +33,17 @@ private:
 
 	/** Is this collection a valid cloth collection */
 	UPROPERTY(Meta = (DataflowOutput))
-	bool bIsClothCollection;
+	bool bIsClothCollection = false;
 
 	/** Does this collection have a cloth sim mesh */
 	UPROPERTY(Meta = (DataflowOutput))
-	bool bHasClothSimMesh;
+	bool bHasClothSimMesh = false;
 
 	/** Does this collection have a cloth render mesh */
 	UPROPERTY(Meta = (DataflowOutput))
-	bool bHasClothRenderMesh;
+	bool bHasClothRenderMesh = false;
 
 	/** Does this collection have proxy deformer data */
 	UPROPERTY(Meta = (DataflowOutput))
-	bool bHasClothProxyDeformer;
+	bool bHasClothProxyDeformer = false;
 };
diff --git a/Engine/Plugins/ChaosClothAssetEditor/Source/ChaosClothAssetDataflowNodes/Public/ChaosClothAsset/ProceduralSelectionNode.h b/Engine/Plugins/ChaosClothAssetEditor/Source/ChaosClothAssetDataflowNodes/Public/ChaosClothAsset/ProceduralSelectionNode.h
--- a/Engine/Plugins/ChaosClothAssetEditor/Source/ChaosClothAssetDataflowNodes/Public/ChaosClothAsset/ProceduralSelectionNode.h
+++ b/Engine/Plugins/ChaosClothAssetEditor/Source/ChaosClothAssetDataflowNodes/Public/ChaosClothAsset/ProceduralSelectionNode.h
@@ -42,7 +42,7 @@ private:
 
 	/** The procedural selection method */
 	UPROPERTY(EditAnywhere, Category = "Procedural Selection")
-	EChaosClothAssetProceduralSelectionType SelectionType;
+	EChaosClothAssetProceduralSelectionType SelectionType = EChaosClothAssetProceduralSelectionType::SelectAll;
 
 	/** Selection set to convert from when using Conversion selection type. */
 	UPROPERTY(EditAnywhere, Category = "Procedural Selection", Meta = (Dataflowinput, EditCondition = "SelectionType == EChaosClothAssetProceduralSelectionType::Conversion"))
