Chaos Cloth: Go back to old viewport toolbar until we can figure out how to bring sim stats back.

p.ChaosCloth.UseOld3DToolbar can be used to toggle between the new and old toolbars. Must reopen cloth editor to pick up the change.

diff --git a/Engine/Plugins/ChaosClothAssetEditor/Source/ChaosClothAssetEditor/Private/ChaosClothAsset/SClothEditor3DViewport.cpp b/Engine/Plugins/ChaosClothAssetEditor/Source/ChaosClothAssetEditor/Private/ChaosClothAsset/SClothEditor3DViewport.cpp
--- a/Engine/Plugins/ChaosClothAssetEditor/Source/ChaosClothAssetEditor/Private/ChaosClothAsset/SClothEditor3DViewport.cpp
+++ b/Engine/Plugins/ChaosClothAssetEditor/Source/ChaosClothAssetEditor/Private/ChaosClothAsset/SClothEditor3DViewport.cpp
@@ -20,6 +20,12 @@
 
 #define LOCTEXT_NAMESPACE "SChaosClothAssetEditor3DViewport"
 
+namespace UE::Chaos::ClothAsset::Private
+{
+	static bool bUseOld3DToolbar = true;
+	static FAutoConsoleVariableRef CVARClothEditorUseOld3DToolbar(TEXT("p.ChaosCloth.UseOld3DToolbar"), bUseOld3DToolbar, TEXT("Use the old 3D toolbar until the sim stats are working in the new one"));
+}
+
 void SChaosClothAssetEditor3DViewport::Construct(const FArguments& InArgs, const FAssetEditorViewportConstructionArgs& InViewportConstructionArgs)
 {
 	SAssetEditorViewport::FArguments ParentArgs;
@@ -176,87 +182,93 @@ TSharedPtr<SWidget> SChaosClothAssetEditor3DViewport::MakeViewportToolbar()
 TSharedPtr<SWidget> SChaosClothAssetEditor3DViewport::BuildViewportToolbar()
 {
 	using namespace UE::Chaos::ClothAsset;
-
-	const FName ViewportToolbarName = "ChaosClothEditor.3DViewportToolbar";
-	
-	if (!UToolMenus::Get()->IsMenuRegistered(ViewportToolbarName))
+	if (Private::bUseOld3DToolbar)
+	{
+		return TSharedPtr<SWidget>(nullptr);
+	}
+	else
 	{
-		UToolMenu* ToolbarMenu = UToolMenus::Get()->RegisterMenu(ViewportToolbarName, NAME_None, EMultiBoxType::SlimHorizontalToolBar);
-		ToolbarMenu->StyleName = "ViewportToolbar";
-		
-		ToolbarMenu->FindOrAddSection("Left");
-		
+		const FName ViewportToolbarName = "ChaosClothEditor.3DViewportToolbar";
+
+		if (!UToolMenus::Get()->IsMenuRegistered(ViewportToolbarName))
 		{
-			FToolMenuSection& RightSection = ToolbarMenu->FindOrAddSection("Right");
-			RightSection.Alignment = EToolMenuSectionAlign::Last;
-			
-			RightSection.AddEntry(UE::UnrealEd::CreateCameraSubmenu(UE::UnrealEd::FViewportCameraMenuOptions().ShowAll()));
-			
-			{
-				RightSection.AddEntry(UE::UnrealEd::CreateViewModesSubmenu());
-				UToolMenu* ViewModesMenu = UToolMenus::Get()->ExtendMenu(UToolMenus::JoinMenuPaths(ViewportToolbarName, "ViewModes"));
-				FToolMenuSection& ViewSection = ViewModesMenu->FindOrAddSection("Cloth", LOCTEXT("ClothViewModeSection", "Cloth"));
-				
-				FToolMenuEntry& WireframeEntry = ViewSection.AddMenuEntry(FChaosClothAssetEditorCommands::Get().TogglePreviewWireframe);
-				WireframeEntry.SetShowInToolbarTopLevel(true);
-				WireframeEntry.ToolBarData.ResizeParams.ClippingPriority = 2000;
-			}
-			
+			UToolMenu* ToolbarMenu = UToolMenus::Get()->RegisterMenu(ViewportToolbarName, NAME_None, EMultiBoxType::SlimHorizontalToolBar);
+			ToolbarMenu->StyleName = "ViewportToolbar";
+
+			ToolbarMenu->FindOrAddSection("Left");
+
 			{
-				// LOD
-				RightSection.AddDynamicEntry("DynamicLOD", FNewToolMenuSectionDelegate::CreateLambda([](FToolMenuSection& Section)
+				FToolMenuSection& RightSection = ToolbarMenu->FindOrAddSection("Right");
+				RightSection.Alignment = EToolMenuSectionAlign::Last;
+
+				RightSection.AddEntry(UE::UnrealEd::CreateCameraSubmenu(UE::UnrealEd::FViewportCameraMenuOptions().ShowAll()));
+
+				{
+					RightSection.AddEntry(UE::UnrealEd::CreateViewModesSubmenu());
+					UToolMenu* ViewModesMenu = UToolMenus::Get()->ExtendMenu(UToolMenus::JoinMenuPaths(ViewportToolbarName, "ViewModes"));
+					FToolMenuSection& ViewSection = ViewModesMenu->FindOrAddSection("Cloth", LOCTEXT("ClothViewModeSection", "Cloth"));
+
+					FToolMenuEntry& WireframeEntry = ViewSection.AddMenuEntry(FChaosClothAssetEditorCommands::Get().TogglePreviewWireframe);
+					WireframeEntry.SetShowInToolbarTopLevel(true);
+					WireframeEntry.ToolBarData.ResizeParams.ClippingPriority = 2000;
+				}
+
 				{
-					if (UUnrealEdViewportToolbarContext* Context = Section.FindContext<UUnrealEdViewportToolbarContext>())
-					{
-						if (TSharedPtr<SEditorViewport> EditorViewport = Context->Viewport.Pin())
+					// LOD
+					RightSection.AddDynamicEntry("DynamicLOD", FNewToolMenuSectionDelegate::CreateLambda([](FToolMenuSection& Section)
 						{
-							if (TSharedPtr<FEditorViewportClient> Client = EditorViewport->GetViewportClient())
+							if (UUnrealEdViewportToolbarContext* Context = Section.FindContext<UUnrealEdViewportToolbarContext>())
 							{
-								TSharedPtr<FChaosClothAssetEditor3DViewportClient> ClothClient = StaticCastSharedPtr<FChaosClothAssetEditor3DViewportClient>(Client);
-								Section.AddEntry(UE::UnrealEd::CreatePreviewLODSelectionSubmenu(ClothClient));
+								if (TSharedPtr<SEditorViewport> EditorViewport = Context->Viewport.Pin())
+								{
+									if (TSharedPtr<FEditorViewportClient> Client = EditorViewport->GetViewportClient())
+									{
+										TSharedPtr<FChaosClothAssetEditor3DViewportClient> ClothClient = StaticCastSharedPtr<FChaosClothAssetEditor3DViewportClient>(Client);
+										Section.AddEntry(UE::UnrealEd::CreatePreviewLODSelectionSubmenu(ClothClient));
+									}
+								}
 							}
-						}
-					}
-				}));
-			}
-			
-			{
-				RightSection.AddEntry(UE::UnrealEd::CreateAssetViewerProfileSubmenu());
-				UToolMenu* AssetViewerMenu = UToolMenus::Get()->ExtendMenu(UToolMenus::JoinMenuPaths(ViewportToolbarName, "AssetViewerProfile"));
-				FToolMenuSection& SimulationSection = AssetViewerMenu->FindOrAddSection("SimulationControls", LOCTEXT("SimulationControlsSection", "Simulation Playback Controls"));
-				
-				// Simulation Controls
-				FToolMenuEntryToolBarData ToolBarData;
-				ToolBarData.BlockGroupName = "SimulationControlGroup";
-				ToolBarData.LabelOverride = FText::GetEmpty();
-				ToolBarData.ResizeParams.ClippingPriority = 2000;
-				
-				FToolMenuEntry& HardReset = SimulationSection.AddMenuEntry(FChaosClothAssetEditorCommands::Get().HardResetSimulation);
-				HardReset.ToolBarData = ToolBarData;
-				HardReset.SetShowInToolbarTopLevel(true);
-				
-				FToolMenuEntry& SoftReset = SimulationSection.AddMenuEntry(FChaosClothAssetEditorCommands::Get().SoftResetSimulation);
-				SoftReset.ToolBarData = ToolBarData;
-				SoftReset.SetShowInToolbarTopLevel(true);
-				
-				FToolMenuEntry& ToggleSimulation = SimulationSection.AddMenuEntry(FChaosClothAssetEditorCommands::Get().ToggleSimulationSuspended);
-				ToggleSimulation.ToolBarData = ToolBarData;
-				ToggleSimulation.SetShowInToolbarTopLevel(true);
+						}));
+				}
+
+				{
+					RightSection.AddEntry(UE::UnrealEd::CreateAssetViewerProfileSubmenu());
+					UToolMenu* AssetViewerMenu = UToolMenus::Get()->ExtendMenu(UToolMenus::JoinMenuPaths(ViewportToolbarName, "AssetViewerProfile"));
+					FToolMenuSection& SimulationSection = AssetViewerMenu->FindOrAddSection("SimulationControls", LOCTEXT("SimulationControlsSection", "Simulation Playback Controls"));
+
+					// Simulation Controls
+					FToolMenuEntryToolBarData ToolBarData;
+					ToolBarData.BlockGroupName = "SimulationControlGroup";
+					ToolBarData.LabelOverride = FText::GetEmpty();
+					ToolBarData.ResizeParams.ClippingPriority = 2000;
+
+					FToolMenuEntry& HardReset = SimulationSection.AddMenuEntry(FChaosClothAssetEditorCommands::Get().HardResetSimulation);
+					HardReset.ToolBarData = ToolBarData;
+					HardReset.SetShowInToolbarTopLevel(true);
+
+					FToolMenuEntry& SoftReset = SimulationSection.AddMenuEntry(FChaosClothAssetEditorCommands::Get().SoftResetSimulation);
+					SoftReset.ToolBarData = ToolBarData;
+					SoftReset.SetShowInToolbarTopLevel(true);
+
+					FToolMenuEntry& ToggleSimulation = SimulationSection.AddMenuEntry(FChaosClothAssetEditorCommands::Get().ToggleSimulationSuspended);
+					ToggleSimulation.ToolBarData = ToolBarData;
+					ToggleSimulation.SetShowInToolbarTopLevel(true);
+				}
 			}
 		}
+
+		FToolMenuContext Context;
+		{
+			Context.AppendCommandList(GetCommandList());
+			Context.AppendCommandList(ToolkitCommandList);
+			Context.AddExtender(GetExtenders());
+
+			UUnrealEdViewportToolbarContext* ContextObject = UE::UnrealEd::CreateViewportToolbarDefaultContext(SharedThis(this));
+			Context.AddObject(ContextObject);
+		}
+
+		return UToolMenus::Get()->GenerateWidget(ViewportToolbarName, Context);
 	}
-	
-	FToolMenuContext Context;
-	{
-		Context.AppendCommandList(GetCommandList());
-		Context.AppendCommandList(ToolkitCommandList);
-		Context.AddExtender(GetExtenders());
-		
-		UUnrealEdViewportToolbarContext* ContextObject = UE::UnrealEd::CreateViewportToolbarDefaultContext(SharedThis(this));
-		Context.AddObject(ContextObject);
-	}
-	
-	return UToolMenus::Get()->GenerateWidget(ViewportToolbarName, Context);
 }
 
 bool SChaosClothAssetEditor3DViewport::IsVisible() const
