Restore the in-viewport details text in the Cloth Editor.
- Revert previous cloth-only disabling code.

diff --git a/Engine/Plugins/ChaosClothAssetEditor/Source/ChaosClothAssetEditor/Private/ChaosClothAsset/SClothEditor3DViewport.cpp b/Engine/Plugins/ChaosClothAssetEditor/Source/ChaosClothAssetEditor/Private/ChaosClothAsset/SClothEditor3DViewport.cpp
--- a/Engine/Plugins/ChaosClothAssetEditor/Source/ChaosClothAssetEditor/Private/ChaosClothAsset/SClothEditor3DViewport.cpp
+++ b/Engine/Plugins/ChaosClothAssetEditor/Source/ChaosClothAssetEditor/Private/ChaosClothAsset/SClothEditor3DViewport.cpp
@@ -3,6 +3,7 @@
 #include "ChaosClothAsset/SClothEditor3DViewport.h"
 
 #include "AdvancedPreviewSceneMenus.h"
+#include "ClothEditorSimulationVisualization.h"
 #include "ClothEditorViewportToolbarSections.h"
 #include "ChaosClothAsset/ClothEditor3DViewportClient.h"
 #include "ChaosClothAsset/SClothEditor3DViewportToolBar.h"
@@ -17,15 +18,10 @@
 #include "Widgets/Layout/SBorder.h"
 #include "Components/SkeletalMeshComponent.h"
 #include "ViewportToolbar/UnrealEdViewportToolbar.h"
+#include "Widgets/Text/SRichTextBlock.h"
 
 #define LOCTEXT_NAMESPACE "SChaosClothAssetEditor3DViewport"
 
-namespace UE::Chaos::ClothAsset::Private
-{
-	static bool bUseOld3DToolbar = true;
-	static FAutoConsoleVariableRef CVARClothEditorUseOld3DToolbar(TEXT("p.ChaosCloth.UseOld3DToolbar"), bUseOld3DToolbar, TEXT("Use the old 3D toolbar until the sim stats are working in the new one"));
-}
-
 void SChaosClothAssetEditor3DViewport::Construct(const FArguments& InArgs, const FAssetEditorViewportConstructionArgs& InViewportConstructionArgs)
 {
 	SAssetEditorViewport::FArguments ParentArgs;
@@ -182,93 +178,107 @@ TSharedPtr<SWidget> SChaosClothAssetEditor3DViewport::MakeViewportToolbar()
 TSharedPtr<SWidget> SChaosClothAssetEditor3DViewport::BuildViewportToolbar()
 {
 	using namespace UE::Chaos::ClothAsset;
-	if (Private::bUseOld3DToolbar)
-	{
-		return TSharedPtr<SWidget>(nullptr);
-	}
-	else
+	
+	const FName ViewportToolbarName = "ChaosClothEditor.3DViewportToolbar";
+
+	if (!UToolMenus::Get()->IsMenuRegistered(ViewportToolbarName))
 	{
-		const FName ViewportToolbarName = "ChaosClothEditor.3DViewportToolbar";
+		UToolMenu* ToolbarMenu = UToolMenus::Get()->RegisterMenu(ViewportToolbarName, NAME_None, EMultiBoxType::SlimHorizontalToolBar);
+		ToolbarMenu->StyleName = "ViewportToolbar";
+
+		ToolbarMenu->FindOrAddSection("Left");
 
-		if (!UToolMenus::Get()->IsMenuRegistered(ViewportToolbarName))
 		{
-			UToolMenu* ToolbarMenu = UToolMenus::Get()->RegisterMenu(ViewportToolbarName, NAME_None, EMultiBoxType::SlimHorizontalToolBar);
-			ToolbarMenu->StyleName = "ViewportToolbar";
+			FToolMenuSection& RightSection = ToolbarMenu->FindOrAddSection("Right");
+			RightSection.Alignment = EToolMenuSectionAlign::Last;
 
-			ToolbarMenu->FindOrAddSection("Left");
+			RightSection.AddEntry(UE::UnrealEd::CreateCameraSubmenu(UE::UnrealEd::FViewportCameraMenuOptions().ShowAll()));
 
 			{
-				FToolMenuSection& RightSection = ToolbarMenu->FindOrAddSection("Right");
-				RightSection.Alignment = EToolMenuSectionAlign::Last;
-
-				RightSection.AddEntry(UE::UnrealEd::CreateCameraSubmenu(UE::UnrealEd::FViewportCameraMenuOptions().ShowAll()));
+				RightSection.AddEntry(UE::UnrealEd::CreateViewModesSubmenu());
+				UToolMenu* ViewModesMenu = UToolMenus::Get()->ExtendMenu(UToolMenus::JoinMenuPaths(ViewportToolbarName, "ViewModes"));
+				FToolMenuSection& ViewSection = ViewModesMenu->FindOrAddSection("Cloth", LOCTEXT("ClothViewModeSection", "Cloth"));
 
-				{
-					RightSection.AddEntry(UE::UnrealEd::CreateViewModesSubmenu());
-					UToolMenu* ViewModesMenu = UToolMenus::Get()->ExtendMenu(UToolMenus::JoinMenuPaths(ViewportToolbarName, "ViewModes"));
-					FToolMenuSection& ViewSection = ViewModesMenu->FindOrAddSection("Cloth", LOCTEXT("ClothViewModeSection", "Cloth"));
-
-					FToolMenuEntry& WireframeEntry = ViewSection.AddMenuEntry(FChaosClothAssetEditorCommands::Get().TogglePreviewWireframe);
-					WireframeEntry.SetShowInToolbarTopLevel(true);
-					WireframeEntry.ToolBarData.ResizeParams.ClippingPriority = 2000;
-				}
+				FToolMenuEntry& WireframeEntry = ViewSection.AddMenuEntry(FChaosClothAssetEditorCommands::Get().TogglePreviewWireframe);
+				WireframeEntry.SetShowInToolbarTopLevel(true);
+				WireframeEntry.ToolBarData.ResizeParams.ClippingPriority = 2000;
+			}
 
-				{
-					// LOD
-					RightSection.AddDynamicEntry("DynamicLOD", FNewToolMenuSectionDelegate::CreateLambda([](FToolMenuSection& Section)
+			{
+				// LOD
+				RightSection.AddDynamicEntry("DynamicLOD", FNewToolMenuSectionDelegate::CreateLambda([](FToolMenuSection& Section)
+					{
+						if (UUnrealEdViewportToolbarContext* Context = Section.FindContext<UUnrealEdViewportToolbarContext>())
 						{
-							if (UUnrealEdViewportToolbarContext* Context = Section.FindContext<UUnrealEdViewportToolbarContext>())
+							if (TSharedPtr<SEditorViewport> EditorViewport = Context->Viewport.Pin())
 							{
-								if (TSharedPtr<SEditorViewport> EditorViewport = Context->Viewport.Pin())
+								if (TSharedPtr<FEditorViewportClient> Client = EditorViewport->GetViewportClient())
 								{
-									if (TSharedPtr<FEditorViewportClient> Client = EditorViewport->GetViewportClient())
-									{
-										TSharedPtr<FChaosClothAssetEditor3DViewportClient> ClothClient = StaticCastSharedPtr<FChaosClothAssetEditor3DViewportClient>(Client);
-										Section.AddEntry(UE::UnrealEd::CreatePreviewLODSelectionSubmenu(ClothClient));
-									}
+									TSharedPtr<FChaosClothAssetEditor3DViewportClient> ClothClient = StaticCastSharedPtr<FChaosClothAssetEditor3DViewportClient>(Client);
+									Section.AddEntry(UE::UnrealEd::CreatePreviewLODSelectionSubmenu(ClothClient));
 								}
 							}
-						}));
-				}
-
-				{
-					RightSection.AddEntry(UE::UnrealEd::CreateAssetViewerProfileSubmenu());
-					UToolMenu* AssetViewerMenu = UToolMenus::Get()->ExtendMenu(UToolMenus::JoinMenuPaths(ViewportToolbarName, "AssetViewerProfile"));
-					FToolMenuSection& SimulationSection = AssetViewerMenu->FindOrAddSection("SimulationControls", LOCTEXT("SimulationControlsSection", "Simulation Playback Controls"));
-
-					// Simulation Controls
-					FToolMenuEntryToolBarData ToolBarData;
-					ToolBarData.BlockGroupName = "SimulationControlGroup";
-					ToolBarData.LabelOverride = FText::GetEmpty();
-					ToolBarData.ResizeParams.ClippingPriority = 2000;
-
-					FToolMenuEntry& HardReset = SimulationSection.AddMenuEntry(FChaosClothAssetEditorCommands::Get().HardResetSimulation);
-					HardReset.ToolBarData = ToolBarData;
-					HardReset.SetShowInToolbarTopLevel(true);
-
-					FToolMenuEntry& SoftReset = SimulationSection.AddMenuEntry(FChaosClothAssetEditorCommands::Get().SoftResetSimulation);
-					SoftReset.ToolBarData = ToolBarData;
-					SoftReset.SetShowInToolbarTopLevel(true);
-
-					FToolMenuEntry& ToggleSimulation = SimulationSection.AddMenuEntry(FChaosClothAssetEditorCommands::Get().ToggleSimulationSuspended);
-					ToggleSimulation.ToolBarData = ToolBarData;
-					ToggleSimulation.SetShowInToolbarTopLevel(true);
-				}
+						}
+					}));
 			}
-		}
 
-		FToolMenuContext Context;
-		{
-			Context.AppendCommandList(GetCommandList());
-			Context.AppendCommandList(ToolkitCommandList);
-			Context.AddExtender(GetExtenders());
-
-			UUnrealEdViewportToolbarContext* ContextObject = UE::UnrealEd::CreateViewportToolbarDefaultContext(SharedThis(this));
-			Context.AddObject(ContextObject);
+			{
+				RightSection.AddEntry(UE::UnrealEd::CreateAssetViewerProfileSubmenu());
+				UToolMenu* AssetViewerMenu = UToolMenus::Get()->ExtendMenu(UToolMenus::JoinMenuPaths(ViewportToolbarName, "AssetViewerProfile"));
+				FToolMenuSection& SimulationSection = AssetViewerMenu->FindOrAddSection("SimulationControls", LOCTEXT("SimulationControlsSection", "Simulation Playback Controls"));
+
+				// Simulation Controls
+				FToolMenuEntryToolBarData ToolBarData;
+				ToolBarData.BlockGroupName = "SimulationControlGroup";
+				ToolBarData.LabelOverride = FText::GetEmpty();
+				ToolBarData.ResizeParams.ClippingPriority = 2000;
+
+				FToolMenuEntry& HardReset = SimulationSection.AddMenuEntry(FChaosClothAssetEditorCommands::Get().HardResetSimulation);
+				HardReset.ToolBarData = ToolBarData;
+				HardReset.SetShowInToolbarTopLevel(true);
+
+				FToolMenuEntry& SoftReset = SimulationSection.AddMenuEntry(FChaosClothAssetEditorCommands::Get().SoftResetSimulation);
+				SoftReset.ToolBarData = ToolBarData;
+				SoftReset.SetShowInToolbarTopLevel(true);
+
+				FToolMenuEntry& ToggleSimulation = SimulationSection.AddMenuEntry(FChaosClothAssetEditorCommands::Get().ToggleSimulationSuspended);
+				ToggleSimulation.ToolBarData = ToolBarData;
+				ToggleSimulation.SetShowInToolbarTopLevel(true);
+			}
 		}
+	}
 
-		return UToolMenus::Get()->GenerateWidget(ViewportToolbarName, Context);
+	FToolMenuContext Context;
+	{
+		Context.AppendCommandList(GetCommandList());
+		Context.AppendCommandList(ToolkitCommandList);
+		Context.AddExtender(GetExtenders());
+
+		UUnrealEdViewportToolbarContext* ContextObject = UE::UnrealEd::CreateViewportToolbarDefaultContext(SharedThis(this));
+		Context.AddObject(ContextObject);
 	}
+
+	return UToolMenus::Get()->GenerateWidget(ViewportToolbarName, Context);
+}
+
+void SChaosClothAssetEditor3DViewport::PopulateViewportOverlays(TSharedRef<SOverlay> Overlay)
+{
+	Overlay->AddSlot()
+	.Padding(FMargin(4.0f, 3.0f, 0.0f, 0.0f))
+	[
+		SNew(SBox)
+		.Visibility_Lambda([]
+		{
+			// Don't show this while the legacy toolbars are visible.
+			return UE::UnrealEd::ShowOldViewportToolbars() ? EVisibility::Collapsed : EVisibility::Visible;
+		})
+		[
+			SNew(SRichTextBlock)
+			.DecoratorStyleSet(&FAppStyle::Get())
+			.TextStyle(&FAppStyle::Get().GetWidgetStyle<FTextBlockStyle>("AnimViewport.MessageText"))
+			.Text(this, &SChaosClothAssetEditor3DViewport::GetViewportDisplayString)
+		]
+	];
 }
 
 bool SChaosClothAssetEditor3DViewport::IsVisible() const
@@ -295,6 +305,17 @@ TSharedRef<class SEditorViewport> SChaosClothAssetEditor3DViewport::GetViewportW
 	return SharedThis(this);
 }
 
+FText SChaosClothAssetEditor3DViewport::GetViewportDisplayString() const
+{
+	using namespace UE::Chaos::ClothAsset;
+	TSharedRef<FChaosClothAssetEditor3DViewportClient> ViewportClient = StaticCastSharedPtr<FChaosClothAssetEditor3DViewportClient>(GetViewportClient()).ToSharedRef();
+	if (FClothEditorSimulationVisualization* const Visualization = ViewportClient->GetSimulationVisualization().Pin().Get())
+	{
+		return Visualization->GetDisplayString(ViewportClient->GetPreviewClothComponent());
+	}
+	return FText();
+}
+
 TSharedPtr<FExtender> SChaosClothAssetEditor3DViewport::GetExtenders() const
 {
 	TSharedPtr<FExtender> Result(MakeShareable(new FExtender));
diff --git a/Engine/Plugins/ChaosClothAssetEditor/Source/ChaosClothAssetEditor/Private/ChaosClothAsset/SClothEditor3DViewport.h b/Engine/Plugins/ChaosClothAssetEditor/Source/ChaosClothAssetEditor/Private/ChaosClothAsset/SClothEditor3DViewport.h
--- a/Engine/Plugins/ChaosClothAssetEditor/Source/ChaosClothAssetEditor/Private/ChaosClothAsset/SClothEditor3DViewport.h
+++ b/Engine/Plugins/ChaosClothAssetEditor/Source/ChaosClothAssetEditor/Private/ChaosClothAsset/SClothEditor3DViewport.h
@@ -33,6 +33,7 @@ public:
 	virtual void BindCommands() override;
 	virtual TSharedPtr<SWidget> MakeViewportToolbar() override;
 	virtual TSharedPtr<SWidget> BuildViewportToolbar() override;
+	virtual void PopulateViewportOverlays(TSharedRef<SOverlay> Overlay) override;
 	virtual bool IsVisible() const override;
 
 	virtual void OnFocusViewportToSelection() override;
@@ -52,6 +53,8 @@ private:
 	TWeakPtr<const UE::Chaos::ClothAsset::FChaosClothPreviewScene> GetPreviewScene() const;
 
 	virtual TSharedPtr<IPreviewProfileController> CreatePreviewProfileController() override;
+	
+	FText GetViewportDisplayString() const;
 
 	float GetViewMinInput() const;
 	float GetViewMaxInput() const;
