Chaos Cloth Asset - Fixed an issue where an existing Dataflow asset is overwritten when creating a new Cloth Asset that creates a new Dataflow of the same name.

[RN] Chaos Cloth Asset - Fixed an issue where an existing Dataflow asset is overwritten when creating a new Cloth Asset that creates a new Dataflow of the same name.

diff --git a/Engine/Plugins/ChaosClothAssetEditor/Source/ChaosClothAssetTools/Private/ChaosClothAsset/ClothAssetFactory.cpp b/Engine/Plugins/ChaosClothAssetEditor/Source/ChaosClothAssetTools/Private/ChaosClothAsset/ClothAssetFactory.cpp
--- a/Engine/Plugins/ChaosClothAssetEditor/Source/ChaosClothAssetTools/Private/ChaosClothAsset/ClothAssetFactory.cpp
+++ b/Engine/Plugins/ChaosClothAssetEditor/Source/ChaosClothAssetTools/Private/ChaosClothAsset/ClothAssetFactory.cpp
@@ -28,9 +28,14 @@ UObject* UChaosClothAssetFactory::FactoryCreateNew(UClass* Class, UObject* Paren
 	// Create a new Dataflow asset
 	const FString DataflowPath = FPackageName::GetLongPackagePath(ClothAsset->GetOutermost()->GetName());
 	const FString ClothAssetName = ClothAsset->GetName();
-	const FString DataflowName = FString(TEXT("DF_")) + (ClothAssetName.StartsWith(TEXT("CA_")) ? ClothAssetName.RightChop(3) : ClothAssetName);
-	const FString DataflowPackageName = FPaths::Combine(DataflowPath, DataflowName);
-
+	FString DataflowName = FString(TEXT("DF_")) + (ClothAssetName.StartsWith(TEXT("CA_")) ? ClothAssetName.RightChop(3) : ClothAssetName);
+	FString DataflowPackageName = FPaths::Combine(DataflowPath, DataflowName);
+	if (FindPackage(nullptr, *DataflowPackageName))
+	{
+		// If a Dataflow asset already exists with this name, make a unique name from it to avoid clobbering it
+		MakeUniqueObjectName(nullptr, UPackage::StaticClass(), FName(DataflowPackageName)).ToString(DataflowPackageName);
+		DataflowName = FPaths::GetBaseFilename(DataflowPackageName);
+	}
 	UPackage* const DataflowPackage = CreatePackage(*DataflowPackageName);
 
 	// Load the cloth template into the new Dataflow asset
