Chaos Cloth Asset - Fixed the Cloth Asset Import node when used in non Cloth Asset contexts.
Chaos Cloth Outfit - Fixed the Cloth Outfit merge when render sections have no valid AssetGuids.

#rnx

diff --git a/Engine/Plugins/ChaosClothAssetEditor/Source/ChaosClothAssetDataflowNodes/Private/ChaosClothAsset/ImportNode.cpp b/Engine/Plugins/ChaosClothAssetEditor/Source/ChaosClothAssetDataflowNodes/Private/ChaosClothAsset/ImportNode.cpp
--- a/Engine/Plugins/ChaosClothAssetEditor/Source/ChaosClothAssetDataflowNodes/Private/ChaosClothAsset/ImportNode.cpp
+++ b/Engine/Plugins/ChaosClothAssetEditor/Source/ChaosClothAssetDataflowNodes/Private/ChaosClothAsset/ImportNode.cpp
@@ -37,7 +37,7 @@ void FChaosClothAssetImportNode::Evaluate(UE::Dataflow::FContext& Context, const
 
 		if (TObjectPtr<const UChaosClothAsset> InClothAsset = GetValue(Context, &ClothAsset))
 		{
-			bool bSetValue = false;
+			bool bSetValue = true;
 			// Copy the main cloth asset details to this dataflow's owner if any
 			if (const UE::Dataflow::FEngineContext* EngineContext = Context.AsType<UE::Dataflow::FEngineContext>())
 			{
@@ -51,17 +51,8 @@ void FChaosClothAssetImportNode::Evaluate(UE::Dataflow::FContext& Context, const
 							LOCTEXT("RecursiveAssetLoopHeadline", "Recursive asset loop."),
 							LOCTEXT("RecursiveAssetLoopDetails", "The source asset cannot be the same as the terminal asset."));
 					}
-					else
-					{
-						bSetValue = true;
-					}
 				}
 			}
-			else
-			{
-				// No terminal asset, this is a stray dataflow and it is safe to set the value without fear of a loop
-				bSetValue = true;
-			}
 
 			// Copy the cloth asset to this node's output collection
 			if (bSetValue)
diff --git a/Engine/Plugins/Experimental/ChaosOutfitAsset/Source/ChaosOutfitAssetEngine/Private/ChaosOutfitAsset/Outfit.cpp b/Engine/Plugins/Experimental/ChaosOutfitAsset/Source/ChaosOutfitAssetEngine/Private/ChaosOutfitAsset/Outfit.cpp
--- a/Engine/Plugins/Experimental/ChaosOutfitAsset/Source/ChaosOutfitAssetEngine/Private/ChaosOutfitAsset/Outfit.cpp
+++ b/Engine/Plugins/Experimental/ChaosOutfitAsset/Source/ChaosOutfitAssetEngine/Private/ChaosOutfitAsset/Outfit.cpp
@@ -1037,7 +1037,7 @@ void UChaosOutfit::MergeLODRenderDatas(
 	for (const FSkelMeshRenderSection& RenderSection : LODRenderData.RenderSections)
 	{
 		const int32 PieceIndex = AssetGuids.Find(RenderSection.ClothingData.AssetGuid);
-		if (PieceIndex != INDEX_NONE)
+		if (PieceIndex != INDEX_NONE || !RenderSection.ClothingData.AssetGuid.IsValid())
 		{
 			FSkelMeshRenderSection& OutRenderSection = OutLODRenderData.RenderSections.AddDefaulted_GetRef();
 
