Optimize Transfer Skin Weight Node
Transfer Inpaint Weights To Mesh Components is now multithreaded.
Several FTransferBoneWeights are created, one for each threads.
All computed Sub Mesh are stored and all sub mesh skinned weight are applied safely in single thread.
On the SimMeshResizing asset performances goes on the Transfer Skin Weight Node from 2.2 seconds to 0.75 seconds.

diff --git a/Engine/Plugins/ChaosClothAssetEditor/Source/ChaosClothAssetDataflowNodes/Private/ChaosClothAsset/TransferSkinWeightsNode.cpp b/Engine/Plugins/ChaosClothAssetEditor/Source/ChaosClothAssetDataflowNodes/Private/ChaosClothAsset/TransferSkinWeightsNode.cpp
--- a/Engine/Plugins/ChaosClothAssetEditor/Source/ChaosClothAssetDataflowNodes/Private/ChaosClothAsset/TransferSkinWeightsNode.cpp
+++ b/Engine/Plugins/ChaosClothAssetEditor/Source/ChaosClothAssetDataflowNodes/Private/ChaosClothAsset/TransferSkinWeightsNode.cpp
@@ -8,6 +8,8 @@
 #include "Animation/Skeleton.h"
 #include "Dataflow/DataflowInputOutput.h"
 #include "DynamicMesh/DynamicBoneAttribute.h"
+#include "DynamicMesh/DynamicMesh3.h"
+#include "DynamicMesh/DynamicMeshAABBTree3.h"
 #include "DynamicMesh/DynamicMeshAttributeSet.h"
 #include "DynamicMesh/DynamicVertexSkinWeightsAttribute.h"
 #include "DynamicMesh/MeshNormals.h"
@@ -24,6 +26,7 @@
 #include "ChaosClothAsset/ClothCollectionGroup.h"
 #include "Rendering/SkeletalMeshRenderData.h"
 #include "SkeletalMeshLODRenderDataToDynamicMesh.h"
+#include "Tasks/Task.h"
 
 #include UE_INLINE_GENERATED_CPP_BY_NAME(TransferSkinWeightsNode)
 
@@ -31,6 +34,13 @@
 
 static const FName InpaintWeightMaskName(TEXT("_InpaintWeightMask"));
 
+namespace UE::Geometry
+{
+	class FDynamicMesh3;
+	template<typename MeshType> class TMeshAABBTree3;
+	typedef TMeshAABBTree3<FDynamicMesh3> FDynamicMeshAABBTree3;
+}
+
 namespace UE::Chaos::ClothAsset::Private
 {	
 	/** Helper struct to pass the transfer settings around */
@@ -363,22 +373,6 @@ namespace UE::Chaos::ClothAsset::Private
 	{
 		using namespace UE::Geometry;
 
-		// Setup transfer weights operator
-		FTransferBoneWeights TransferBoneWeights(&SourceDynamicMesh, FSkeletalMeshAttributes::DefaultSkinWeightProfileName);
-		TransferBoneWeights.bUseParallel = TransferSettings.bUseParallel;
-		TransferBoneWeights.MaxNumInfluences = TransferSettings.MaxNumInfluences;
-		TransferBoneWeights.TransferMethod = TransferSettings.TransferMethod;
-		TransferBoneWeights.NormalThreshold = FMathd::DegToRad * TransferSettings.NormalThreshold;
-		TransferBoneWeights.SearchRadius = TransferSettings.RadiusPercentage * TargetDynamicMesh.GetBounds().DiagonalLength();
-		TransferBoneWeights.NumSmoothingIterations = TransferSettings.NumSmoothingIterations;
-		TransferBoneWeights.SmoothingStrength = TransferSettings.SmoothingStrength;
-		TransferBoneWeights.LayeredMeshSupport = TransferSettings.LayeredMeshSupport; // multilayerd clothing
-
-		if (TransferBoneWeights.Validate() != EOperationValidationResult::Ok)
-		{
-			UE_LOG(LogChaosClothAssetDataflowNodes, Warning, TEXT("TransferSkinWeightsNode: Transfer method parameters are invalid."));
-			return false;
-		}
 
 		// Find connected-components
 		FMeshConnectedComponents ConnectedComponents(&TargetDynamicMesh);
@@ -401,98 +395,159 @@ namespace UE::Chaos::ClothAsset::Private
 
 		// Iterate over each component and perform per-component skin weight transfer
 		const int32 NumComponents = ConnectedComponents.Num();
-		for (int32 ComponentIdx = 0; ComponentIdx < NumComponents; ++ComponentIdx)
-		{
-			// Mesh containing the geometry of the current component
-			FDynamicMesh3 Submesh;
-			Submesh.EnableAttributes();
-			Submesh.Attributes()->EnableBones(TargetRefSkeleton.GetRawBoneNum());
-			for (int32 BoneIdx = 0; BoneIdx < TargetRefSkeleton.GetRawBoneNum(); ++BoneIdx)
-			{
-				Submesh.Attributes()->GetBoneNames()->SetValue(BoneIdx, TargetRefSkeleton.GetRawRefBoneInfo()[BoneIdx].Name);
-			}
+		const int32 NumTasks = FMath::Max(FMath::Min(FTaskGraphInterface::Get().GetNumWorkerThreads(), NumComponents), 1);
+		constexpr int32 MinComponentByTask = 1;
+		const int32 ComnponentsByTask = FMath::Max(FMath::Max(FMath::DivideAndRoundUp(NumComponents, NumTasks), MinComponentByTask), 1);
+		const int32 NumBatches = FMath::DivideAndRoundUp(NumComponents, ComnponentsByTask);
+		TArray<UE::Tasks::FTask> PendingTasks;
+		PendingTasks.Reserve(NumBatches);
 
-			// Keep track of the maps between Mesh and the component submesh vertex/triangle indices
-			TMap<int32, int32> BaseToSubmeshV;
-			TArray<int32> SubmeshToBaseV;
-			TArray<int32> SubmeshToBaseT;
+		
+		TArray<FDynamicMesh3> Submeshes;
+		Submeshes.SetNum(NumComponents);
+		TArray<TArray<int32>> SubmeshToBaseVs;
+		SubmeshToBaseVs.SetNum(NumComponents);
+		std::atomic<bool> bTransferParametersInvalid(false);
 
-			const TArray<int32>& ComponentTris = ConnectedComponents[ComponentIdx].Indices;
+		// Creating the BVH here in single threaded which is expensive allow to avoid every thread to pay the expensive cost
+		FDynamicMeshAABBTree3 SourceBVH(&SourceDynamicMesh);
 
-			for (const int32 TID : ComponentTris)
+		for (int32 BatchIndex = 0; BatchIndex < NumBatches; BatchIndex++)
+		{
+			const int32 StartIndex = BatchIndex * ComnponentsByTask;
+			int32 EndIndex = (BatchIndex + 1) * ComnponentsByTask;
+			EndIndex = BatchIndex == NumBatches - 1 ? FMath::Min(NumComponents, EndIndex) : EndIndex;
+
+			UE::Tasks::FTask PendingTask = UE::Tasks::Launch(UE_SOURCE_LOCATION, [StartIndex, EndIndex, &TargetRefSkeleton, &ConnectedComponents, 
+				&TargetDynamicMesh, &ForceInpaintWeightLayer, &TransferSettings, &SourceDynamicMesh, &SourceBVH, &Submeshes, &SubmeshToBaseVs, &bTransferParametersInvalid]()
 			{
-				const FIndex3i Triangle = TargetDynamicMesh.GetTriangle(TID);
-				FIndex3i NewTriangle;
-				for (int32 TriCornerIdx = 0; TriCornerIdx < 3; ++TriCornerIdx)
+				QUICK_SCOPE_CYCLE_COUNTER(STAT_TransferInpaintWeightsToMeshComponentsTask);
+				if (bTransferParametersInvalid.load(std::memory_order_relaxed))
 				{
-					const int32* FoundIdx = BaseToSubmeshV.Find(Triangle[TriCornerIdx]);
-					if (FoundIdx)
+					return;
+				}
+				// Setup transfer weights operator
+				FTransferBoneWeights TransferBoneWeights(&SourceDynamicMesh, FSkeletalMeshAttributes::DefaultSkinWeightProfileName, &SourceBVH);
+				TransferBoneWeights.bUseParallel = TransferSettings.bUseParallel;
+				TransferBoneWeights.MaxNumInfluences = TransferSettings.MaxNumInfluences;
+				TransferBoneWeights.TransferMethod = TransferSettings.TransferMethod;
+				TransferBoneWeights.NormalThreshold = FMathd::DegToRad * TransferSettings.NormalThreshold;
+				TransferBoneWeights.SearchRadius = TransferSettings.RadiusPercentage * TargetDynamicMesh.GetBounds().DiagonalLength();
+				TransferBoneWeights.NumSmoothingIterations = TransferSettings.NumSmoothingIterations;
+				TransferBoneWeights.SmoothingStrength = TransferSettings.SmoothingStrength;
+				TransferBoneWeights.LayeredMeshSupport = TransferSettings.LayeredMeshSupport; // multilayerd clothing
+				if (TransferBoneWeights.Validate() != EOperationValidationResult::Ok)
+				{
+					bTransferParametersInvalid.store(true, std::memory_order_relaxed);
+					return;
+				}
+
+				for (int32 ComponentIdx = StartIndex; ComponentIdx < EndIndex; ++ComponentIdx)
+				{
+					// Mesh containing the geometry of the current component
+					FDynamicMesh3& Submesh = Submeshes[ComponentIdx];
+					Submesh.EnableAttributes();
+					Submesh.Attributes()->EnableBones(TargetRefSkeleton.GetRawBoneNum());
+					for (int32 BoneIdx = 0; BoneIdx < TargetRefSkeleton.GetRawBoneNum(); ++BoneIdx)
 					{
-						NewTriangle[TriCornerIdx] = *FoundIdx;
+						Submesh.Attributes()->GetBoneNames()->SetValue(BoneIdx, TargetRefSkeleton.GetRawRefBoneInfo()[BoneIdx].Name);
 					}
-					else
+
+					// Keep track of the maps between Mesh and the component submesh vertex/triangle indices
+					TMap<int32, int32> BaseToSubmeshV;
+					TArray<int32>& SubmeshToBaseV = SubmeshToBaseVs[ComponentIdx];
+					TArray<int32> SubmeshToBaseT;
+
+					const TArray<int32>& ComponentTris = ConnectedComponents[ComponentIdx].Indices;
+
+					for (const int32 TID : ComponentTris)
 					{
-						const FVector3d Position = TargetDynamicMesh.GetVertex(Triangle[TriCornerIdx]);
-						const int32 NewVtxID = Submesh.AppendVertex(Position);
-						check(NewVtxID == SubmeshToBaseV.Num());
-						SubmeshToBaseV.Add(Triangle[TriCornerIdx]);
-						BaseToSubmeshV.Add(Triangle[TriCornerIdx], NewVtxID);
-						NewTriangle[TriCornerIdx] = NewVtxID;
+						const FIndex3i Triangle = TargetDynamicMesh.GetTriangle(TID);
+						FIndex3i NewTriangle;
+						for (int32 TriCornerIdx = 0; TriCornerIdx < 3; ++TriCornerIdx)
+						{
+							const int32* FoundIdx = BaseToSubmeshV.Find(Triangle[TriCornerIdx]);
+							if (FoundIdx)
+							{
+								NewTriangle[TriCornerIdx] = *FoundIdx;
+							}
+							else
+							{
+								const FVector3d Position = TargetDynamicMesh.GetVertex(Triangle[TriCornerIdx]);
+								const int32 NewVtxID = Submesh.AppendVertex(Position);
+								check(NewVtxID == SubmeshToBaseV.Num());
+								SubmeshToBaseV.Add(Triangle[TriCornerIdx]);
+								BaseToSubmeshV.Add(Triangle[TriCornerIdx], NewVtxID);
+								NewTriangle[TriCornerIdx] = NewVtxID;
+							}
+						}
+
+						const int32 NewTriID = Submesh.AppendTriangle(NewTriangle);
+						check(NewTriID == SubmeshToBaseT.Num());
+						SubmeshToBaseT.Add(TID);
 					}
-				}
 
-				const int32 NewTriID = Submesh.AppendTriangle(NewTriangle);
-				check(NewTriID == SubmeshToBaseT.Num());
-				SubmeshToBaseT.Add(TID);
-			}
+					// Copy over the force inpaint values
+					if (ForceInpaintWeightLayer != nullptr)
+					{
+						TArray<float> ForceInpaint;
+						ForceInpaint.SetNum(Submesh.MaxVertexID());
 
-			// Copy over the force inpaint values
-			if (ForceInpaintWeightLayer != nullptr)
-			{
-				TArray<float> ForceInpaint;
-				ForceInpaint.SetNum(Submesh.MaxVertexID());
-				
-				for (int32 VID = 0; VID < Submesh.MaxVertexID(); ++VID)
-				{	
-					float Value = 0;
-					ForceInpaintWeightLayer->GetValue(SubmeshToBaseV[VID], &Value);
-					ForceInpaint[VID] = Value;
-				}
+						for (int32 VID = 0; VID < Submesh.MaxVertexID(); ++VID)
+						{
+							float Value = 0;
+							ForceInpaintWeightLayer->GetValue(SubmeshToBaseV[VID], &Value);
+							ForceInpaint[VID] = Value;
+						}
 
-				// Set the mask in the transfer operator
-				TransferBoneWeights.ForceInpaint = MoveTemp(ForceInpaint);
-			}
+						// Set the mask in the transfer operator
+						TransferBoneWeights.ForceInpaint = MoveTemp(ForceInpaint);
+					}
 
-			// Transfer weights to the current submesh only. If transfer using the inpaint method fails, fall back to  
-			// the closest point method which should always succeed. Common reason for the failure is if we didn't find  
-			// any matches on the source at all with the current transfer settings.
-			if (!TransferBoneWeights.TransferWeightsToMesh(Submesh, FSkeletalMeshAttributes::DefaultSkinWeightProfileName))
-			{	
-				UE_LOG(LogChaosClothAssetDataflowNodes, Warning, TEXT("TransferSkinWeightsNode: Failed to transfer skin weights to some of the vertices of the render mesh using inpaint method, falling back to closest point method."));
+					// Transfer weights to the current submesh only. If transfer using the inpaint method fails, fall back to  
+					// the closest point method which should always succeed. Common reason for the failure is if we didn't find  
+					// any matches on the source at all with the current transfer settings.
+					if (!TransferBoneWeights.TransferWeightsToMesh(Submesh, FSkeletalMeshAttributes::DefaultSkinWeightProfileName))
+					{
+						UE_LOG(LogChaosClothAssetDataflowNodes, Warning, TEXT("TransferSkinWeightsNode: Failed to transfer skin weights to some of the vertices of the render mesh using inpaint method, falling back to closest point method."));
 
-				// We can reuse the same operator but change the method type
-				TransferBoneWeights.TransferMethod = FTransferBoneWeights::ETransferBoneWeightsMethod::ClosestPointOnSurface;
-				
-				// Ignore radius and normal settings, so all points on the source are considered
-				const double OldSearchRadius = TransferBoneWeights.SearchRadius;
-				const double OldNormalThreshold = TransferBoneWeights.NormalThreshold;
-				TransferBoneWeights.SearchRadius = -1;	  
-				TransferBoneWeights.NormalThreshold = -1; 
-
-				// This should always succeed
-				if (!ensure(TransferBoneWeights.TransferWeightsToMesh(Submesh, FSkeletalMeshAttributes::DefaultSkinWeightProfileName)))
-				{
-					UE_LOG(LogChaosClothAssetDataflowNodes, Warning, TEXT("TransferSkinWeightsNode:  Failed to transfer skin weights to some of the vertices of the render mesh."));
+						// We can reuse the same operator but change the method type
+						TransferBoneWeights.TransferMethod = FTransferBoneWeights::ETransferBoneWeightsMethod::ClosestPointOnSurface;
+
+						// Ignore radius and normal settings, so all points on the source are considered
+						const double OldSearchRadius = TransferBoneWeights.SearchRadius;
+						const double OldNormalThreshold = TransferBoneWeights.NormalThreshold;
+						TransferBoneWeights.SearchRadius = -1;
+						TransferBoneWeights.NormalThreshold = -1;
+
+						// This should always succeed
+						if (!ensure(TransferBoneWeights.TransferWeightsToMesh(Submesh, FSkeletalMeshAttributes::DefaultSkinWeightProfileName)))
+						{
+							UE_LOG(LogChaosClothAssetDataflowNodes, Warning, TEXT("TransferSkinWeightsNode:  Failed to transfer skin weights to some of the vertices of the render mesh."));
+						}
+
+						// Revert back the settings
+						TransferBoneWeights.TransferMethod = TransferSettings.TransferMethod;
+						TransferBoneWeights.SearchRadius = OldSearchRadius;
+						TransferBoneWeights.NormalThreshold = OldNormalThreshold;
+					}
 				}
+			});
+			PendingTasks.Add(PendingTask);
+		}
+		UE::Tasks::Wait(PendingTasks);
 
-				// Revert back the settings
-				TransferBoneWeights.TransferMethod = TransferSettings.TransferMethod;
-				TransferBoneWeights.SearchRadius = OldSearchRadius;
-				TransferBoneWeights.NormalThreshold = OldNormalThreshold;
-			}
-				
+		if (bTransferParametersInvalid)
+		{
+			UE_LOG(LogChaosClothAssetDataflowNodes, Warning, TEXT("TransferSkinWeightsNode: Transfer method parameters are invalid."));
+			return false;
+		}
+
+		for (int32 ComponentIdx = 0; ComponentIdx < NumComponents; ++ComponentIdx)
+		{
+			TArray<int32>& SubmeshToBaseV = SubmeshToBaseVs[ComponentIdx];
 			// Copy over the data from the submesh to the base mesh
-			FDynamicMeshVertexSkinWeightsAttribute* SubMeshSkinWeights = UE::Chaos::ClothAsset::Private::GetOrCreateSkinWeightsAttribute(Submesh);
+			FDynamicMeshVertexSkinWeightsAttribute* SubMeshSkinWeights = UE::Chaos::ClothAsset::Private::GetOrCreateSkinWeightsAttribute(Submeshes[ComponentIdx]);
 			FDynamicMeshVertexSkinWeightsAttribute* BaseMeshSkinWeights = UE::Chaos::ClothAsset::Private::GetOrCreateSkinWeightsAttribute(TargetDynamicMesh);
 
 			for (int32 SubMeshVID = 0; SubMeshVID < SubmeshToBaseV.Num(); ++SubMeshVID)
@@ -760,6 +815,7 @@ FChaosClothAssetTransferSkinWeightsNode::FChaosClothAssetTransferSkinWeightsNode
 
 void FChaosClothAssetTransferSkinWeightsNode::Evaluate(UE::Dataflow::FContext& Context, const FDataflowOutput* Out) const
 {
+	QUICK_SCOPE_CYCLE_COUNTER(STAT_FChaosClothAssetTransferSkinWeightsNode_Evaluate);
 	using namespace UE::Chaos::ClothAsset;
 	using namespace UE::Chaos::ClothAsset::Private;
 	using namespace UE::Geometry;
