Chaos Outfit Asset - Fixed crash when exporting an outfit asset with a null material into a skeletal mesh.

#rnx

diff --git a/Engine/Plugins/ChaosClothAssetEditor/Source/ChaosClothAssetTools/Private/ChaosClothAsset/SkeletalMeshConverter.cpp b/Engine/Plugins/ChaosClothAssetEditor/Source/ChaosClothAssetTools/Private/ChaosClothAsset/SkeletalMeshConverter.cpp
--- a/Engine/Plugins/ChaosClothAssetEditor/Source/ChaosClothAssetTools/Private/ChaosClothAsset/SkeletalMeshConverter.cpp
+++ b/Engine/Plugins/ChaosClothAssetEditor/Source/ChaosClothAssetTools/Private/ChaosClothAsset/SkeletalMeshConverter.cpp
@@ -105,7 +105,14 @@ bool UClothAssetEditorSkeletalMeshConverter::ExportToSkeletalMesh(const UChaosCl
 	SkeletalMaterials.Reserve(RenderMaterials.Num());
 	for (UMaterialInterface* RenderMaterial : RenderMaterials)
 	{
-		SkeletalMaterials.Emplace(RenderMaterial, RenderMaterial->GetFName());
+		if (RenderMaterial)
+		{
+			SkeletalMaterials.Emplace(RenderMaterial, RenderMaterial->GetFName());
+		}
+		else
+		{
+			SkeletalMaterials.Emplace(UMaterial::GetDefaultMaterial(MD_Surface));
+		}
 	}
 	if (!SkeletalMaterials.Num())
 	{
