Cloth Asset - Enable the USD cloth schema based import.

[RN] Chaos Cloth Asset - Added support for importing USD files complying to the new USD cloth schema.

diff --git a/Engine/Plugins/ChaosClothAssetEditor/Source/ChaosClothAssetDataflowNodes/ChaosClothAssetDataflowNodes.Build.cs b/Engine/Plugins/ChaosClothAssetEditor/Source/ChaosClothAssetDataflowNodes/ChaosClothAssetDataflowNodes.Build.cs
--- a/Engine/Plugins/ChaosClothAssetEditor/Source/ChaosClothAssetDataflowNodes/ChaosClothAssetDataflowNodes.Build.cs
+++ b/Engine/Plugins/ChaosClothAssetEditor/Source/ChaosClothAssetDataflowNodes/ChaosClothAssetDataflowNodes.Build.cs
@@ -17,6 +17,7 @@ public class ChaosClothAssetDataflowNodes : ModuleRules
 		PrivateDependencyModuleNames.AddRange(
 			new string[]
 			{
+				"AssetTools",
 				"Chaos",
 				"ChaosCaching",
 				"ChaosCloth",
diff --git a/Engine/Plugins/ChaosClothAssetEditor/Source/ChaosClothAssetDataflowNodes/Private/ChaosClothAsset/ClothDataflowTools.cpp b/Engine/Plugins/ChaosClothAssetEditor/Source/ChaosClothAssetDataflowNodes/Private/ChaosClothAsset/ClothDataflowTools.cpp
--- a/Engine/Plugins/ChaosClothAssetEditor/Source/ChaosClothAssetDataflowNodes/Private/ChaosClothAsset/ClothDataflowTools.cpp
+++ b/Engine/Plugins/ChaosClothAssetEditor/Source/ChaosClothAssetDataflowNodes/Private/ChaosClothAsset/ClothDataflowTools.cpp
@@ -563,22 +563,32 @@ namespace UE::Chaos::ClothAsset
 		return nullptr;
 	}
 
-	bool FClothDataflowTools::RemoveDegenerateTriangles(
-		const TArray<FIntVector3>& TriangleToVertexIndex,
-		const TArray<FVector2f>& RestPositions2D,
-		const TArray<FVector3f>& DrapedPositions3D,
-		TArray<FIntVector3>& OutTriangleToVertexIndex,
-		TArray<FVector2f>& OutRestPositions2D,
-		TArray<FVector3f>& OutDrapedPositions3D,
-		TArray<int32>& OutIndices)  // Old to new vertices lookup
+	FClothDataflowTools::FSimMeshCleanup::FSimMeshCleanup(const TArray<FIntVector3>& InTriangleToVertexIndex, const TArray<FVector2f>& InRestPositions2D, const TArray<FVector3f>& InDrapedPositions3D)
+		: TriangleToVertexIndex(InTriangleToVertexIndex)
+		, RestPositions2D(InRestPositions2D)
+		, DrapedPositions3D(InDrapedPositions3D)
 	{
-		bool bHasDegenerateTriangles = false;
+		check(RestPositions2D.Num() == DrapedPositions3D.Num());
+
+		OriginalTriangles.SetNum(TriangleToVertexIndex.Num());
+		for (int32 Index = 0; Index < OriginalTriangles.Num(); ++Index)
+		{
+			OriginalTriangles[Index].Add(Index);
+		}
+		OriginalVertices.SetNum(DrapedPositions3D.Num());
+		for (int32 Index = 0; Index < OriginalVertices.Num(); ++Index)
+		{
+			OriginalVertices[Index].Add(Index);
+		}
+	}
 
+	bool FClothDataflowTools::FSimMeshCleanup::RemoveDegenerateTriangles()
+	{
 		check(RestPositions2D.Num() == DrapedPositions3D.Num());
-		const int32 VertexCount = RestPositions2D.Num();
-		const int32 TriangleCount = TriangleToVertexIndex.Num();
 
-		OutTriangleToVertexIndex.Reset(TriangleCount);
+		bool bHasDegenerateTriangles = false;
+
+		const int32 VertexCount = RestPositions2D.Num();
 
 		// Remap[Index] is the index of the first vertex in a group of degenerated triangles to be callapsed.
 		// When two groups of collapsed vertices are merged, the group with the greatest Remap[index] value must adopt the one from the other group.
@@ -617,6 +627,12 @@ namespace UE::Chaos::ClothAsset
 				}
 			};
 
+		const int32 TriangleCount = TriangleToVertexIndex.Num();
+		TArray<FIntVector3> OutTriangleToVertexIndex;
+		OutTriangleToVertexIndex.Reserve(TriangleCount);
+		TArray<TSet<int32>> OutOriginalTriangles;
+		OutOriginalTriangles.Reserve(TriangleCount);
+
 		for (int32 TriangleIndex = 0; TriangleIndex < TriangleCount; ++TriangleIndex)
 		{
 			const int32 Index0 = TriangleToVertexIndex[TriangleIndex][0];
@@ -650,41 +666,58 @@ namespace UE::Chaos::ClothAsset
 			else
 			{
 				OutTriangleToVertexIndex.Emplace(TriangleToVertexIndex[TriangleIndex]);
+				OutOriginalTriangles.Emplace(OriginalTriangles[TriangleIndex]);
 			}
 		}
 
-		const int32 OutTriangleCount = OutTriangleToVertexIndex.Num();
+		TriangleToVertexIndex = MoveTemp(OutTriangleToVertexIndex);
+		OriginalTriangles = MoveTemp(OutOriginalTriangles);
+
+		const int32 OutTriangleCount = TriangleToVertexIndex.Num();
 		bHasDegenerateTriangles = (TriangleCount != OutTriangleCount);
 
 		UE_CLOG(bHasDegenerateTriangles, LogChaosClothAssetDataflowNodes, Display,
 			TEXT("USD import found and removed %d degenerated triangles out of %d source triangles."), TriangleCount - OutTriangleCount, TriangleCount);
 
-		OutRestPositions2D.Reset(OutVertexCount);
-		OutDrapedPositions3D.Reset(OutVertexCount);
-		OutIndices.Reset(VertexCount);
+		// Reconstruct vertices
+		TArray<FVector2f> OutRestPositions2D;
+		OutRestPositions2D.Reserve(OutVertexCount);
+		TArray<FVector3f> OutDrapedPositions3D;
+		OutDrapedPositions3D.Reserve(OutVertexCount);
+		TArray<TSet<int32>> OutOriginalVertices;
+		OutOriginalVertices.Reserve(OutVertexCount);
+		TArray<int32> OutIndices;
+		OutIndices.Reserve(VertexCount);
 		int32 OutIndex = -1;
 
 		for (int32 VertexIndex = 0; VertexIndex < VertexCount; ++VertexIndex)
 		{
 			if (Remap[VertexIndex] == VertexIndex)
 			{
-				OutRestPositions2D.Add(RestPositions2D[Remap[VertexIndex]]);
-				OutDrapedPositions3D.Add(DrapedPositions3D[Remap[VertexIndex]]);
-				OutIndices.Add(++OutIndex);
+				OutRestPositions2D.Emplace(RestPositions2D[VertexIndex]);
+				OutDrapedPositions3D.Emplace(DrapedPositions3D[VertexIndex]);
+				OutOriginalVertices.Emplace(OriginalVertices[VertexIndex]);
+				OutIndices.Emplace(++OutIndex);
 			}
 			else
 			{
 				const int32 OutRemappedIndex = OutIndices[Remap[VertexIndex]];
-				OutIndices.Add(OutRemappedIndex);
+				OutOriginalVertices[OutRemappedIndex].Append(OriginalVertices[VertexIndex]);
+				OutIndices.Emplace(OutRemappedIndex);
 			}
 		}
 		ensure(OutIndex + 1 == OutVertexCount);
 
+		RestPositions2D = MoveTemp(OutRestPositions2D);
+		DrapedPositions3D = MoveTemp(OutDrapedPositions3D);
+		OriginalVertices = MoveTemp(OutOriginalVertices);
+
+		// Remap final triangles
 		for (int32 TriangleIndex = 0; TriangleIndex < OutTriangleCount; ++TriangleIndex)
 		{
-			int32& Index0 = OutTriangleToVertexIndex[TriangleIndex][0];
-			int32& Index1 = OutTriangleToVertexIndex[TriangleIndex][1];
-			int32& Index2 = OutTriangleToVertexIndex[TriangleIndex][2];
+			int32& Index0 = TriangleToVertexIndex[TriangleIndex][0];
+			int32& Index1 = TriangleToVertexIndex[TriangleIndex][1];
+			int32& Index2 = TriangleToVertexIndex[TriangleIndex][2];
 
 			Index0 = OutIndices[Index0];
 			Index1 = OutIndices[Index1];
@@ -701,17 +734,19 @@ namespace UE::Chaos::ClothAsset
 		return bHasDegenerateTriangles;
 	}
 
-	bool FClothDataflowTools::RemoveDuplicateTriangles(TArray<FIntVector3>& TriangleToVertexIndex)
+	bool FClothDataflowTools::FSimMeshCleanup::RemoveDuplicateTriangles()
 	{
 		bool bHasDuplicatedTriangles = false;
 
 		const int32 TriangleCount = TriangleToVertexIndex.Num();
 
-		TSet<FIntVector3> Triangles;
+		TMap<FIntVector3, int32> Triangles;
 		Triangles.Reserve(TriangleCount);
 
 		TArray<FIntVector3> OutTriangleToVertexIndex;
 		OutTriangleToVertexIndex.Reserve(TriangleCount);
+		TArray<TSet<int32>> OutOriginalTriangles;
+		OutOriginalTriangles.Reserve(TriangleCount);
 
 		auto GetSortedIndices = [](const FIntVector3& TriangleIndices)->FIntVector3
 			{
@@ -729,27 +764,66 @@ namespace UE::Chaos::ClothAsset
 			const FIntVector3& TriangleIndices = TriangleToVertexIndex[Index];
 			const FIntVector3 TriangleSortedIndices = GetSortedIndices(TriangleIndices);
 
-			bool bIsAlreadyInSet;
-			Triangles.FindOrAdd(TriangleSortedIndices, &bIsAlreadyInSet);
-
-			if (bIsAlreadyInSet)
+			if (int32* NewTriangle = Triangles.Find(TriangleSortedIndices))
 			{
 				bHasDuplicatedTriangles = true;
+				OutOriginalTriangles[*NewTriangle].Append(OriginalTriangles[Index]);
 			}
 			else
 			{
-				OutTriangleToVertexIndex.Emplace(TriangleIndices);
+				NewTriangle = &Triangles.Emplace(TriangleSortedIndices);
+				*NewTriangle = OutTriangleToVertexIndex.Emplace(TriangleIndices);
+				OutOriginalTriangles.Emplace(OriginalTriangles[Index]);
 			}
 		}
 
-		UE_CLOG(bHasDuplicatedTriangles, LogChaosClothAssetDataflowNodes, Display,
-			TEXT("USD import found and removed %d duplicated triangles out of %d source triangles."), TriangleCount - OutTriangleToVertexIndex.Num(), TriangleCount);
-
 		TriangleToVertexIndex = MoveTemp(OutTriangleToVertexIndex);
+		OriginalTriangles = MoveTemp(OutOriginalTriangles);
+
+		UE_CLOG(bHasDuplicatedTriangles, LogChaosClothAssetDataflowNodes, Display,
+			TEXT("USD import found and removed %d duplicated triangles out of %d source triangles."), TriangleCount - TriangleToVertexIndex.Num(), TriangleCount);
 
 		return bHasDuplicatedTriangles;
 	}
 
+	template<typename T UE_REQUIRES_DEFINITION(std::is_same_v<T, TArray<int32>> || std::is_same_v<T, TSet<int32>>)>
+	TArray<int32> FClothDataflowTools::GetOriginalToNewIndices(const TConstArrayView<T>& NewToOriginals, int32 NumOriginalIndices)
+	{
+		TArray<int32> OriginalToNewIndices;
+		OriginalToNewIndices.Init(INDEX_NONE, NumOriginalIndices);
+
+		for (int32 NewIndex = 0; NewIndex < NewToOriginals.Num(); ++NewIndex)
+		{
+			for (const int32 OriginalIndex : NewToOriginals[NewIndex])
+			{
+				check(OriginalToNewIndices.IsValidIndex(OriginalIndex));
+				check(OriginalToNewIndices[OriginalIndex] == INDEX_NONE);
+				OriginalToNewIndices[OriginalIndex] = NewIndex;
+			}
+		}
+		return OriginalToNewIndices;
+	}
+	template TArray<int32> FClothDataflowTools::GetOriginalToNewIndices<TArray<int32>>(const TConstArrayView<TArray<int32>>& NewToOriginals, int32 NumOriginalIndices);
+	template TArray<int32> FClothDataflowTools::GetOriginalToNewIndices<TSet<int32>>(const TConstArrayView<TSet<int32>>& NewToOriginals, int32 NumOriginalIndices);
+
+	bool FClothDataflowTools::RemoveDegenerateTriangles(
+		const TArray<FIntVector3>& TriangleToVertexIndex,
+		const TArray<FVector2f>& RestPositions2D,
+		const TArray<FVector3f>& DrapedPositions3D,
+		TArray<FIntVector3>& OutTriangleToVertexIndex,
+		TArray<FVector2f>& OutRestPositions2D,
+		TArray<FVector3f>& OutDrapedPositions3D,
+		TArray<int32>& OutIndices)
+	{
+		FSimMeshCleanup SimMeshCleanup(TriangleToVertexIndex, RestPositions2D, DrapedPositions3D);
+		const bool bHasDegenerateTriangles = SimMeshCleanup.RemoveDegenerateTriangles();
+		OutIndices = GetOriginalToNewIndices<TSet<int32>>(SimMeshCleanup.OriginalVertices, DrapedPositions3D.Num());
+		OutTriangleToVertexIndex = MoveTemp(SimMeshCleanup.TriangleToVertexIndex);
+		OutRestPositions2D = MoveTemp(SimMeshCleanup.RestPositions2D);
+		OutDrapedPositions3D = MoveTemp(SimMeshCleanup.DrapedPositions3D);
+		return bHasDegenerateTriangles;
+	}
+
 	bool FClothDataflowTools::RemoveDuplicateStitches(TArray<TArray<FIntVector2>>& SeamStitches)
 	{
 		bool bHasDuplicateStitches = false;
diff --git a/Engine/Plugins/ChaosClothAssetEditor/Source/ChaosClothAssetDataflowNodes/Private/ChaosClothAsset/ClothDataflowTools.h b/Engine/Plugins/ChaosClothAssetEditor/Source/ChaosClothAssetDataflowNodes/Private/ChaosClothAsset/ClothDataflowTools.h
--- a/Engine/Plugins/ChaosClothAssetEditor/Source/ChaosClothAssetDataflowNodes/Private/ChaosClothAsset/ClothDataflowTools.h
+++ b/Engine/Plugins/ChaosClothAssetEditor/Source/ChaosClothAssetDataflowNodes/Private/ChaosClothAsset/ClothDataflowTools.h
@@ -2,8 +2,9 @@
 
 #pragma once
 
-#include "Templates/SharedPointer.h"
+#include "Math/Vector.h"
 #include "Logging/LogMacros.h"
+#include "Templates/SharedPointer.h"
 
 struct FManagedArrayCollection;
 struct FMeshBuildSettings;
@@ -45,6 +46,29 @@ namespace UE::Chaos::ClothAsset
 			return static_cast<FDataflowNode*>(GetPropertyOwnerDataflowNode(PropertyHandle, T::StaticStruct()));
 		}
 
+		/** Simulation mesh cleanup tools. */
+		struct FSimMeshCleanup
+		{
+			TArray<FIntVector3> TriangleToVertexIndex;
+			TArray<FVector2f> RestPositions2D;
+			TArray<FVector3f> DrapedPositions3D;
+			TArray<TSet<int32>> OriginalTriangles;  // New to original face index lookup
+			TArray<TSet<int32>> OriginalVertices;  // New to original vertex index lookup
+
+			FSimMeshCleanup(
+				const TArray<FIntVector3>& InTriangleToVertexIndex,
+				const TArray<FVector2f>& InRestPositions2D,
+				const TArray<FVector3f>& InDrapedPositions3D);
+
+			bool RemoveDegenerateTriangles();
+			bool RemoveDuplicateTriangles();
+		};
+
+		/** Returns the inverse mapping from a reduced set of original indices to the new indices. */
+		template<typename T UE_REQUIRES(std::is_same_v<T, TArray<int32>> || std::is_same_v<T, TSet<int32>>)>
+		static TArray<int32> GetOriginalToNewIndices(const TConstArrayView<T>& NewToOriginals, int32 NumOriginalIndices);
+
+		UE_DEPRECATED(5.6, "Use FSimMeshCleanup instead.")
 		static bool RemoveDegenerateTriangles(
 			const TArray<FIntVector3>& TriangleToVertexIndex,
 			const TArray<FVector2f>& RestPositions2D,
@@ -53,7 +77,16 @@ namespace UE::Chaos::ClothAsset
 			TArray<FVector2f>& OutRestPositions2D,
 			TArray<FVector3f>& OutDrapedPositions3D,
 			TArray<int32>& OutIndices);  // Old to new vertices lookup
-		static bool RemoveDuplicateTriangles(TArray<FIntVector3>& TriangleToVertexIndex);
+
+		UE_DEPRECATED(5.6, "Use FSimMeshCleanup instead.")
+		static bool RemoveDuplicateTriangles(TArray<FIntVector3>& TriangleToVertexIndex)
+		{
+			FSimMeshCleanup SimMeshCleanup(TriangleToVertexIndex, TArray<FVector2f>(), TArray<FVector3f>());
+			const bool bHasDuplicateTriangles = SimMeshCleanup.RemoveDuplicateTriangles();
+			TriangleToVertexIndex = MoveTemp(SimMeshCleanup.TriangleToVertexIndex);
+			return bHasDuplicateTriangles;
+		}
+
 		static bool RemoveDuplicateStitches(TArray<TArray<FIntVector2>>& SeamStitches);
 
 	private:
diff --git a/Engine/Plugins/ChaosClothAssetEditor/Source/ChaosClothAssetDataflowNodes/Private/ChaosClothAsset/RenderMeshImport.cpp b/Engine/Plugins/ChaosClothAssetEditor/Source/ChaosClothAssetDataflowNodes/Private/ChaosClothAsset/RenderMeshImport.cpp
new file mode 100644
--- /dev/null
+++ b/Engine/Plugins/ChaosClothAssetEditor/Source/ChaosClothAssetDataflowNodes/Private/ChaosClothAsset/RenderMeshImport.cpp
@@ -0,0 +1,365 @@
+// Copyright Epic Games, Inc. All Rights Reserved.
+
+#include "ChaosClothAsset/RenderMeshImport.h"
+#include "ChaosClothAsset/ClothCollectionGroup.h"
+#include "ChaosClothAsset/CollectionClothRenderPatternFacade.h"
+#include "ChaosClothAsset/CollectionClothFacade.h"
+#include "Engine/EngineTypes.h"
+#include "Engine/StaticMesh.h"
+#include "Materials/MaterialInterface.h"
+#include "MeshDescription.h"
+#include "MeshUVChannelInfo.h"
+#include "StaticMeshAttributes.h"
+#include "StaticMeshOperations.h"
+
+namespace UE::Chaos::ClothAsset
+{
+	namespace Private
+	{
+		struct FBuildVertex
+		{
+			FVertexID OrigVertID;
+			FVertexInstanceID OrigVertInstanceID;
+
+			FVector3f Position;
+			TArray<FVector2f> UVs;
+			FVector3f Normal;
+			FVector3f Tangent;
+			float BiNormalSign;
+			FVector4f Color;
+
+			bool operator==(const FBuildVertex& Other) const
+			{
+				if (OrigVertID != Other.OrigVertID)
+				{
+					return false;
+				}
+
+				// No need to check position since it's from OrigVertID
+
+				if (!ensure(UVs.Num() == Other.UVs.Num()))
+				{
+					return false;
+				}
+				for (int32 UVChannelIndex = 0; UVChannelIndex < UVs.Num(); ++UVChannelIndex)
+				{
+					if (!UVs[UVChannelIndex].Equals(Other.UVs[UVChannelIndex], UE_THRESH_UVS_ARE_SAME))
+					{
+						return false;
+					}
+				}
+				if (!Normal.Equals(Other.Normal, UE_THRESH_NORMALS_ARE_SAME))
+				{
+					return false;
+				}
+				if (!Tangent.Equals(Other.Tangent, UE_THRESH_NORMALS_ARE_SAME))
+				{
+					return false;
+				}
+				if (BiNormalSign != Other.BiNormalSign) // I think these are just -1 or 1, so just comparing them
+				{
+					return false;
+				}
+				// It looks like we just quantize the 0-1 color values to uint8s. (Call FLinearColor::ToFColor(bSRGB = false)) when we consume these.
+				// Going to be a little more strict in case we ever decide to switch to the gamma conversion.
+				if (!Color.Equals(Other.Color, UE_THRESH_NORMALS_ARE_SAME))
+				{
+					return false;
+				}
+				return true;
+			}
+		};
+
+		static void MergeVertexInstances(const FMeshDescription* const MeshDescription, TArray<FBuildVertex>& MergedVertices, TArray<int32>& VertexInstanceToMerged)
+		{
+			FStaticMeshConstAttributes Attributes(*MeshDescription);
+
+			TVertexAttributesConstRef<FVector3f> VertexPositions = Attributes.GetVertexPositions();
+
+			TVertexInstanceAttributesConstRef<FVector2f> VertexInstanceUVs = Attributes.GetVertexInstanceUVs();
+			TVertexInstanceAttributesConstRef<FVector3f> VertexInstanceNormals = Attributes.GetVertexInstanceNormals();
+			TVertexInstanceAttributesConstRef<FVector3f> VertexInstanceTangents = Attributes.GetVertexInstanceTangents();
+			TVertexInstanceAttributesConstRef<float> VertexInstanceBiNormalSigns = Attributes.GetVertexInstanceBinormalSigns();
+			TVertexInstanceAttributesConstRef<FVector4f> VertexInstanceColors = Attributes.GetVertexInstanceColors();
+
+			// Merge identical vertex instances that correspond with the same vertex.
+			auto MakeBuildVertex = [&VertexPositions, &VertexInstanceUVs, &VertexInstanceNormals, &VertexInstanceTangents, &VertexInstanceBiNormalSigns, &VertexInstanceColors](FVertexID VertID, FVertexInstanceID VertexInstanceID)
+				{
+					FBuildVertex BuildVertex;
+					BuildVertex.OrigVertID = VertID;
+					BuildVertex.OrigVertInstanceID = VertexInstanceID;
+					BuildVertex.Position = VertexPositions[VertID];
+					BuildVertex.UVs.SetNumUninitialized(VertexInstanceUVs.GetNumChannels());
+					for (int32 UVChannelIndex = 0; UVChannelIndex < VertexInstanceUVs.GetNumChannels(); ++UVChannelIndex)
+					{
+						BuildVertex.UVs[UVChannelIndex] = VertexInstanceUVs.Get(VertexInstanceID, UVChannelIndex);
+					}
+					BuildVertex.Normal = VertexInstanceNormals[VertexInstanceID];
+					BuildVertex.Tangent = VertexInstanceTangents[VertexInstanceID];
+					BuildVertex.BiNormalSign = VertexInstanceBiNormalSigns[VertexInstanceID];
+					BuildVertex.Color = VertexInstanceColors[VertexInstanceID];
+					return BuildVertex;
+				};
+			MergedVertices.Reset(MeshDescription->VertexInstances().GetArraySize());
+			VertexInstanceToMerged.Init(INDEX_NONE, MeshDescription->VertexInstances().GetArraySize());
+
+			for (const FVertexID& VertID : MeshDescription->Vertices().GetElementIDs())
+			{
+				TConstArrayView<FVertexInstanceID> VertexInstances = MeshDescription->GetVertexVertexInstanceIDs(VertID);
+				if (VertexInstances.IsEmpty())
+				{
+					continue;
+				}
+				const int32 FirstMergedVert = MergedVertices.Add(MakeBuildVertex(VertID, VertexInstances[0]));
+				VertexInstanceToMerged[VertexInstances[0].GetValue()] = FirstMergedVert;
+				for (int32 LocalInstance = 1; LocalInstance < VertexInstances.Num(); ++LocalInstance)
+				{
+					FBuildVertex BuildVert = MakeBuildVertex(VertID, VertexInstances[LocalInstance]);
+					bool bFoundDuplicate = false;
+					for (int32 CompareIndex = FirstMergedVert; CompareIndex < MergedVertices.Num(); ++CompareIndex)
+					{
+						if (BuildVert == MergedVertices[CompareIndex])
+						{
+							bFoundDuplicate = true;
+							VertexInstanceToMerged[VertexInstances[LocalInstance].GetValue()] = CompareIndex;
+							break;
+						}
+					}
+
+					if (!bFoundDuplicate)
+					{
+						const int32 AddedMergedVert = MergedVertices.Add(MoveTemp(BuildVert));
+						VertexInstanceToMerged[VertexInstances[LocalInstance].GetValue()] = AddedMergedVert;
+					}
+				}
+			}
+		}
+	}  // End namespace Private
+
+	FRenderMeshImport::FRenderMeshImport(const FMeshDescription& InMeshDescription, const FMeshBuildSettings& BuildSettings)
+	{
+		using namespace UE::Chaos::ClothAsset::Private;
+
+		const FMeshDescription* MeshDescription = &InMeshDescription;
+		FMeshDescription WritableMeshDescription;
+		if (BuildSettings.bRecomputeNormals || BuildSettings.bRecomputeTangents)
+		{
+			// check if any are invalid
+			bool bHasInvalidNormals, bHasInvalidTangents;
+			FStaticMeshOperations::HasInvalidVertexInstanceNormalsOrTangents(InMeshDescription, bHasInvalidNormals, bHasInvalidTangents);
+
+			// if neither are invalid we are not going to recompute
+			if (bHasInvalidNormals || bHasInvalidTangents)
+			{
+				WritableMeshDescription = InMeshDescription;
+				MeshDescription = &WritableMeshDescription;
+				FStaticMeshAttributes Attributes(WritableMeshDescription);
+				if (!Attributes.GetTriangleNormals().IsValid() || !Attributes.GetTriangleTangents().IsValid())
+				{
+					// If these attributes don't exist, create them and compute their values for each triangle
+					FStaticMeshOperations::ComputeTriangleTangentsAndNormals(WritableMeshDescription);
+				}
+
+				EComputeNTBsFlags ComputeNTBsOptions = EComputeNTBsFlags::BlendOverlappingNormals;
+				ComputeNTBsOptions |= BuildSettings.bRecomputeNormals ? EComputeNTBsFlags::Normals : EComputeNTBsFlags::None;
+				ComputeNTBsOptions |= BuildSettings.bRecomputeTangents ? EComputeNTBsFlags::Tangents : EComputeNTBsFlags::None;
+				ComputeNTBsOptions |= BuildSettings.bUseMikkTSpace ? EComputeNTBsFlags::UseMikkTSpace : EComputeNTBsFlags::None;
+				ComputeNTBsOptions |= BuildSettings.bComputeWeightedNormals ? EComputeNTBsFlags::WeightedNTBs : EComputeNTBsFlags::None;
+				ComputeNTBsOptions |= BuildSettings.bRemoveDegenerates ? EComputeNTBsFlags::IgnoreDegenerateTriangles : EComputeNTBsFlags::None;
+
+				FStaticMeshOperations::ComputeTangentsAndNormals(WritableMeshDescription, ComputeNTBsOptions);
+			}
+		}
+
+		// Merge vertex instances that share the same vertex. These will become the pattern vertices.
+		TArray<FBuildVertex> MergedVertices;
+		TArray<int32> VertexInstanceToMerged;
+		MergeVertexInstances(MeshDescription, MergedVertices, VertexInstanceToMerged);
+
+		// Vertex data (this will MoveTemp stuff out of MergedVertices!!)
+		TArray<FVertex> Vertices;
+		if (const int32 NumVertices = MergedVertices.Num())
+		{
+			Vertices.SetNumUninitialized(NumVertices);
+
+			for (int32 VertexIndex = 0; VertexIndex < NumVertices; ++VertexIndex)
+			{
+				FBuildVertex& BuildVertex = MergedVertices[VertexIndex];
+				FVertex& Vertex = Vertices[VertexIndex];
+				Vertex.RenderPosition = BuildVertex.Position;
+				Vertex.RenderNormal = BuildVertex.Normal;
+				Vertex.RenderTangentU = BuildVertex.Tangent;
+				Vertex.RenderTangentV = FVector3f::CrossProduct(BuildVertex.Normal, BuildVertex.Tangent).GetSafeNormal() * BuildVertex.BiNormalSign;
+				new (&Vertex.RenderUVs) TArray<FVector2f>(MoveTemp(BuildVertex.UVs));
+				Vertex.RenderColor = FLinearColor(BuildVertex.Color);
+				Vertex.OriginalIndex = BuildVertex.OrigVertID.GetValue();
+			}
+		}
+
+		// Face data
+		TArray<FTriangle> Triangles;
+		if (const int32 NumTriangles = MeshDescription->Triangles().Num())
+		{
+			Triangles.SetNumUninitialized(NumTriangles);
+
+			int32 FaceIndex = 0;
+			for (const FTriangleID TriangleID : MeshDescription->Triangles().GetElementIDs())
+			{
+				const FPolygonGroupID PolygonGroupID = MeshDescription->GetTrianglePolygonGroup(TriangleID);
+				const TConstArrayView<FVertexInstanceID> VertexInstances = MeshDescription->GetTriangleVertexInstances(TriangleID);
+				check(VertexInstances.Num() == 3);
+				check(MergedVertices.IsValidIndex(VertexInstanceToMerged[VertexInstances[0].GetValue()]));
+				check(MergedVertices.IsValidIndex(VertexInstanceToMerged[VertexInstances[1].GetValue()]));
+				check(MergedVertices.IsValidIndex(VertexInstanceToMerged[VertexInstances[2].GetValue()]));
+				FTriangle& Triangle = Triangles[FaceIndex];
+				Triangle.VertexIndices = FIntVector3(
+					VertexInstanceToMerged[VertexInstances[0].GetValue()],
+					VertexInstanceToMerged[VertexInstances[1].GetValue()],
+					VertexInstanceToMerged[VertexInstances[2].GetValue()]);
+				Triangle.OriginalIndex = TriangleID.GetValue();
+				Triangle.MaterialIndex = PolygonGroupID.GetValue();
+				++FaceIndex;
+			}
+
+			// Sort the triangles by material index
+			Triangles.Sort([](const FTriangle& Triangle0, const FTriangle& Triangle1) -> bool
+				{
+					return Triangle0.MaterialIndex < Triangle1.MaterialIndex;
+				});
+
+			// Split into material sections
+			TArray<int32> RemappedVertexIndices;
+			int32 SectionTrianglesStart = 0;
+			int32 SectionMaterialIndex = Triangles[SectionTrianglesStart].MaterialIndex;
+			int32 SectionMinVertex = Triangles[SectionTrianglesStart].VertexIndices.GetMin();
+			int32 SectionMaxVertex = Triangles[SectionTrianglesStart].VertexIndices.GetMax();
+
+			for (int32 TriangleIndex = SectionTrianglesStart + 1; ; ++TriangleIndex)
+			{
+				const bool bEndIteration = (TriangleIndex == Triangles.Num());
+				if (bEndIteration || SectionMaterialIndex != Triangles[TriangleIndex].MaterialIndex)
+				{
+					// Add a new section with this triangle range
+					FSection& Section = Sections.Emplace(SectionMaterialIndex);
+					const int32 SectionTrianglesCount = TriangleIndex - SectionTrianglesStart;
+					Section.Triangles = TConstArrayView<FTriangle>(Triangles.GetData() + SectionTrianglesStart, SectionTrianglesCount);
+					Section.Vertices.Reserve(SectionMaxVertex - SectionMinVertex + 1);
+
+					// Copy the section's vertices and remap the section's triangle vertex indices
+					RemappedVertexIndices.Init(INDEX_NONE, Vertices.Num());
+					int32 SectionNumTexCoords = MAX_TEXCOORDS;
+					for (FTriangle& Triangle : Section.Triangles)
+					{
+						for (int32 Index = 0; Index < 3; ++Index)
+						{
+							int32& TriangleVertexIndex = Triangle.VertexIndices[Index];
+							if (RemappedVertexIndices[TriangleVertexIndex] == INDEX_NONE)
+							{
+								RemappedVertexIndices[TriangleVertexIndex] = Section.Vertices.Emplace(Vertices[TriangleVertexIndex]);
+								SectionNumTexCoords = FMath::Min(SectionNumTexCoords, Vertices[TriangleVertexIndex].RenderUVs.Num());
+							}
+							TriangleVertexIndex = RemappedVertexIndices[TriangleVertexIndex];
+						}
+					}
+					Section.NumTexCoords = SectionNumTexCoords;
+
+					if (bEndIteration)
+					{
+						break;
+					}
+
+					// Prepare next section
+					SectionTrianglesStart = TriangleIndex;
+					SectionMaterialIndex = Triangles[SectionTrianglesStart].MaterialIndex;
+					SectionMinVertex = Triangles[SectionTrianglesStart].VertexIndices.GetMin();
+					SectionMaxVertex = Triangles[SectionTrianglesStart].VertexIndices.GetMax();
+					continue;
+				}
+				// Update Min/Max
+				SectionMinVertex = FMath::Min(SectionMinVertex, Triangles[TriangleIndex].VertexIndices.GetMin());
+				SectionMaxVertex = FMath::Max(SectionMaxVertex, Triangles[TriangleIndex].VertexIndices.GetMax());
+			}
+		}
+	}
+
+	void FRenderMeshImport::AddRenderSections(
+		const TSharedRef<FManagedArrayCollection> ClothCollection,
+		const TArray<FStaticMaterial>& Materials,
+		const FName OriginalTrianglesName,
+		const FName OriginalVerticesName)
+	{
+		FCollectionClothFacade ClothFacade(ClothCollection);
+		if (!ClothFacade.IsValid())
+		{
+			ClothFacade.DefineSchema();
+		}
+
+		// Add support for original indices
+		ClothFacade.AddUserDefinedAttribute<TArray<int32>>(OriginalTrianglesName, ClothCollectionGroup::RenderFaces);
+		ClothFacade.AddUserDefinedAttribute<TArray<int32>>(OriginalVerticesName, ClothCollectionGroup::RenderVertices);
+
+		for (const TPair<int32, FSection>& Section : Sections)
+		{
+			const TArray<FVertex>& Vertices = Section.Value.Vertices;
+			const TArray<FTriangle>& Triangles = Section.Value.Triangles;
+			const int32 NumVertices = Vertices.Num();
+			const int32 NumTriangles = Triangles.Num();
+			if (NumVertices && NumTriangles)
+			{
+				FCollectionClothRenderPatternFacade ClothPatternFacade = ClothFacade.AddGetRenderPattern();
+				ClothPatternFacade.SetNumRenderVertices(NumVertices);
+				ClothPatternFacade.SetNumRenderFaces(NumTriangles);
+
+				TArrayView<FVector3f> RenderPosition = ClothPatternFacade.GetRenderPosition();
+				TArrayView<FVector3f> RenderNormal = ClothPatternFacade.GetRenderNormal();
+				TArrayView<FVector3f> RenderTangentU = ClothPatternFacade.GetRenderTangentU();
+				TArrayView<FVector3f> RenderTangentV = ClothPatternFacade.GetRenderTangentV();
+				TArrayView<TArray<FVector2f>> RenderUVs = ClothPatternFacade.GetRenderUVs();
+				TArrayView<FLinearColor> RenderColor = ClothPatternFacade.GetRenderColor();
+				TArrayView<TArray<int32>> RenderBoneIndices = ClothPatternFacade.GetRenderBoneIndices();
+				TArrayView<TArray<float>> RenderBoneWeights = ClothPatternFacade.GetRenderBoneWeights();
+				const int32 VertexOffset = ClothPatternFacade.GetRenderVerticesOffset();
+				const int32 FaceOffset = ClothPatternFacade.GetRenderFacesOffset();
+				TArrayView<TArray<int32>> OriginalTriangles = ClothFacade.GetUserDefinedAttribute<TArray<int32>>(OriginalTrianglesName, ClothCollectionGroup::RenderFaces);   // Requery after adding a pattern
+				TArrayView<TArray<int32>> OriginalVertices = ClothFacade.GetUserDefinedAttribute<TArray<int32>>(OriginalVerticesName, ClothCollectionGroup::RenderVertices);  // in case the array gets reallocated
+
+				const int32 NumTexCoords = FMath::Min((int32)MAX_TEXCOORDS, Section.Value.NumTexCoords);
+				for (int32 VertexIndex = 0; VertexIndex < NumVertices; ++VertexIndex)
+				{
+					const FVertex& Vertex = Vertices[VertexIndex];
+					RenderPosition[VertexIndex] = Vertex.RenderPosition;
+					RenderNormal[VertexIndex] = Vertex.RenderNormal;
+					RenderTangentU[VertexIndex] = Vertex.RenderTangentU;
+					RenderTangentV[VertexIndex] = Vertex.RenderTangentV;
+					RenderUVs[VertexIndex].SetNum(NumTexCoords);
+					for (int32 TexCoordIndex = 0; TexCoordIndex < NumTexCoords; ++TexCoordIndex)
+					{
+						RenderUVs[VertexIndex][TexCoordIndex] = Vertex.RenderUVs[TexCoordIndex];
+					}
+
+					RenderColor[VertexIndex] = FLinearColor(Vertex.RenderColor);
+
+					constexpr int32 NumBones = 0;
+					RenderBoneIndices[VertexIndex].SetNum(NumBones);
+					RenderBoneWeights[VertexIndex].SetNum(NumBones);
+
+					OriginalVertices[VertexIndex + VertexOffset] = { Vertex.OriginalIndex };
+				}
+
+				TArrayView<FIntVector3> RenderIndices = ClothPatternFacade.GetRenderIndices();
+				for (int32 FaceIndex = 0; FaceIndex < NumTriangles; ++FaceIndex)
+				{
+					RenderIndices[FaceIndex] = Triangles[FaceIndex].VertexIndices + FIntVector3(VertexOffset);
+
+					OriginalTriangles[FaceIndex + FaceOffset] = { Triangles[FaceIndex].OriginalIndex };
+				}
+
+				const int32 MaterialIndex = Section.Key;
+				const FString RenderMaterialPathName = Materials[MaterialIndex].MaterialInterface ?
+					Materials[MaterialIndex].MaterialInterface->GetPathName() :
+					FString();
+				ClothPatternFacade.SetRenderMaterialPathName(RenderMaterialPathName);
+			}
+		}
+	}
+}  // End namespace UE::Chaos::ClothAsset
diff --git a/Engine/Plugins/ChaosClothAssetEditor/Source/ChaosClothAssetDataflowNodes/Private/ChaosClothAsset/RenderMeshImport.h b/Engine/Plugins/ChaosClothAssetEditor/Source/ChaosClothAssetDataflowNodes/Private/ChaosClothAsset/RenderMeshImport.h
new file mode 100644
--- /dev/null
+++ b/Engine/Plugins/ChaosClothAssetEditor/Source/ChaosClothAssetDataflowNodes/Private/ChaosClothAsset/RenderMeshImport.h
@@ -0,0 +1,57 @@
+// Copyright Epic Games, Inc. All Rights Reserved.
+
+#pragma once
+
+#include "Containers/Array.h"
+#include "Math/Vector.h"
+#include "Math/Color.h"
+#include "Templates/SharedPointer.h"
+
+struct FManagedArrayCollection;
+struct FMeshBuildSettings;
+struct FMeshDescription;
+struct FStaticMaterial;
+
+namespace UE::Chaos::ClothAsset
+{
+	/**
+	 * Render mesh import from mesh description.
+	 */
+	struct FRenderMeshImport
+	{
+		struct FVertex
+		{
+			FVector3f RenderPosition;
+			FVector3f RenderNormal;
+			FVector3f RenderTangentU;
+			FVector3f RenderTangentV;
+			TArray<FVector2f> RenderUVs;
+			FLinearColor RenderColor;
+			int32 OriginalIndex;
+		};
+
+		struct FTriangle
+		{
+			FIntVector3 VertexIndices;
+			int32 OriginalIndex;
+			int32 MaterialIndex;
+		};
+
+		struct FSection
+		{
+			TArray<FVertex> Vertices;
+			TArray<FTriangle> Triangles;
+			int32 NumTexCoords;
+		};
+		TSortedMap<int32, FSection> Sections;
+
+		FRenderMeshImport(const FMeshDescription& InMeshDescription, const FMeshBuildSettings& BuildSettings);
+
+		void AddRenderSections(
+			const TSharedRef<FManagedArrayCollection> ClothCollection,
+			const TArray<FStaticMaterial>& Materials,
+			const FName OriginalTrianglesName,
+			const FName OriginalVerticesName);
+	};
+}  // End namespace UE::Chaos::ClothAsset
+
diff --git a/Engine/Plugins/ChaosClothAssetEditor/Source/ChaosClothAssetDataflowNodes/Private/ChaosClothAsset/USDImportNode_v2.cpp b/Engine/Plugins/ChaosClothAssetEditor/Source/ChaosClothAssetDataflowNodes/Private/ChaosClothAsset/USDImportNode_v2.cpp
--- a/Engine/Plugins/ChaosClothAssetEditor/Source/ChaosClothAssetDataflowNodes/Private/ChaosClothAsset/USDImportNode_v2.cpp
+++ b/Engine/Plugins/ChaosClothAssetEditor/Source/ChaosClothAssetDataflowNodes/Private/ChaosClothAsset/USDImportNode_v2.cpp
@@ -2,12 +2,16 @@
 
 #include "ChaosClothAsset/USDImportNode_v2.h"
 #include "ChaosClothAsset/ClothAsset.h"
+#include "ChaosClothAsset/ClothCollectionGroup.h"
 #include "ChaosClothAsset/ClothDataflowTools.h"
 #include "ChaosClothAsset/ClothGeometryTools.h"
 #include "ChaosClothAsset/CollectionClothFacade.h"
+#include "ChaosClothAsset/CollectionClothSelectionFacade.h"
+#include "ChaosClothAsset/RenderMeshImport.h"
 #include "ChaosClothAsset/USDImportNode.h"
 #include "AssetRegistry/AssetData.h"
 #include "AssetRegistry/AssetRegistryModule.h"
+#include "Chaos/CollectionPropertyFacade.h"
 #include "Dataflow/DataflowInputOutput.h"
 #include "Dataflow/DataflowGraph.h"
 #include "Dataflow/DataflowObject.h"
@@ -28,10 +32,11 @@
 #include "USDStageImportOptions.h"
 #include "USDTypesConversion.h"
 #include "USDValueConversion.h"
-#include "Chaos/CollectionPropertyFacade.h"
 #include "UsdWrappers/UsdPrim.h"
+#include "UsdWrappers/UsdGeomSubset.h"
 #include "UsdWrappers/UsdRelationship.h"
 #include "UsdWrappers/VtValue.h"
+#include "AssetViewUtils.h"
 
 #include UE_INLINE_GENERATED_CPP_BY_NAME(USDImportNode_v2)
 
@@ -39,7 +44,19 @@
 
 namespace UE::Chaos::ClothAsset::Private
 {
-	static TArray<FSoftObjectPath> UsdClothOverrideMaterials_v2 {
+	// User attribute names
+	static const FName OriginalIndicesName(TEXT("OriginalIndices"));
+
+	// Cloth USD API names
+	static const FName ClothRootAPI(TEXT("ClothRootAPI"));
+	static const FName RenderPatternAPI(TEXT("RenderPatternAPI"));
+	static const FName SimMeshDataAPI(TEXT("SimMeshDataAPI"));
+	static const FName SimPatternAPI(TEXT("SimPatternAPI"));
+	static const FName SewingAPI(TEXT("SewingAPI"));
+
+	// USD import material overrides
+	static TArray<FSoftObjectPath> UsdClothOverrideMaterials_v2 =
+	{
 		FSoftObjectPath(TEXT("/ChaosClothAsset/Materials/USDImportMaterial.USDImportMaterial")),
 		FSoftObjectPath(TEXT("/ChaosClothAsset/Materials/USDImportTranslucentMaterial.USDImportTranslucentMaterial")),
 		FSoftObjectPath(TEXT("/ChaosClothAsset/Materials/USDImportTwoSidedMaterial.USDImportTwoSidedMaterial")),
@@ -47,6 +64,33 @@ namespace UE::Chaos::ClothAsset::Private
 		FSoftObjectPath(TEXT("/ChaosClothAsset/Materials/USDImportDisplayColorMaterial.USDImportDisplayColorMaterial")),
 	};
 	
+	// Return the specified UObject's dependencies for top level UAssets that are not in the engine folders
+	static TArray<UObject*> GetAssetDependencies(const UObject* Asset)
+	{
+		constexpr bool bInRequireDirectOuter = true;
+		constexpr bool bShouldIgnoreArchetype = true;
+		constexpr bool bInSerializeRecursively = false;  // Ignored if LimitOuter is nullptr
+		constexpr bool bShouldIgnoreTransient = true;
+		constexpr UObject* LimitOuter = nullptr;
+		TArray<UObject*> References;
+		FReferenceFinder ReferenceFinder(References, LimitOuter, bInRequireDirectOuter, bShouldIgnoreArchetype, bInSerializeRecursively, bShouldIgnoreTransient);
+		ReferenceFinder.FindReferences(const_cast<UObject*>(Asset));
+
+		TArray<UObject*> Dependencies;
+		Dependencies.Reserve(References.Num());
+		for (UObject* const Reference : References)
+		{
+			constexpr bool bEnginePluginIsAlsoEngine = true;  // Only includes non Engine or non Engine plugins assets (e.g. no USD materials)
+			if (FAssetData::IsUAsset(Reference) &&
+				FAssetData::IsTopLevelAsset(Reference) &&
+				!AssetViewUtils::IsEngineFolder(Reference->GetPackage()->GetName(), bEnginePluginIsAlsoEngine))
+			{
+				Dependencies.Emplace(Reference);
+			}
+		}
+		return Dependencies;
+	}
+
 	static void OverrideUsdImportMaterials_v2(const TArray<FSoftObjectPath>& Materials, TArray<FSoftObjectPath>* SavedValues = nullptr)
 	{
 		if (UUsdProjectSettings* UsdProjectSettings = GetMutableDefault<UUsdProjectSettings>())
@@ -68,7 +112,7 @@ namespace UE::Chaos::ClothAsset::Private
 		}
 	}
 
-	static void ImportStaticMeshesFromUsdStage(const UE::FUsdStage& UsdStage, const FString& UsdFilePath, const FString& PackagePath)
+	static void ImportStaticMeshesFromUsdStage(const FUsdStage& UsdStage, const FString& UsdFilePath, const FString& PackagePath)
 	{
 		// Import recognised assets
 		FUsdStageImportContext ImportContext;
@@ -128,32 +172,57 @@ namespace UE::Chaos::ClothAsset::Private
 		OverrideUsdImportMaterials_v2(OriginalUsdMaterials);
 	}
 
-	static UE::FUsdPrim FindClothPrim(const UE::FUsdPrim& RootPrim)
+	static FUsdPrim FindClothPrim(const FUsdPrim& RootPrim)
 	{
-		for (UE::FUsdPrim& ChildPrim : RootPrim.GetChildren())
+		for (FUsdPrim& ChildPrim : RootPrim.GetChildren())
 		{
-			static const FName ClothRootAPI(TEXT("ClothRootAPI"));
 			if (ChildPrim.HasAPI(ClothRootAPI))
 			{
 				return ChildPrim;
 			}
 		}
-		return UE::FUsdPrim();
+		return FUsdPrim();
 	}
 
-	static UE::FUsdPrim FindSimMeshPrim(const UE::FUsdPrim& ClothPrim)
+	static bool RemoveMaterialOpacity(const FUsdPrim& Prim)
 	{
-		for (UE::FUsdPrim& ClothChildPrim : ClothPrim.GetChildren())
+		bool bHasOpacity = false;
+		for (const FUsdPrim& ChildPrim : Prim.GetChildren())
+		{
+			if (ChildPrim.IsA(TEXT("Material")))
+			{
+				for (const FUsdPrim& GrandChildPrim : ChildPrim.GetChildren())
+				{
+					if (GrandChildPrim.IsA(TEXT("Shader")))
+					{
+						if (const FUsdAttribute OpacityAttr = GrandChildPrim.GetAttribute(TEXT("inputs:opacity")))
+						{
+							OpacityAttr.ClearConnections();
+							OpacityAttr.Clear();
+							bHasOpacity = true;
+						}
+					}
+				}
+			}
+			else
+			{
+				bHasOpacity = RemoveMaterialOpacity(ChildPrim) || bHasOpacity;
+			}
+		}
+		return bHasOpacity;
+	}
+
+	static FUsdPrim FindSimMeshPrim(const FUsdPrim& ClothPrim)
+	{
+		for (FUsdPrim& ClothChildPrim : ClothPrim.GetChildren())
 		{
 			if (ClothChildPrim.IsA(TEXT("Mesh")))
 			{
-				static const FName SimMeshDataAPI(TEXT("SimMeshDataAPI"));
 				if (ClothChildPrim.HasAPI(SimMeshDataAPI))
 				{
 					// Check that the sim mesh has at least one valid geomsubset patern
-					for (UE::FUsdPrim& SimMeshChildPrim : ClothChildPrim.GetChildren())
+					for (FUsdPrim& SimMeshChildPrim : ClothChildPrim.GetChildren())
 					{
-						static const FName SimPatternAPI(TEXT("SimPatternAPI"));
 						if (SimMeshChildPrim.IsA(TEXT("GeomSubset")) && SimMeshChildPrim.HasAPI(SimPatternAPI))
 						{
 							return ClothChildPrim;
@@ -162,20 +231,19 @@ namespace UE::Chaos::ClothAsset::Private
 				}
 			}
 		}
-		return UE::FUsdPrim();
+		return FUsdPrim();
 	}
 
-	static UE::FUsdPrim FindRenderMeshPrim(const UE::FUsdPrim& ClothPrim)
+	static FUsdPrim FindRenderMeshPrim(const FUsdPrim& ClothPrim)
 	{
-		for (UE::FUsdPrim& ClothChildPrim : ClothPrim.GetChildren())
+		for (FUsdPrim& ClothChildPrim : ClothPrim.GetChildren())
 		{
 			if (ClothChildPrim.IsA(TEXT("Mesh")))
 			{
 				// Look for all GeomSubsets to see if this is a suitable render mesh prim
 				bool bIsRenderMesh = false;
-				for (UE::FUsdPrim& RenderMeshChildPrim : ClothChildPrim.GetChildren())
+				for (FUsdPrim& RenderMeshChildPrim : ClothChildPrim.GetChildren())
 				{
-					static const FName RenderPatternAPI(TEXT("RenderPatternAPI"));
 					if (RenderMeshChildPrim.IsA(TEXT("GeomSubset")) && RenderMeshChildPrim.HasAPI(RenderPatternAPI))
 					{
 						return ClothChildPrim;
@@ -183,16 +251,16 @@ namespace UE::Chaos::ClothAsset::Private
 				}
 			}
 		}
-		return UE::FUsdPrim();
+		return FUsdPrim();
 	}
 
-	static FVector2f GetSimMeshUVScale(const UE::FUsdPrim& SimMeshPrim)
+	static FVector2f GetSimMeshUVScale(const FUsdPrim& SimMeshPrim)
 	{
 		FVector2f UVScale(1.f);
-		const UE::FUsdAttribute RestPositionScaleAttr = SimMeshPrim.GetAttribute(TEXT("restPositionScale"));
+		const FUsdAttribute RestPositionScaleAttr = SimMeshPrim.GetAttribute(TEXT("restPositionScale"));
 		if (RestPositionScaleAttr.HasValue() && RestPositionScaleAttr.GetTypeName() == TEXT("float2"))
 		{
-			UE::FVtValue Value;
+			FVtValue Value;
 			RestPositionScaleAttr.Get(Value);
 			UsdUtils::FConvertedVtValue ConvertedVtValue;
 			if (UsdToUnreal::ConvertValue(Value, ConvertedVtValue) &&
@@ -206,6 +274,217 @@ namespace UE::Chaos::ClothAsset::Private
 		}
 		return UVScale;
 	}
+
+	static FString GetStringValue(const FUsdAttribute& UsdAttribute)
+	{
+		if (UsdAttribute.HasValue())
+		{
+			FVtValue Value;
+			UsdAttribute.Get(Value);
+			return UsdUtils::Stringify(Value);
+		}
+		return FString();
+	}
+
+	static TArray<int32> GetIntArrayValues(const FUsdAttribute& UsdAttribute)
+	{
+		using namespace UsdUtils;
+		TArray<int32> IntArray;
+		if (UsdAttribute.HasValue())
+		{
+			FVtValue Value;
+			UsdAttribute.Get(Value);
+			FConvertedVtValue ConvertedVtValue;
+			if (UsdToUnreal::ConvertValue(Value, ConvertedVtValue) && ConvertedVtValue.bIsArrayValued && !ConvertedVtValue.bIsEmpty)
+			{
+				IntArray.Reserve(ConvertedVtValue.Entries.Num());
+				for (const FConvertedVtValueEntry& ValueEntry : ConvertedVtValue.Entries)
+				{
+					IntArray.Emplace(ValueEntry[0].Get<int32>());
+				}
+			}
+		}
+		return MoveTemp(IntArray);
+	}
+
+	static bool CheckSimMeshPrimTriangles(const FUsdPrim& SimMeshPrim, FText& OutErrorText)
+	{
+		const FUsdAttribute FaceVertexCountsAttr = SimMeshPrim.GetAttribute(TEXT("faceVertexCounts"));
+		if (!FaceVertexCountsAttr)
+		{
+			OutErrorText = LOCTEXT("MissingSimMeshFaceCountAttribute", "Missing simulation mesh faceVertexCounts attribute.");
+		}
+		else if (FaceVertexCountsAttr.GetTypeName() != TEXT("int[]"))
+		{
+			OutErrorText = LOCTEXT("WrongSimMeshFaceCountTypeName", "Wrong simulation mesh faceVertexCounts type name. Needs to be 'int[]'.");
+		}
+		else
+		{
+			bool bIsTriangleMesh = true;
+			const TArray<int32> FaceVertexCounts = GetIntArrayValues(FaceVertexCountsAttr);
+			for (int32 FaceVertexCount : FaceVertexCounts)
+			{
+				if (FaceVertexCount != 3)
+				{
+					OutErrorText = LOCTEXT("WrongSimMeshFaceCount", "Wrong simulation mesh face vertex count. The simulation mesh only supports '3' for triangles.");
+					bIsTriangleMesh = false;
+					break;
+				}
+			}
+			return bIsTriangleMesh;
+		}
+		return false;
+	}
+
+	static bool ImportPatternsFromMeshPrim(const FUsdPrim& MeshPrim, const FName PatternAPI, TMap<FName, TSet<int32>>& Patterns, FText& OutErrorText)
+	{
+		Patterns.Reset();
+		for (const FUsdPrim& MeshChildPrim : MeshPrim.GetChildren())
+		{
+			if (MeshChildPrim.IsA(TEXT("GeomSubset")) && MeshChildPrim.HasAPI(PatternAPI))
+			{
+				const FUsdGeomSubset GeomSubset(MeshChildPrim);
+
+				// Read FamillyName
+				const FUsdAttribute FamilyNameAttr = GeomSubset.GetFamilyNameAttr();
+				if (GetStringValue(FamilyNameAttr) != TEXT("pattern"))
+				{
+					OutErrorText = FText::Format(LOCTEXT("WrongPatternGeomSubsetFamilyName", "Wrong pattern family name for GeomSubset '{0}'. Needs to be 'pattern'."), FText::FromString(MeshPrim.GetPrimPath().GetString()));
+					return false;
+				}
+
+				// Read ElementTypeAttr
+				const FUsdAttribute ElementTypeAttr = GeomSubset.GetElementTypeAttr();
+				if (GetStringValue(ElementTypeAttr) != TEXT("face"))
+				{
+					OutErrorText = FText::Format(LOCTEXT("WrongPatternGeomSubsetType", "Wrong pattern type for GeomSubset '{0}'. Needs to be 'face'."), FText::FromString(MeshPrim.GetPrimPath().GetString()));
+					return false;
+				}
+
+				// Read indices 
+				const FUsdAttribute IndicesAttr = GeomSubset.GetIndicesAttr();
+				if (IndicesAttr.GetTypeName() != TEXT("int[]"))
+				{
+					OutErrorText = FText::Format(LOCTEXT("WrongPatternGeomSubsetIndexType", "Wrong pattern index type for GeomSubset '{0}'. Needs to be 'int[]'."), FText::FromString(MeshPrim.GetPrimPath().GetString()));
+					return false;
+				}
+
+				if (Patterns.Contains(MeshChildPrim.GetName()))
+				{
+					OutErrorText = FText::Format(LOCTEXT("DuplicatePatternGeomSubsetName", "Duplicate pattern name for GeomSubset '{0}'. The name needs to be unique."), FText::FromString(MeshPrim.GetPrimPath().GetString()));
+					return false;
+				}
+
+				Patterns.Emplace(MeshChildPrim.GetName(), GetIntArrayValues(IndicesAttr));
+			}
+		}
+		return true;
+	}
+
+	static bool ImportPatternsFromRenderMeshPrim(const FUsdPrim& RenderMeshPrim, const FUsdPrim& SimMeshPrim, TMap<FName, TSet<int32>>& Patterns, TMap<FName, TSet<FName>>& RenderToSimPatterns, FText& OutErrorText)
+	{
+		RenderToSimPatterns.Reset();
+
+		if (ImportPatternsFromMeshPrim(RenderMeshPrim, RenderPatternAPI, Patterns, OutErrorText))
+		{
+			const FSdfPath SimMeshPath = SimMeshPrim.GetPrimPath();
+
+			for (const FUsdPrim& MeshChildPrim : RenderMeshPrim.GetChildren())
+			{
+				if (MeshChildPrim.IsA(TEXT("GeomSubset")) && MeshChildPrim.HasAPI(RenderPatternAPI))
+				{
+					// Read simPattern relationship
+					const FUsdRelationship Relationship = MeshChildPrim.GetRelationship(TEXT("simPattern"));
+					TArray<FSdfPath> Targets;
+					Relationship.GetTargets(Targets);
+
+					// Add a new set of sim mesh patterns for this render pattern
+					TSet<FName>& SimMeshPatterns = RenderToSimPatterns.Add(MeshChildPrim.GetName());
+
+					// Add all sim mesh targets
+					for (const FSdfPath& Target : Targets)
+					{
+						if (Target.GetParentPath() == SimMeshPath)
+						{
+							SimMeshPatterns.Emplace(Target.GetName());
+						}
+						else
+						{
+							OutErrorText = FText::Format(LOCTEXT("UnknownOrMultipleSimMesh", "Unknown or more than one simulation mesh found while getting simPattern relationship of render pattern '{0}'."), FText::FromString(MeshChildPrim.GetPrimPath().GetString()));
+							return false;
+						}
+					}
+				}
+			}
+			return true;
+		}
+		return false;
+	}
+
+	static bool ImportPatternsFromSimMeshPrim(const FUsdPrim& SimMeshPrim, TMap<FName, TSet<int32>>& Patterns, FText& OutErrorText)
+	{
+		return ImportPatternsFromMeshPrim(SimMeshPrim, SimPatternAPI, Patterns, OutErrorText);
+	}
+
+	static bool ImportSewingsFromSimMeshPrim(const FUsdPrim& SimMeshPrim, TMap<FName, TSet<FIntVector2>>& Sewings, FText& OutErrorText)
+	{
+		Sewings.Reset();
+		for (const FUsdPrim& SimMeshChildPrim : SimMeshPrim.GetChildren())
+		{
+			if (SimMeshChildPrim.IsA(TEXT("GeomSubset")) && SimMeshChildPrim.HasAPI(SewingAPI))
+			{
+				const FUsdGeomSubset GeomSubset(SimMeshChildPrim);
+
+				// Read FamillyName
+				const FUsdAttribute FamilyNameAttr = GeomSubset.GetFamilyNameAttr();
+				if (GetStringValue(FamilyNameAttr) != TEXT("sewing"))
+				{
+					OutErrorText = LOCTEXT("WrongSewingGeomSubsetFamilyName", "Wrong sewing GeomSubset family name. Needs to be 'pattern'.");
+					return false;
+				}
+
+				// Read ElementTypeAttr
+				const FUsdAttribute ElementTypeAttr = GeomSubset.GetElementTypeAttr();
+				if (GetStringValue(ElementTypeAttr) != TEXT("edge"))
+				{
+					OutErrorText = LOCTEXT("WrongSewingGeomSubsetType", "Wrong sewing GeomSubset type. Needs to be edge.");
+					return false;
+				}
+
+				// Read indices 
+				const FUsdAttribute IndicesAttr = GeomSubset.GetIndicesAttr();
+				if (IndicesAttr.GetTypeName() != TEXT("int[]"))
+				{
+					OutErrorText = LOCTEXT("WrongSewingGeomSubsetIndexType", "Wrong sewing GeomSubset index type. Needs to be int[].");
+					return false;
+				}
+
+				if (Sewings.Contains(SimMeshChildPrim.GetName()))
+				{
+					OutErrorText = LOCTEXT("DuplicateSewingGeomSubsetName", "Duplicate sewing GeomSubset name. The name needs to be unique.");
+					return false;
+				}
+
+				const TArray<int32> IntArrayValues = GetIntArrayValues(IndicesAttr);
+				const int32 NumStitches = IntArrayValues.Num() / 2;
+				if (NumStitches * 2 != IntArrayValues.Num())
+				{
+					OutErrorText = LOCTEXT("OddSewingGeomSubsetIndices", "Odd number of indices for the sewing edges.");
+					return false;
+				}
+
+				TSet<FIntVector2>& Stitches = Sewings.Emplace(SimMeshChildPrim.GetName());
+				Stitches.Reserve(NumStitches);
+				for (int32 Index = 0; Index < NumStitches; ++Index)
+				{
+					const int32 Index0 = IntArrayValues[Index * 2];
+					const int32 Index1 = IntArrayValues[Index * 2 + 1];
+					Stitches.Emplace(Index0 <= Index1 ? FIntVector2(Index0, Index1) : FIntVector2(Index1, Index0));
+				}
+			}
+		}
+		return true;
+	}
 }  // End namespace Private
 
 FChaosClothAssetUSDImportNode_v2::FChaosClothAssetUSDImportNode_v2(const UE::Dataflow::FNodeParameters& InParam, FGuid InGuid)
@@ -215,12 +494,18 @@ FChaosClothAssetUSDImportNode_v2::FChaosClothAssetUSDImportNode_v2(const UE::Dat
 			{
 				const FString AssetPath = OwningObject ? OwningObject->GetPackage()->GetPathName() : FString();
 				FText ErrorText;
-				if (!ImportUsdFile_Schemaless(UsdFile.FilePath, AssetPath, ErrorText) &&
-					!ImportUsdFile(UsdFile.FilePath, AssetPath, ErrorText))
+				if (!ImportUsdFile(UsdFile.FilePath, AssetPath, ErrorText))
 				{
 					UE::Chaos::ClothAsset::FClothDataflowTools::LogAndToastWarning(*this,
-						LOCTEXT("FailedToImportUsdFileHeadline", "Failed to import USD file from file."),
-						FText::Format(LOCTEXT("FailedToImportUsdDetails", "Error while importing USD cloth from file '{0}':\n{1}"), FText::FromString(UsdFile.FilePath), ErrorText));
+						LOCTEXT("FailedToImportUsdFileHeadline", "Invalid USD Cloth Schema."),
+						FText::Format(LOCTEXT("FailedToImportUsdDetails", "Error while importing USD cloth from file '{0}':\n{1}\n\nWill now fallback to the legacy schema-less USD import."), FText::FromString(UsdFile.FilePath), ErrorText));
+
+					if(!ImportUsdFile_Schemaless(UsdFile.FilePath, AssetPath, ErrorText))
+					{
+						UE::Chaos::ClothAsset::FClothDataflowTools::LogAndToastWarning(*this,
+							LOCTEXT("FailedToImportUsdFileHeadline", "Failed to import USD file from file."),
+							FText::Format(LOCTEXT("FailedToImportUsdDetails", "Error while importing USD cloth from file '{0}':\n{1}"), FText::FromString(UsdFile.FilePath), ErrorText));
+					}
 				}
 			}))
 	, ReimportUsdFile(
@@ -291,24 +576,40 @@ void FChaosClothAssetUSDImportNode_v2::Serialize(FArchive& Ar)
 			ClothFacade.DefineSchema();
 		}
 		Collection = MoveTemp(*ClothCollection);
+
+		// Regenerate correct dependencies if needed
+		if (!ImportedAssets_DEPRECATED.IsEmpty())
+		{
+			ImportedAssets_DEPRECATED.Empty();
+			ImportedSimAssets = GetImportedAssetDependencies(ImportedSimStaticMesh);
+			ImportedRenderAssets = GetImportedAssetDependencies(ImportedRenderStaticMesh);
+		}
 	}
 }
 
+void FChaosClothAssetUSDImportNode_v2::ResetImport()
+{
+	Collection.Reset();
+	PackagePath = FString();
+	ImportedRenderStaticMesh = nullptr;
+	ImportedSimStaticMesh = nullptr;
+	ImportedUVScale = { 1.f, 1.f };
+	ImportedRenderAssets.Reset();
+	ImportedSimAssets.Reset();
+}
+
 // V1 of the USD importer (schemaless)
 bool FChaosClothAssetUSDImportNode_v2::ImportUsdFile_Schemaless(const FString& UsdFilePath, const FString& AssetPath, FText& OutErrorText)
 {
 	using namespace UE::Chaos::ClothAsset;
 
-	ImportedRenderStaticMesh = nullptr;
-	ImportedSimStaticMesh = nullptr;
-	ImportedUVScale = { 1.f, 1.f };
-	ImportedAssets.Reset();
+	ResetImport();
 
 	// Temporary borrow the collection to make the shared ref
 	const TSharedRef<FManagedArrayCollection> ClothCollection = MakeShared<FManagedArrayCollection>(MoveTemp(Collection));
 	ON_SCOPE_EXIT{ Collection = MoveTemp(*ClothCollection); };
 
-	const float NumSteps = bImportRenderMesh ? 2.f : 1.f;
+	const float NumSteps = bImportRenderMesh ? 2.f : 1.f;  // Sim mesh is always imported
 	FScopedSlowTask SlowTask(NumSteps, LOCTEXT("ImportingUSDFile", "Importing USD file..."));
 
 	SlowTask.EnterProgressFrame(1.f, LOCTEXT("CreatingAssets", "Creating assets and importing simulation mesh..."));
@@ -329,23 +630,20 @@ PRAGMA_ENABLE_DEPRECATION_WARNINGS
 			return false;
 		}
 	}
+
 	return true;
 }
 
 // V2 of the USD importer (using cloth schema)
 bool FChaosClothAssetUSDImportNode_v2::ImportUsdFile(const FString& UsdFilePath, const FString& AssetPath, FText& OutErrorText)
 {
-#if USE_USD_SDK
+	using namespace UE;
 	using namespace UE::Chaos::ClothAsset;
+	using namespace UE::Chaos::ClothAsset::Private;
 
-	// Reset collection
-	Collection.Reset();
-	PackagePath = FString();
-	ImportedRenderStaticMesh = nullptr;
-	ImportedSimStaticMesh = nullptr;
-	ImportedUVScale = { 1.f, 1.f };
-	ImportedAssets.Reset();
+	ResetImport();
 
+#if USE_USD_SDK
 	// Temporary borrow the collection to make the shared ref
 	const TSharedRef<FManagedArrayCollection> ClothCollection = MakeShared<FManagedArrayCollection>(MoveTemp(Collection));
 	ON_SCOPE_EXIT{ Collection = MoveTemp(*ClothCollection); };
@@ -368,7 +666,7 @@ bool FChaosClothAssetUSDImportNode_v2::ImportUsdFile(const FString& UsdFilePath,
 	constexpr bool bUseStageCache = false;  // Reload from disk, not from cache
 	constexpr EUsdInitialLoadSet UsdInitialLoadSet = EUsdInitialLoadSet::LoadAll;  // TODO: Ideally we should only use LoadNone to start with and load what's needed once the Schema is defined
 
-	const UE::FUsdStage UsdStage = UnrealUSDWrapper::OpenStage(*UsdFilePath, UsdInitialLoadSet, bUseStageCache);
+	FUsdStage UsdStage = UnrealUSDWrapper::OpenStage(*UsdFilePath, UsdInitialLoadSet, bUseStageCache);
 	if (!UsdStage)
 	{
 		OutErrorText = LOCTEXT("CantCreateNewStage", "Failed to open the specified USD file.");
@@ -376,7 +674,7 @@ bool FChaosClothAssetUSDImportNode_v2::ImportUsdFile(const FString& UsdFilePath,
 	}
 
 	// Find the cloth prim
-	const UE::FUsdPrim ClothPrim = Private::FindClothPrim(UsdStage.GetPseudoRoot());
+	const FUsdPrim ClothPrim = FindClothPrim(UsdStage.GetPseudoRoot());
 	if (!ClothPrim)
 	{
 		OutErrorText = LOCTEXT("CantFindClothRootAPI", "Can't find a cloth root inside the specified USD file.");
@@ -384,16 +682,22 @@ bool FChaosClothAssetUSDImportNode_v2::ImportUsdFile(const FString& UsdFilePath,
 	}
 
 	// Find SimMesh and Render Mesh prims
-	const UE::FUsdPrim SimMeshPrim = Private::FindSimMeshPrim(ClothPrim);
-	const UE::FUsdPrim RenderMeshPrim = Private::FindRenderMeshPrim(ClothPrim);
+	const FUsdPrim SimMeshPrim = FindSimMeshPrim(ClothPrim);
+	const FUsdPrim RenderMeshPrim = FindRenderMeshPrim(ClothPrim);
 	if (!SimMeshPrim && !RenderMeshPrim)
 	{
 		OutErrorText = LOCTEXT("CantFindMeshPrims", "Can't find a sim mesh or render mesh prim with valid pattern data.");
 		return false;
 	}
 
+	// Remove Opacity from the stage before import since otherwise it messes up all materials
+	if (!bImportWithOpacity)
+	{
+		RemoveMaterialOpacity(UsdStage.GetPseudoRoot());
+	}
+
 	// Read UVScale attribute
-	ImportedUVScale = Private::GetSimMeshUVScale(SimMeshPrim);
+	ImportedUVScale = GetSimMeshUVScale(SimMeshPrim);
 
 	// Update import location
 	const uint32 UsdPathHash = GetTypeHash(UsdFile.FilePath);  // Path hash to store all import from the same file/same path to the same content folder
@@ -403,7 +707,7 @@ bool FChaosClothAssetUSDImportNode_v2::ImportUsdFile(const FString& UsdFilePath,
 
 	// Import the stage
 	SlowTask.EnterProgressFrame(1.f);
-	Private::ImportStaticMeshesFromUsdStage(UsdStage, UsdFilePath, PackagePath);
+	ImportStaticMeshesFromUsdStage(UsdStage, UsdFilePath, PackagePath);
 
 	// Fill up the asset list from the imported USD assets
 	const FString SimMeshName = SimMeshPrim ? FString::Printf(TEXT("SM_%s"), *SimMeshPrim.GetName().ToString()) : FString();
@@ -414,6 +718,26 @@ bool FChaosClothAssetUSDImportNode_v2::ImportUsdFile(const FString& UsdFilePath,
 	if (bImportSimMesh)
 	{
 		SlowTask.EnterProgressFrame(1.f);
+
+		// Check that the entire mesh is made of triangles
+		if (!CheckSimMeshPrimTriangles(SimMeshPrim, OutErrorText))
+		{
+			return false;
+		}
+
+		// Import the simulation patterns
+		if (!ImportPatternsFromSimMeshPrim(SimMeshPrim, UsdClothData.SimPatterns, OutErrorText))
+		{
+			return false;
+		}
+
+		// Import the sewings
+		if (!ImportSewingsFromSimMeshPrim(SimMeshPrim, UsdClothData.Sewings, OutErrorText))
+		{
+			return false;
+		}
+
+		// Lastly import the geometry and finalize the patterns
 		if (!ImportSimStaticMesh(ClothCollection, OutErrorText))
 		{
 			return false;
@@ -424,6 +748,13 @@ bool FChaosClothAssetUSDImportNode_v2::ImportUsdFile(const FString& UsdFilePath,
 	if (bImportRenderMesh)
 	{
 		SlowTask.EnterProgressFrame(1.f);
+
+		// Import the render patterns
+		if (!ImportPatternsFromRenderMeshPrim(RenderMeshPrim, SimMeshPrim, UsdClothData.RenderPatterns, UsdClothData.RenderToSimPatterns, OutErrorText))
+		{
+			return false;
+		}
+
 		if (!ImportRenderStaticMesh(ClothCollection, OutErrorText))
 		{
 			return false;
@@ -444,20 +775,16 @@ void FChaosClothAssetUSDImportNode_v2::UpdateImportedAssets(const FString& SimMe
 {
 	ImportedSimStaticMesh = nullptr;
 	ImportedRenderStaticMesh = nullptr;
-	ImportedAssets.Reset();
 
-	if (!PackagePath.IsEmpty())
+	if (!PackagePath.IsEmpty() && (!SimMeshName.IsEmpty() || !RenderMeshName.IsEmpty()))
 	{
-		TArray<FAssetData> AssetData;
-
 		const FAssetRegistryModule& AssetRegistryModule = FModuleManager::LoadModuleChecked<FAssetRegistryModule>("AssetRegistry");
 
 		constexpr bool bRecursive = true;
 		constexpr bool bIncludeOnlyOnDiskAssets = false;
+		TArray<FAssetData> AssetData;
 		AssetRegistryModule.Get().GetAssetsByPath(FName(*PackagePath), AssetData, bRecursive, bIncludeOnlyOnDiskAssets);
 
-		ImportedAssets.Reserve(AssetData.Num());
-
 		// Find sim mesh and render mesh (static meshes) dependencies
 		for (const FAssetData& AssetDatum : AssetData)
 		{
@@ -473,65 +800,48 @@ void FChaosClothAssetUSDImportNode_v2::UpdateImportedAssets(const FString& SimMe
 					ImportedRenderStaticMesh = CastChecked<UStaticMesh>(AssetDatum.GetAsset());
 					UE_LOG(LogChaosClothAssetDataflowNodes, Display, TEXT("Imported USD Render Mesh %s, path: %s"), *AssetDatum.AssetName.ToString(), *AssetDatum.GetFullName());
 				}
+				if ((ImportedSimStaticMesh || SimMeshName.IsEmpty()) &&
+					(ImportedRenderStaticMesh || RenderMeshName.IsEmpty()))
+				{
+					break;
+				}
 			}
 		}
+	}
+}
 
-		// Find all other dependencies
-		auto AddImportedAssetDependencies = [this, &AssetRegistryModule, &AssetData](const UObject* Object)
-			{
-				TQueue<const UObject*> ObjectsToVisit;
-				ObjectsToVisit.Enqueue(Object);
-
-				const UObject* VisitedObject;
-				while (ObjectsToVisit.Dequeue(VisitedObject))
-				{
-					const FName PackageName(VisitedObject->GetPackage()->GetName());
-					TArray<FName> Dependencies;
-					AssetRegistryModule.Get().GetDependencies(PackageName, Dependencies);
+TArray<TObjectPtr<UObject>> FChaosClothAssetUSDImportNode_v2::GetImportedAssetDependencies(const UObject* StaticMesh)
+{
+	using namespace UE::Chaos::ClothAsset::Private;
 
-					UE_LOG(LogChaosClothAssetDataflowNodes, Verbose, TEXT("Dependencies for Object %s - %s:"), *Object->GetName(), *PackageName.ToString());
-					for (const FName& Dependency : Dependencies)
-					{
-						// Only add depedencies the USD import has created
-						if (AssetData.FindByPredicate([&Dependency](const FAssetData& AssetDatum) { return AssetDatum.PackageName == Dependency; }))
-						{
-							UE_LOG(LogChaosClothAssetDataflowNodes, Verbose, TEXT("Found %s"), *Dependency.ToString());
-							TArray<FAssetData> DependencyAssetData;
-							AssetRegistryModule.Get().GetAssetsByPackageName(Dependency, DependencyAssetData);
-
-							for (const FAssetData& DependencyAssetDatum : DependencyAssetData)
-							{
-								if (DependencyAssetDatum.IsUAsset() && DependencyAssetDatum.IsTopLevelAsset())  // IsUAsset returns false for redirects
-								{
-									const int32 Index = ImportedAssets.Emplace(DependencyAssetDatum.GetAsset());  // GetAsset does not handle redirects
-
-									ObjectsToVisit.Enqueue(ImportedAssets[Index]);  // Visit this asset too
-								}
-							}
-						}
-					}
-				}
-			};
+	TSet<TObjectPtr<UObject>> ImportedAssets;
+	if (StaticMesh)
+	{
+		TQueue<const UObject*> AssetsToVisit;
+		AssetsToVisit.Enqueue(StaticMesh);
 
-		if (ImportedSimStaticMesh)
-		{
-			AddImportedAssetDependencies(ImportedSimStaticMesh);
-		}
-		if (ImportedRenderStaticMesh)
+		const UObject* VisitedAsset;
+		while (AssetsToVisit.Dequeue(VisitedAsset))
 		{
-			AddImportedAssetDependencies(ImportedRenderStaticMesh);
-		}
+			const FName VisitedAssetPackageName(VisitedAsset->GetPackage()->GetName());
+			const TArray<UObject*> AssetDependencies = GetAssetDependencies(VisitedAsset);
 
-		for (const UObject* const Asset : ImportedAssets)
-		{
-			UE_LOG(LogChaosClothAssetDataflowNodes,
-				Display,
-				TEXT("Imported USD Object %s of type %s, path: %s"),
-				*Asset->GetName(),
-				*Asset->GetClass()->GetName(),
-				*Asset->GetPathName());
+			UE_CLOG(AssetDependencies.Num(), LogChaosClothAssetDataflowNodes, Verbose, TEXT("Dependencies for Object %s - %s:"), *VisitedAsset->GetName(), *VisitedAssetPackageName.ToString());
+			for (UObject* const AssetDependency : AssetDependencies)
+			{
+				if (!ImportedAssets.Contains(AssetDependency))
+				{
+					// Add the dependency
+					UE_LOG(LogChaosClothAssetDataflowNodes, Verbose, TEXT("Found %s"), *AssetDependency->GetPackage()->GetName());
+					ImportedAssets.Emplace(AssetDependency);
+
+					// Visit this asset too
+					AssetsToVisit.Enqueue(AssetDependency);
+				}
+			}
 		}
 	}
+	return ImportedAssets.Array();
 }
 
 bool FChaosClothAssetUSDImportNode_v2::ImportSimStaticMesh(const TSharedRef<FManagedArrayCollection> ClothCollection, FText& OutErrorText)
@@ -541,8 +851,27 @@ bool FChaosClothAssetUSDImportNode_v2::ImportSimStaticMesh(const TSharedRef<FMan
 	FCollectionClothFacade ClothFacade(ClothCollection);
 	check(ClothFacade.IsValid());  // The Cloth Collection schema must be valid at this point
 
-	// Empty the current sim mesh
+	// Define the selection schema if needed
+	FCollectionClothSelectionFacade ClothSelectionFacade(ClothCollection);
+	if (!ClothSelectionFacade.IsValid())
+	{
+		ClothSelectionFacade.DefineSchema();
+	}
+
+	// Empty the current sim mesh and any previously created selection set
 	FClothGeometryTools::DeleteSimMesh(ClothCollection);
+	FClothGeometryTools::DeleteSelections(ClothCollection, ClothCollectionGroup::SimFaces);
+
+	ON_SCOPE_EXIT
+		{
+			// Bind to root bone on exit
+			constexpr bool bBindSimMesh = true;
+			constexpr bool bBindRenderMesh = false;
+			FClothGeometryTools::BindMeshToRootBone(ClothCollection, bBindSimMesh, bBindRenderMesh);
+
+			// Make sure to clean the dependencies whatever the import status is
+			ImportedSimAssets = GetImportedAssetDependencies(ImportedSimStaticMesh);
+		};
 
 	if (!ImportedSimStaticMesh)
 	{
@@ -613,41 +942,158 @@ bool FChaosClothAssetUSDImportNode_v2::ImportSimStaticMesh(const TSharedRef<FMan
 		}
 	}
 
-	// Rescale the 2D mesh with the UV scale
+	// Rescale the 2D mesh with the UV scale, and flip the UV's Y coordinates
 	for (FVector2f& Pos : RestPositions2D)
 	{
+		Pos.Y = 1.f - Pos.Y;
 		Pos *= ImportedUVScale;
 	}
 
 	// Save pattern to the collection cache
-	check(RestPositions2D.Num() == DrapedPositions3D.Num());
-	const int32 VertexCount = RestPositions2D.Num();
-	const int32 TriangleCount = TriangleToVertexIndex.Num();
-
-	if (TriangleCount && VertexCount)
-	{
-		// Remove degenerated triangles
-		TArray<FIntVector3> OutTriangleToVertexIndex;
-		TArray<FVector2f> OutRestPositions2D;
-		TArray<FVector3f> OutDrapedPositions3D;
-		TArray<int32> OutIndices;
-
-		bool bHasRepairedTriangles = FClothDataflowTools::RemoveDegenerateTriangles(
-			TriangleToVertexIndex,
-			RestPositions2D,
-			DrapedPositions3D,
-			OutTriangleToVertexIndex,
-			OutRestPositions2D,
-			OutDrapedPositions3D,
-			OutIndices);
-
-		// Remove duplicate triangles
-		bHasRepairedTriangles = FClothDataflowTools::RemoveDuplicateTriangles(OutTriangleToVertexIndex) || bHasRepairedTriangles;
-
-		// Add the new pattern
-		const int32 SimPatternIndex = ClothFacade.AddSimPattern();
-		FCollectionClothSimPatternFacade SimPattern = ClothFacade.GetSimPattern(SimPatternIndex);
-		SimPattern.Initialize(OutRestPositions2D, OutDrapedPositions3D, OutTriangleToVertexIndex);
+	check(RestPositions2D.Num() == DrapedPositions3D.Num());  // Should have already exited with the UsdSimMeshWelded error in this case
+	if (TriangleToVertexIndex.Num() && RestPositions2D.Num())
+	{
+		// Cleanup sim mesh
+		FClothDataflowTools::FSimMeshCleanup SimMeshCleanup(TriangleToVertexIndex, RestPositions2D, DrapedPositions3D);
+
+		bool bHasRepairedTriangles = SimMeshCleanup.RemoveDegenerateTriangles();
+		bHasRepairedTriangles = SimMeshCleanup.RemoveDuplicateTriangles() || bHasRepairedTriangles;
+
+		const TArray<int32> OriginalToNewTriangles = FClothDataflowTools::GetOriginalToNewIndices<TSet<int32>>(SimMeshCleanup.OriginalTriangles, TriangleToVertexIndex.Num());
+
+		// Add support for original indices
+		ClothFacade.AddUserDefinedAttribute<TArray<int32>>(Private::OriginalIndicesName, ClothCollectionGroup::SimFaces);
+		ClothFacade.AddUserDefinedAttribute<TArray<int32>>(Private::OriginalIndicesName, ClothCollectionGroup::SimVertices2D);
+
+		// Add the patterns from the clean mesh
+		for (const TPair<FName, TSet<int32>>& PatternNameFaces : UsdClothData.SimPatterns)
+		{
+			// Filter the pattern selection set using the remaining triangles from the cleaned triangle list
+			TSet<int32> PatternSet;
+			PatternSet.Reserve(PatternNameFaces.Value.Num());
+			for (const int32 Face : PatternNameFaces.Value)
+			{
+				if (OriginalToNewTriangles.IsValidIndex(Face) && OriginalToNewTriangles[Face] != INDEX_NONE)
+				{
+					PatternSet.Emplace(OriginalToNewTriangles[Face]);
+				}
+			}
+
+			// Add the new pattern
+			if (PatternSet.Num())
+			{
+				TArray<FIntVector3> PatternTriangleToVertexIndex;
+				TArray<TArray<int32>> PatternOriginalTriangles;
+				PatternTriangleToVertexIndex.Reserve(PatternSet.Num());
+				PatternOriginalTriangles.Reserve(PatternSet.Num());
+				{
+					for (const int32 Index : PatternSet)
+					{
+						PatternTriangleToVertexIndex.Emplace(SimMeshCleanup.TriangleToVertexIndex[Index]);
+						PatternOriginalTriangles.Emplace(SimMeshCleanup.OriginalTriangles[Index].Array());
+					}
+				}
+
+				TArray<FVector2f> PatternRestPositions2D;
+				TArray<FVector3f> PatternDrapedPositions3D;
+				TArray<TArray<int32>> PatternOriginalVertices;
+				TArray<int32> PatternVertexReindex;
+				const int32 MaxNumVertices = SimMeshCleanup.RestPositions2D.Num();
+				PatternRestPositions2D.Reserve(MaxNumVertices);
+				PatternDrapedPositions3D.Reserve(MaxNumVertices);
+				PatternOriginalVertices.Reserve(MaxNumVertices);
+				PatternVertexReindex.Init(INDEX_NONE, MaxNumVertices);
+
+				int32 NewIndex = -1;
+				for (FIntVector3& Triangle : PatternTriangleToVertexIndex)
+				{
+					for (int32 Vertex = 0; Vertex < 3; ++Vertex)
+					{
+						// Add the new vertex
+						int32& Index = Triangle[Vertex];
+						if (PatternVertexReindex[Index] == INDEX_NONE)
+						{
+							PatternVertexReindex[Index] = ++NewIndex;
+							PatternRestPositions2D.Emplace(SimMeshCleanup.RestPositions2D[Index]);
+							PatternDrapedPositions3D.Emplace(SimMeshCleanup.DrapedPositions3D[Index]);
+							PatternOriginalVertices.Emplace(SimMeshCleanup.OriginalVertices[Index].Array());
+						}
+						// Reindex the triangle vertex with the new index
+						Index = PatternVertexReindex[Index];
+					}
+				}
+
+				// Add this pattern to the cloth collection
+				const int32 SimPatternIndex = ClothFacade.AddSimPattern();
+				FCollectionClothSimPatternFacade SimPattern = ClothFacade.GetSimPattern(SimPatternIndex);
+				SimPattern.Initialize(PatternRestPositions2D, PatternDrapedPositions3D, PatternTriangleToVertexIndex);
+
+				// Keep track of the original triangle indices
+				const TArrayView<TArray<int32>> OriginalTriangles =
+					ClothFacade.GetUserDefinedAttribute<TArray<int32>>(Private::OriginalIndicesName, ClothCollectionGroup::SimFaces);  // Don't move outside the loop, the array might get re-allocated
+				const int32 SimFacesOffset = SimPattern.GetSimFacesOffset();
+				for (int32 Index = 0; Index < PatternOriginalTriangles.Num(); ++Index)
+				{
+					OriginalTriangles[SimFacesOffset + Index] = PatternOriginalTriangles[Index];
+				}
+
+				// Keep track of the original vertex indices
+				const TArrayView<TArray<int32>> OriginalVertices =
+					ClothFacade.GetUserDefinedAttribute<TArray<int32>>(Private::OriginalIndicesName, ClothCollectionGroup::SimVertices2D);  // Don't move outside the loop, the array might get re-allocated
+				const int32 SimVertices2DOffset = SimPattern.GetSimVertices2DOffset();
+				for (int32 Index = 0; Index < PatternOriginalVertices.Num(); ++Index)
+				{
+					OriginalVertices[SimVertices2DOffset + Index] = PatternOriginalVertices[Index];
+				}
+
+				// Add the pattern triangle list as a selection set
+				TSet<int32>& SelectionSet = ClothSelectionFacade.FindOrAddSelectionSet(PatternNameFaces.Key, ClothCollectionGroup::SimFaces);
+				SelectionSet.Empty(PatternSet.Num());
+				for (int32 Index = SimFacesOffset; Index < SimFacesOffset + PatternTriangleToVertexIndex.Num(); ++Index)
+				{
+					SelectionSet.Emplace(Index);
+				}
+			}
+		}
+
+		// Check the resulting cleaned mesh
+		const int32 NumSimVertices2D = ClothFacade.GetNumSimVertices2D();
+		const int32 NumSimFaces = ClothFacade.GetNumSimFaces();
+		if (!NumSimVertices2D || !NumSimFaces)
+		{
+			return true;  // Empty mesh
+		}
+
+		const TConstArrayView<TArray<int32>> OriginalTriangles =
+			ClothFacade.GetUserDefinedAttribute<TArray<int32>>(Private::OriginalIndicesName, ClothCollectionGroup::SimFaces);
+		const TArray<int32> OriginalToNewFaceIndices = FClothDataflowTools::GetOriginalToNewIndices(OriginalTriangles, TriangleToVertexIndex.Num());
+
+		const TConstArrayView<TArray<int32>> OriginalVertices =
+			ClothFacade.GetUserDefinedAttribute<TArray<int32>>(Private::OriginalIndicesName, ClothCollectionGroup::SimVertices2D);
+		const TArray<int32> OriginalToNewVertexIndices = FClothDataflowTools::GetOriginalToNewIndices(OriginalVertices, RestPositions2D.Num());
+
+		// Add the sewings
+		for (const TPair<FName, TSet<FIntVector2>>& SewingNameIndices : UsdClothData.Sewings)
+		{
+			TSet<FIntVector2> Indices;
+			for (const FIntVector2& Stitch : SewingNameIndices.Value)
+			{
+				if (!OriginalToNewVertexIndices.IsValidIndex(Stitch[0]) || !OriginalToNewVertexIndices.IsValidIndex(Stitch[1]))
+				{
+					OutErrorText = LOCTEXT("BadSewingIndex", "An out of renge sewing index has been found.");
+					return false;
+				}
+				const int32 StitchIndex0 = OriginalToNewVertexIndices[Stitch[0]];
+				const int32 StitchIndex1 = OriginalToNewVertexIndices[Stitch[1]];
+				if (StitchIndex0 != INDEX_NONE && StitchIndex1 != INDEX_NONE)
+				{
+					Indices.Emplace(StitchIndex0 < StitchIndex1 ? FIntVector2(StitchIndex0, StitchIndex1) : FIntVector2(StitchIndex1, StitchIndex0));
+				}
+			}
+
+			FCollectionClothSeamFacade ClothSeamFacade = ClothFacade.AddGetSeam();
+			ClothSeamFacade.Initialize(Indices.Array());
+		}
 	}
 	return true;
 }
@@ -659,37 +1105,67 @@ bool FChaosClothAssetUSDImportNode_v2::ImportRenderStaticMesh(const TSharedRef<F
 	FCollectionClothFacade ClothFacade(ClothCollection);
 	check(ClothFacade.IsValid());  // The Cloth Collection schema must be valid at this point
 
-	// Empty the current render mesh
-	FClothGeometryTools::DeleteRenderMesh(ClothCollection);
-
-	if (!ImportedRenderStaticMesh)
+	// Define the selection schema if needed
+	FCollectionClothSelectionFacade ClothSelectionFacade(ClothCollection);
+	if (!ClothSelectionFacade.IsValid())
 	{
-		return true;  // Nothing to import
+		ClothSelectionFacade.DefineSchema();
 	}
 
+	// Empty the current render mesh and previously create selections
+	FClothGeometryTools::DeleteRenderMesh(ClothCollection);
+	FClothGeometryTools::DeleteSelections(ClothCollection, ClothCollectionGroup::RenderFaces);
+
+	// Make sure to clean the dependencies whatever the import status is
+	ON_SCOPE_EXIT
+		{
+			// Bind to root bone on exit
+			constexpr bool bBindSimMesh = false;
+			constexpr bool bBindRenderMesh = true;
+			FClothGeometryTools::BindMeshToRootBone(ClothCollection, bBindSimMesh, bBindRenderMesh);
+
+			// Make sure to clean the dependencies whatever the import status is
+			ImportedRenderAssets = GetImportedAssetDependencies(ImportedRenderStaticMesh);
+		};
+
 	// Import the LOD 0
-	if (ImportedRenderStaticMesh->GetNumSourceModels())
+	if (ImportedRenderStaticMesh && ImportedRenderStaticMesh->GetNumSourceModels())
 	{
 		constexpr int32 LODIndex = 0;
-		const FMeshDescription* const MeshDescription = ImportedRenderStaticMesh->GetMeshDescription(LODIndex);
-		const FMeshBuildSettings& BuildSettings = ImportedRenderStaticMesh->GetSourceModel(LODIndex).BuildSettings;
-		const TArray<FStaticMaterial>& StaticMaterials = ImportedRenderStaticMesh->GetStaticMaterials();
-
-		FSkeletalMeshLODModel SkeletalMeshModel;
-		if (FClothDataflowTools::BuildSkeletalMeshModelFromMeshDescription(MeshDescription, BuildSettings, SkeletalMeshModel))
+		if (const FMeshDescription* const MeshDescription = ImportedRenderStaticMesh->GetMeshDescription(LODIndex))
 		{
-			FStaticMeshConstAttributes MeshAttributes(*MeshDescription);
-			TPolygonGroupAttributesConstRef<FName> MaterialSlotNames = MeshAttributes.GetPolygonGroupMaterialSlotNames();
-			for (int32 SectionIndex = 0; SectionIndex < SkeletalMeshModel.Sections.Num(); ++SectionIndex)
+			const FMeshBuildSettings& BuildSettings = ImportedRenderStaticMesh->GetSourceModel(LODIndex).BuildSettings;
+			FRenderMeshImport RenderMeshImport(*MeshDescription, BuildSettings);
+
+			const TArray<FStaticMaterial>& StaticMaterials = ImportedRenderStaticMesh->GetStaticMaterials();
+			RenderMeshImport.AddRenderSections(ClothCollection, StaticMaterials, Private::OriginalIndicesName, Private::OriginalIndicesName);
+
+			// Create pattern selection sets
+			const TConstArrayView<TArray<int32>> OriginalTriangles =
+				ClothFacade.GetUserDefinedAttribute<TArray<int32>>(Private::OriginalIndicesName, ClothCollectionGroup::RenderFaces);
+
+			if (OriginalTriangles.Num())
 			{
-				// Section MaterialIndex refers to the polygon group index. Look up which material this corresponds with.
-				const FName& MaterialSlotName = MaterialSlotNames[SkeletalMeshModel.Sections[SectionIndex].MaterialIndex];
-				const int32 MaterialIndex = ImportedRenderStaticMesh->GetMaterialIndexFromImportedMaterialSlotName(MaterialSlotName);
-				const FString RenderMaterialPathName = StaticMaterials.IsValidIndex(MaterialIndex) && StaticMaterials[MaterialIndex].MaterialInterface ?
-					StaticMaterials[MaterialIndex].MaterialInterface->GetPathName() :
-					FString();
-				FClothDataflowTools::AddRenderPatternFromSkeletalMeshSection(ClothCollection, SkeletalMeshModel, SectionIndex, RenderMaterialPathName);
+				const TConstArrayView<FIntVector3> TriangleToVertexIndex = ClothFacade.GetRenderIndices();
+				const TArray<int32> OriginalToNewTriangles = FClothDataflowTools::GetOriginalToNewIndices<TArray<int32>>(OriginalTriangles, TriangleToVertexIndex.Num());
+
+				for (const TPair<FName, TSet<int32>>& PatternNameFaces : UsdClothData.RenderPatterns)
+				{
+					// Add the pattern triangle list as a selection set
+					TSet<int32>& SelectionSet = ClothSelectionFacade.FindOrAddSelectionSet(PatternNameFaces.Key, ClothCollectionGroup::RenderFaces);
+					SelectionSet.Empty(PatternNameFaces.Value.Num());
+					for (const int32 Index : PatternNameFaces.Value)
+					{
+						SelectionSet.Emplace(OriginalToNewTriangles[Index]);
+					}
+				}
 			}
+			// TODO: Proxy deformer
+		}
+		else
+		{
+			OutErrorText = LOCTEXT("MissingMeshDescription", "An imported render static mesh has no mesh description!");
+			return false;
 		}
 	}
 	return true;
diff --git a/Engine/Plugins/ChaosClothAssetEditor/Source/ChaosClothAssetDataflowNodes/Public/ChaosClothAsset/USDImportNode_v2.h b/Engine/Plugins/ChaosClothAssetEditor/Source/ChaosClothAssetDataflowNodes/Public/ChaosClothAsset/USDImportNode_v2.h
--- a/Engine/Plugins/ChaosClothAssetEditor/Source/ChaosClothAssetDataflowNodes/Public/ChaosClothAsset/USDImportNode_v2.h
+++ b/Engine/Plugins/ChaosClothAssetEditor/Source/ChaosClothAssetDataflowNodes/Public/ChaosClothAsset/USDImportNode_v2.h
@@ -19,6 +19,35 @@ namespace UE_DEPRECATED(5.5, "Use UE::Dataflow instead.") Dataflow {}
 
 class UStaticMesh;
 
+/** Support struct for storage of the extra data that doesn't fit inside a static mesh. */
+USTRUCT()
+struct FChaosClothAssetUsdClothData final
+{
+	GENERATED_USTRUCT_BODY()
+
+	TMap<FName, TSet<int32>> SimPatterns;
+	TMap<FName, TSet<FIntVector2>> Sewings;
+	TMap<FName, TSet<int32>> RenderPatterns;
+	TMap<FName, TSet<FName>> RenderToSimPatterns;
+
+	bool Serialize(FArchive& Ar)
+	{
+		Ar << SimPatterns;
+		Ar << Sewings;
+		Ar << RenderPatterns;
+		Ar << RenderToSimPatterns;
+		return true;
+	}
+};
+
+template<> struct TStructOpsTypeTraits<FChaosClothAssetUsdClothData> : public TStructOpsTypeTraitsBase2<FChaosClothAssetUsdClothData>
+{
+	enum
+	{
+		WithSerializer = true,
+	};
+};
+
 /** Import a USD file from a third party garment construction software. */
 USTRUCT(meta = (DataflowCloth))
 struct FChaosClothAssetUSDImportNode_v2 : public FDataflowNode
@@ -39,6 +68,10 @@ public:
 	UPROPERTY(EditAnywhere, Category = "USD Import")
 	bool bImportRenderMesh = true;
 
+	/** Importing the render mesh with opacity requires transluscency to be enable in the project settings. */
+	UPROPERTY(EditAnywhere, Category = "USD Import", Meta = (EditCondition = "bImportRenderMesh", EditConditionHides))
+	bool bImportWithOpacity = false;
+
 	/** Path of the USD file to import. */
 	UPROPERTY(EditAnywhere, Category = "USD Import", Meta = (DisplayName = "USD File"))
 	FChaosClothAssetImportFilePath UsdFile;
@@ -63,6 +96,8 @@ private:
 	virtual void Serialize(FArchive& Archive) override;
 	//~ End FDataflowNode interface
 
+	// Clear the imported data properties
+	void ResetImport();
 	// Import USD data if provided in the original legacy schemaless USD format
 	bool ImportUsdFile_Schemaless(const FString& UsdPath, const FString& AssetPath, class FText& OutErrorText);
 	// Import USD data if provided with a valid cloth schema
@@ -73,8 +108,10 @@ private:
 	// Import the given static mesh as a render mesh into the cloth collection
 	bool ImportRenderStaticMesh(const TSharedRef<FManagedArrayCollection> ClothCollection, class FText& OutErrorText);
 
-	// Rebuild the asset dependencies for the two specified static meshes
+	// Find the two imported static mesh and update the dependencies
 	void UpdateImportedAssets(const FString& SimMeshName, const FString& RenderMeshName);
+	// Find the asset dependencies for the specified static mesh
+	TArray<TObjectPtr<UObject>> GetImportedAssetDependencies(const UObject* StaticMesh);
 
 	/** Content folder where all the USD assets are imported. */
 	UPROPERTY(VisibleAnywhere, Category = "USD Import")
@@ -96,7 +133,18 @@ private:
 	UPROPERTY(VisibleAnywhere, Category = "USD Import")
 	TObjectPtr<UStaticMesh> ImportedRenderStaticMesh;
 
-	/** List of all the dependent assets created from the USD import process. */
+	/** List of all the simulation static mesh's dependent assets. This does not include any Engine, or Engine plugin content. */
+	UPROPERTY(VisibleAnywhere, Category = "USD Import")
+	TArray<TObjectPtr<UObject>> ImportedSimAssets;
+
+	/** List of all the render static mesh's dependent assets. This does not include any Engine, or Engine plugin content. */
 	UPROPERTY(VisibleAnywhere, Category = "USD Import")
-	TArray<TObjectPtr<UObject>> ImportedAssets;
+	TArray<TObjectPtr<UObject>> ImportedRenderAssets;
+
+	/** Support struct for storage of the extra data that doesn't fit inside a static mesh. */
+	UPROPERTY()
+	FChaosClothAssetUsdClothData UsdClothData;
+
+	UPROPERTY()
+	TArray<TObjectPtr<UObject>> ImportedAssets_DEPRECATED;
 };
