CIS - Fixed duplicated LOCTEXT key in USDImportNode_v2.cpp.

#rnx

diff --git a/Engine/Plugins/ChaosClothAssetEditor/Source/ChaosClothAssetDataflowNodes/Private/ChaosClothAsset/USDImportNode_v2.cpp b/Engine/Plugins/ChaosClothAssetEditor/Source/ChaosClothAssetDataflowNodes/Private/ChaosClothAsset/USDImportNode_v2.cpp
--- a/Engine/Plugins/ChaosClothAssetEditor/Source/ChaosClothAssetDataflowNodes/Private/ChaosClothAsset/USDImportNode_v2.cpp
+++ b/Engine/Plugins/ChaosClothAssetEditor/Source/ChaosClothAssetDataflowNodes/Private/ChaosClothAsset/USDImportNode_v2.cpp
@@ -497,8 +497,8 @@ FChaosClothAssetUSDImportNode_v2::FChaosClothAssetUSDImportNode_v2(const UE::Dat
 				if (!ImportUsdFile(UsdFile.FilePath, AssetPath, ErrorText))
 				{
 					UE::Chaos::ClothAsset::FClothDataflowTools::LogAndToastWarning(*this,
-						LOCTEXT("FailedToImportUsdFileHeadline", "Invalid USD Cloth Schema."),
-						FText::Format(LOCTEXT("FailedToImportUsdDetails", "Error while importing USD cloth from file '{0}':\n{1}\n\nWill now fallback to the legacy schema-less USD import."), FText::FromString(UsdFile.FilePath), ErrorText));
+						LOCTEXT("InvalidUSDClothSchemaHeadline", "Invalid USD Cloth Schema."),
+						FText::Format(LOCTEXT("InvalidUSDClothSchemaDetails", "Error while importing USD cloth from file '{0}':\n{1}\n\nWill now fallback to the legacy schema-less USD import."), FText::FromString(UsdFile.FilePath), ErrorText));
 
 					if(!ImportUsdFile_Schemaless(UsdFile.FilePath, AssetPath, ErrorText))
 					{
