Another crash fix when cooking cloth assets with morph targets. This should be the last of them because I was able to successfully cook a shelf with morph targets enabled.

diff --git a/Engine/Plugins/ChaosClothAsset/Source/ChaosClothAsset/Private/ChaosClothAsset/ClothCollection.cpp b/Engine/Plugins/ChaosClothAsset/Source/ChaosClothAsset/Private/ChaosClothAsset/ClothCollection.cpp
--- a/Engine/Plugins/ChaosClothAsset/Source/ChaosClothAsset/Private/ChaosClothAsset/ClothCollection.cpp
+++ b/Engine/Plugins/ChaosClothAsset/Source/ChaosClothAsset/Private/ChaosClothAsset/ClothCollection.cpp
@@ -551,7 +551,6 @@ namespace UE::Chaos::ClothAsset
 
 	void FClothCollection::SetNumElements(int32 InNumElements, const FName& GroupName)
 	{
-		check(IsValid());
 		check(InNumElements >= 0);
 		
 		const int32 NumElements = ManagedArrayCollection->NumElements(GroupName);
@@ -571,7 +570,6 @@ namespace UE::Chaos::ClothAsset
 
 	int32 FClothCollection::SetNumElements(int32 InNumElements, const FName& GroupName, TManagedArray<int32>* StartArray, TManagedArray<int32>* EndArray, int32 ArrayIndex)
 	{
-		check(IsValid());
 		check(InNumElements >= 0);
 
 		check(StartArray && StartArray->GetConstArray().IsValidIndex(ArrayIndex));
@@ -643,8 +641,6 @@ namespace UE::Chaos::ClothAsset
 			return;
 		}
 
-		check(IsValid());
-
 		check(StartArray && StartArray->GetConstArray().IsValidIndex(ArrayIndex));
 		check(EndArray && EndArray->GetConstArray().IsValidIndex(ArrayIndex));
 
diff --git a/Engine/Plugins/ChaosClothAsset/Source/ChaosClothAsset/Private/ChaosClothAsset/CollectionClothSimMorphTargetFacade.cpp b/Engine/Plugins/ChaosClothAsset/Source/ChaosClothAsset/Private/ChaosClothAsset/CollectionClothSimMorphTargetFacade.cpp
--- a/Engine/Plugins/ChaosClothAsset/Source/ChaosClothAsset/Private/ChaosClothAsset/CollectionClothSimMorphTargetFacade.cpp
+++ b/Engine/Plugins/ChaosClothAsset/Source/ChaosClothAsset/Private/ChaosClothAsset/CollectionClothSimMorphTargetFacade.cpp
@@ -65,7 +65,7 @@ namespace UE::Chaos::ClothAsset
 		: ClothCollection(ClothCollection)
 		, MorphTargetIndex(InMorphTargetIndex)
 	{
-		check(ClothCollection->IsValid());
+		check(ClothCollection->IsValid(EClothCollectionOptionalSchemas::CookedOnly | EClothCollectionOptionalSchemas::SimMorphTargets));
 		check(MorphTargetIndex >= 0 && MorphTargetIndex < ClothCollection->GetNumElements(ClothCollectionGroup::SimMorphTargets));
 	}
 
diff --git a/Engine/Plugins/ChaosClothAssetEditor/Source/ChaosClothAssetDataflowNodes/Private/ChaosClothAsset/SimulationMorphTargetConfigNode.cpp b/Engine/Plugins/ChaosClothAssetEditor/Source/ChaosClothAssetDataflowNodes/Private/ChaosClothAsset/SimulationMorphTargetConfigNode.cpp
--- a/Engine/Plugins/ChaosClothAssetEditor/Source/ChaosClothAssetDataflowNodes/Private/ChaosClothAsset/SimulationMorphTargetConfigNode.cpp
+++ b/Engine/Plugins/ChaosClothAssetEditor/Source/ChaosClothAssetDataflowNodes/Private/ChaosClothAsset/SimulationMorphTargetConfigNode.cpp
@@ -13,6 +13,5 @@ FChaosClothAssetSimulationMorphTargetConfigNode::FChaosClothAssetSimulationMorph
 
 void FChaosClothAssetSimulationMorphTargetConfigNode::AddProperties(FPropertyHelper& PropertyHelper) const
 {
-	PropertyHelper.SetPropertyString(FName(TEXT("ActiveMorphTarget")), ActiveMorphTarget.Name);
-	PropertyHelper.SetProperty(FName(TEXT("ActiveMorphTarget")), ActiveMorphTarget.Weight);
+	PropertyHelper.SetPropertyAndString(FName(TEXT("ActiveMorphTarget")), ActiveMorphTarget.Weight, ActiveMorphTarget.Name);
 }
diff --git a/Engine/Plugins/ChaosClothAssetEditor/Source/ChaosClothAssetDataflowNodes/Public/ChaosClothAsset/SimulationBaseConfigNode.h b/Engine/Plugins/ChaosClothAssetEditor/Source/ChaosClothAssetDataflowNodes/Public/ChaosClothAsset/SimulationBaseConfigNode.h
--- a/Engine/Plugins/ChaosClothAssetEditor/Source/ChaosClothAssetDataflowNodes/Public/ChaosClothAsset/SimulationBaseConfigNode.h
+++ b/Engine/Plugins/ChaosClothAssetEditor/Source/ChaosClothAssetDataflowNodes/Public/ChaosClothAsset/SimulationBaseConfigNode.h
@@ -51,6 +51,9 @@ protected:
 		template<typename PropertyType UE_REQUIRES(::Chaos::Softs::TIsWeightedType<PropertyType>::Value && !std::is_same_v<PropertyType, bool>)>
 		inline int32 SetProperty(const FName& PropertyName, const PropertyType& PropertyValue, const TArray<FName>& SimilarPropertyNames = {}, ECollectionPropertyFlags PropertyFlags = ECollectionPropertyFlags::Animatable);
 
+		template<typename PropertyType UE_REQUIRES(::Chaos::Softs::TIsWeightedType<PropertyType>::Value && !std::is_same_v<PropertyType, bool>)>
+		inline int32 SetPropertyAndString(const FName& PropertyName, const PropertyType& PropertyValue, const FString& StringValue, const TArray<FName>& SimilarPropertyNames = {}, ECollectionPropertyFlags PropertyFlags = ECollectionPropertyFlags::Animatable);
+
 		template<typename T, typename PropertyType UE_REQUIRES(std::is_base_of_v<FChaosClothAssetSimulationBaseConfigNode, T> && ::Chaos::Softs::TIsWeightedType<PropertyType>::Value && !std::is_same_v<PropertyType, bool>)>
 		inline int32 SetProperty(const T* ConfigStruct, const PropertyType* PropertyValue, const TArray<FName>& SimilarPropertyNames = {}, ECollectionPropertyFlags PropertyFlags = ECollectionPropertyFlags::Animatable);
 
@@ -190,6 +193,20 @@ inline int32 FChaosClothAssetSimulationBaseConfigNode::FPropertyHelper::SetPrope
 	return PropertyKeyIndex;
 }
 
+template<typename PropertyType UE_REQUIRES_DEFINITION(::Chaos::Softs::TIsWeightedType<PropertyType>::Value && !std::is_same_v<PropertyType, bool>)>
+inline int32 FChaosClothAssetSimulationBaseConfigNode::FPropertyHelper::SetPropertyAndString(
+	const FName& PropertyName,
+	const PropertyType& PropertyValue,
+	const FString& StringValue,
+	const TArray<FName>& SimilarPropertyNames,
+	ECollectionPropertyFlags PropertyFlags)
+{
+	const int32 PropertyKeyIndex = ConfigNode.AddPropertyHelper(Properties, PropertyName, SimilarPropertyNames, PropertyFlags);
+	Properties.SetValue(PropertyKeyIndex, PropertyValue);
+	Properties.SetStringValue(PropertyKeyIndex, StringValue);
+	return PropertyKeyIndex;
+}
+
 template<typename T, typename PropertyType UE_REQUIRES_DEFINITION(std::is_base_of_v<FChaosClothAssetSimulationBaseConfigNode, T> && ::Chaos::Softs::TIsWeightedType<PropertyType>::Value && !std::is_same_v<PropertyType, bool>)>
 inline int32 FChaosClothAssetSimulationBaseConfigNode::FPropertyHelper::SetProperty(
 	const T* ConfigStruct,
