Cloth Render Node: populate all UV Channels for output

diff --git a/Engine/Plugins/ChaosClothAssetEditor/Source/ChaosClothAssetDataflowNodes/Private/ChaosClothAsset/RemeshNode.cpp b/Engine/Plugins/ChaosClothAssetEditor/Source/ChaosClothAssetDataflowNodes/Private/ChaosClothAsset/RemeshNode.cpp
--- a/Engine/Plugins/ChaosClothAssetEditor/Source/ChaosClothAssetDataflowNodes/Private/ChaosClothAsset/RemeshNode.cpp
+++ b/Engine/Plugins/ChaosClothAssetEditor/Source/ChaosClothAssetDataflowNodes/Private/ChaosClothAsset/RemeshNode.cpp
@@ -1585,25 +1585,27 @@ namespace UE::Chaos::ClothAsset::Private
 		// UVs
 		//
 
-		// TODO: Account for multiple UV layers
-		TArray<FVector2f> UVs;
+		TArray<TStaticArray<FVector2f, MAX_TEXCOORDS>> UVs;
 		if (bHasUVs)
 		{
-			UVs.SetNum(DynamicMesh.VertexCount());
+			UVs.SetNumZeroed(DynamicMesh.VertexCount());
 
-			const FDynamicMeshUVOverlay* const UVOverlay = DynamicMesh.Attributes()->PrimaryUV();
-
-			// Assume no seams in the dynamic mesh UV overlay
-			checkSlow(!UVOverlay->HasInteriorSeamEdges());
-
-			for (const int TriangleIndex : DynamicMesh.TriangleIndicesItr())
+			for (int32 UVChannel = 0; UVChannel < DynamicMesh.Attributes()->NumUVLayers(); ++UVChannel)
 			{
-				const FIndex3i Tri = DynamicMesh.GetTriangle(TriangleIndex);
+				const FDynamicMeshUVOverlay* const UVOverlay = DynamicMesh.Attributes()->GetUVLayer(UVChannel);
 
-				for (int TriangleVertexIndex = 0; TriangleVertexIndex < 3; ++TriangleVertexIndex)
+				// Assume no seams in the dynamic mesh UV overlay
+				checkSlow(!UVOverlay->HasInteriorSeamEdges());
+
+				for (const int TriangleIndex : DynamicMesh.TriangleIndicesItr())
 				{
-					const int VertexIndex = Tri[TriangleVertexIndex];
-					UVs[VertexIndex] = UVOverlay->GetElementAtVertex(TriangleIndex, VertexIndex);
+					const FIndex3i Tri = DynamicMesh.GetTriangle(TriangleIndex);
+
+					for (int TriangleVertexIndex = 0; TriangleVertexIndex < 3; ++TriangleVertexIndex)
+					{
+						const int VertexIndex = Tri[TriangleVertexIndex];
+						UVs[VertexIndex][UVChannel] = UVOverlay->GetElementAtVertex(TriangleIndex, VertexIndex);
+					}
 				}
 			}
 		}
@@ -1710,7 +1712,10 @@ namespace UE::Chaos::ClothAsset::Private
 				if (bHasUVs)
 				{
 					RenderUVs[PatternVertexIndex].SetNum(MAX_TEXCOORDS);
-					RenderUVs[PatternVertexIndex][0] = UVs[SourceVertexIndex];
+					for (int32 UVChannel = 0; UVChannel < MAX_TEXCOORDS; ++UVChannel)
+					{
+						RenderUVs[PatternVertexIndex][UVChannel] = UVs[SourceVertexIndex][UVChannel];
+					}
 				}
 				if (bHasNormals)
 				{
