[Chaos] Implemented thumbnail renderer for Chaos Cloth assets

diff --git a/Engine/Plugins/ChaosClothAsset/Source/ChaosClothAssetEngine/Public/ChaosClothAsset/ClothComponent.h b/Engine/Plugins/ChaosClothAsset/Source/ChaosClothAssetEngine/Public/ChaosClothAsset/ClothComponent.h
--- a/Engine/Plugins/ChaosClothAsset/Source/ChaosClothAssetEngine/Public/ChaosClothAsset/ClothComponent.h
+++ b/Engine/Plugins/ChaosClothAsset/Source/ChaosClothAssetEngine/Public/ChaosClothAsset/ClothComponent.h
@@ -18,6 +18,7 @@ class UChaosClothAssetBase;
 class UChaosClothAsset;
 class UChaosClothComponent;
 class UChaosClothAssetInteractor;
+class UThumbnailInfo;
 struct FManagedArrayCollection;
 
 namespace Chaos::Softs
@@ -257,7 +258,15 @@ public:
 	void SetTickOnceInEditor() { bTickOnceInEditor = true; bTickInEditor = true; }
 #endif
 
+#if WITH_EDITORONLY_DATA
+	UThumbnailInfo* GetThumbnailInfo()
+	{
+		return ThumbnailInfo;
+	}
+#endif
+
 protected:
+	
 	//~ Begin UObject Interface
 	UE_API virtual void PostLoad() override;
 #if WITH_EDITOR
@@ -427,6 +436,12 @@ private:
 #if WITH_EDITOR
 	bool bTickOnceInEditor = false;
 #endif
+
+#if WITH_EDITORONLY_DATA
+	/** Information for thumbnail rendering */
+	UPROPERTY(VisibleAnywhere, Instanced, AdvancedDisplay, Category=StaticMesh)
+	TObjectPtr<UThumbnailInfo> ThumbnailInfo;
+#endif
 };
 
 #undef UE_API
diff --git a/Engine/Plugins/ChaosClothAssetEditor/Source/ChaosClothAssetEditor/Private/ChaosClothAsset/ChaosClothAssetEditorModule.cpp b/Engine/Plugins/ChaosClothAssetEditor/Source/ChaosClothAssetEditor/Private/ChaosClothAsset/ChaosClothAssetEditorModule.cpp
--- a/Engine/Plugins/ChaosClothAssetEditor/Source/ChaosClothAssetEditor/Private/ChaosClothAsset/ChaosClothAssetEditorModule.cpp
+++ b/Engine/Plugins/ChaosClothAssetEditor/Source/ChaosClothAssetEditor/Private/ChaosClothAsset/ChaosClothAssetEditorModule.cpp
@@ -1,7 +1,9 @@
 // Copyright Epic Games, Inc. All Rights Reserved.
 
 #include "ChaosClothAsset/ChaosClothAssetEditorModule.h"
+
 #include "ChaosClothAsset/AssetDefinition_ClothAsset.h"
+#include "ChaosClothAsset/ChaosClothAssetThumbnailRenderer.h"
 #include "ChaosClothAsset/ClothAsset.h"
 #include "ChaosClothAsset/ClothComponent.h"
 #include "ChaosClothAsset/ClothDataflowConstructionVisualization.h"
@@ -155,6 +157,8 @@ PRAGMA_ENABLE_DEPRECATION_WARNINGS
 
 		ClothAssetComponentBroker = MakeShareable(new FClothAssetComponentBroker);
 		FComponentAssetBrokerage::RegisterBroker(ClothAssetComponentBroker, UChaosClothComponent::StaticClass(), true, true);
+
+		UThumbnailManager::Get().RegisterCustomRenderer(UChaosClothAsset::StaticClass(), UChaosClothAssetThumbnailRenderer::StaticClass());
 	}
 
 	void FChaosClothAssetEditorModule::ShutdownModule()
diff --git a/Engine/Plugins/ChaosClothAssetEditor/Source/ChaosClothAssetEditor/Private/ChaosClothAsset/ChaosClothAssetThumbnailRenderer.cpp b/Engine/Plugins/ChaosClothAssetEditor/Source/ChaosClothAssetEditor/Private/ChaosClothAsset/ChaosClothAssetThumbnailRenderer.cpp
new file mode 100644
--- /dev/null
+++ b/Engine/Plugins/ChaosClothAssetEditor/Source/ChaosClothAssetEditor/Private/ChaosClothAsset/ChaosClothAssetThumbnailRenderer.cpp
@@ -0,0 +1,116 @@
+﻿// Copyright Epic Games, Inc. All Rights Reserved.
+
+#include "ChaosClothAsset/ChaosClothAssetThumbnailRenderer.h"
+
+#include "SceneView.h"
+#include "ShowFlags.h"
+#include "ChaosClothAsset/ClothAssetBase.h"
+#include "ChaosClothAsset/ClothComponent.h"
+#include "Components/SkeletalMeshComponent.h"
+#include "Engine/SkeletalMesh.h"
+#include "ThumbnailRendering/SceneThumbnailInfo.h"
+
+AChaosClothPreviewActor::AChaosClothPreviewActor()
+{
+	ClothComponent = CreateDefaultSubobject<UChaosClothComponent>(TEXT("ClothComponent0"));
+	RootComponent = ClothComponent;
+}
+
+namespace UE::Chaos::ClothAsset
+{
+	FThumbnailScene::FThumbnailScene() : FThumbnailPreviewScene()
+	{
+		bForceAllUsedMipsResident = false;
+
+		FActorSpawnParameters SpawnInfo;
+		SpawnInfo.SpawnCollisionHandlingOverride = ESpawnActorCollisionHandlingMethod::AlwaysSpawn;
+		SpawnInfo.bNoFail = true;
+		SpawnInfo.ObjectFlags = RF_Transient;
+
+		PreviewActor = GetWorld()->SpawnActor<AChaosClothPreviewActor>( SpawnInfo );
+		PreviewActor->SetActorEnableCollision(false);
+
+		check(PreviewActor);
+	}
+
+	void FThumbnailScene::SetClothAsset(UChaosClothAssetBase* InClothAsset)
+	{
+		UChaosClothComponent* ClothComponent = PreviewActor->GetClothComponent();
+		check(ClothComponent);
+		
+		ClothComponent->SetAsset(InClothAsset);
+	}
+
+	void FThumbnailScene::GetViewMatrixParameters(const float InFOVDegrees, FVector& OutOrigin, float& OutOrbitPitch, float& OutOrbitYaw,
+		float& OutOrbitZoom) const
+	{	
+		UChaosClothComponent* ClothComponent = PreviewActor->GetClothComponent();
+		check(ClothComponent);
+
+		FBoxSphereBounds Bounds = ClothComponent->Bounds;
+		Bounds = Bounds.ExpandBy(2.0f);
+		
+		const float HalfFOVRadians = FMath::DegreesToRadians<float>(InFOVDegrees) * 0.5f;
+		const float HalfMeshSize = static_cast<float>(Bounds.SphereRadius); 
+		const float TargetDistance = HalfMeshSize / FMath::Tan(HalfFOVRadians);
+
+		USceneThumbnailInfo* ThumbnailInfo = Cast<USceneThumbnailInfo>(ClothComponent->GetThumbnailInfo());
+		if (ThumbnailInfo)
+		{
+			if (TargetDistance + ThumbnailInfo->OrbitZoom < 0 )
+			{
+				ThumbnailInfo->OrbitZoom = -TargetDistance;
+			}
+		}
+		else
+		{
+			ThumbnailInfo = USceneThumbnailInfo::StaticClass()->GetDefaultObject<USceneThumbnailInfo>();
+		}
+
+		OutOrigin = -Bounds.Origin;
+		OutOrbitPitch = ThumbnailInfo->OrbitPitch;
+		OutOrbitYaw = ThumbnailInfo->OrbitYaw;
+		OutOrbitZoom = TargetDistance + ThumbnailInfo->OrbitZoom;
+	}
+}
+
+void UChaosClothAssetThumbnailRenderer::Draw(UObject* Object, int32 X, int32 Y, uint32 Width, uint32 Height, FRenderTarget* Viewport, FCanvas* Canvas, bool bAdditionalViewFamily)
+{
+	UChaosClothAssetBase* AsClothAsset = Cast<UChaosClothAssetBase>(Object);
+	if (!AsClothAsset)
+	{
+		return;
+	}
+	
+	TSharedRef<UE::Chaos::ClothAsset::FThumbnailScene> ThumbnailScene = ClothThumbnailSceneCache.EnsureThumbnailScene(Object);
+
+	ThumbnailScene->SetClothAsset(AsClothAsset);
+
+	FSceneViewFamilyContext ViewFamily( FSceneViewFamily::ConstructionValues( Viewport, ThumbnailScene->GetScene(), FEngineShowFlags(ESFIM_Game) )
+		.SetTime(UThumbnailRenderer::GetTime())
+		.SetAdditionalViewFamily(bAdditionalViewFamily));
+
+	ViewFamily.EngineShowFlags.DisableAdvancedFeatures();
+	ViewFamily.EngineShowFlags.MotionBlur = 0;
+	ViewFamily.EngineShowFlags.LOD = 0;
+
+	RenderViewFamily(Canvas, &ViewFamily, ThumbnailScene->CreateView(&ViewFamily, X, Y, Width, Height));
+	ThumbnailScene->SetClothAsset(nullptr);
+}
+
+bool UChaosClothAssetThumbnailRenderer::CanVisualizeAsset(UObject* Object)
+{
+	return Cast<UChaosClothAssetBase>(Object) != nullptr;
+}
+
+EThumbnailRenderFrequency UChaosClothAssetThumbnailRenderer::GetThumbnailRenderFrequency(UObject* Object) const
+{
+	UChaosClothAssetBase* AsClothAsset = Cast<UChaosClothAssetBase>(Object);
+	return AsClothAsset && AsClothAsset->GetResourceForRendering() ? EThumbnailRenderFrequency::Realtime : EThumbnailRenderFrequency::OnPropertyChange;
+}
+
+void UChaosClothAssetThumbnailRenderer::BeginDestroy()
+{
+	ClothThumbnailSceneCache.Clear();
+	Super::BeginDestroy();
+}
\ No newline at end of file
diff --git a/Engine/Plugins/ChaosClothAssetEditor/Source/ChaosClothAssetEditor/Public/ChaosClothAsset/ChaosClothAssetThumbnailRenderer.h b/Engine/Plugins/ChaosClothAssetEditor/Source/ChaosClothAssetEditor/Public/ChaosClothAsset/ChaosClothAssetThumbnailRenderer.h
new file mode 100644
--- /dev/null
+++ b/Engine/Plugins/ChaosClothAssetEditor/Source/ChaosClothAssetEditor/Public/ChaosClothAsset/ChaosClothAssetThumbnailRenderer.h
@@ -0,0 +1,69 @@
+﻿// Copyright Epic Games, Inc. All Rights Reserved.
+
+#pragma once
+
+#include "HAL/Platform.h"
+#include "ThumbnailRendering/SkeletalMeshThumbnailRenderer.h"
+
+#include "ChaosClothAssetThumbnailRenderer.generated.h"
+
+class UChaosClothComponent;
+class UChaosClothAssetBase;
+
+/**
+ * Minimal actor with a cloth component used to render cloth assets previews
+ */
+UCLASS(MinimalAPI)
+class AChaosClothPreviewActor : public AActor
+{
+	GENERATED_BODY()
+public:
+	AChaosClothPreviewActor();
+
+	UChaosClothComponent* GetClothComponent()
+	{
+		return ClothComponent;
+	}
+
+protected:
+	UPROPERTY()
+	TObjectPtr<UChaosClothComponent> ClothComponent;
+};
+
+namespace UE::Chaos::ClothAsset
+{
+	/**
+	 * Preview scene for a Cloth asset thumbnail
+	 */
+	class FThumbnailScene : public FThumbnailPreviewScene
+	{
+	public:
+		FThumbnailScene();
+
+		void SetClothAsset(UChaosClothAssetBase* InClothAsset);
+
+		virtual void GetViewMatrixParameters(const float InFOVDegrees, FVector& OutOrigin, float& OutOrbitPitch, float& OutOrbitYaw, float& OutOrbitZoom) const override;
+
+	private:
+		AChaosClothPreviewActor* PreviewActor;
+	};
+}
+
+
+/**
+ * Renders a preview scene to be used as a thumbnail for a supported cloth asset
+ */
+UCLASS(MinimalAPI)
+class UChaosClothAssetThumbnailRenderer : public UDefaultSizedThumbnailRenderer
+{
+	GENERATED_BODY()
+
+	virtual void Draw(UObject* Object, int32 X, int32 Y, uint32 Width, uint32 Height, FRenderTarget* Viewport, FCanvas* Canvas, bool bAdditionalViewFamily) override;
+	virtual bool CanVisualizeAsset(UObject* Object) override;
+	virtual EThumbnailRenderFrequency GetThumbnailRenderFrequency(UObject* Object) const override;
+
+	virtual void BeginDestroy() override;
+
+protected:
+	TObjectInstanceThumbnailScene<UE::Chaos::ClothAsset::FThumbnailScene, 128> ClothThumbnailSceneCache;
+};
diff --git a/Engine/Plugins/Experimental/ChaosOutfitAsset/ChaosOutfitAsset.uplugin b/Engine/Plugins/Experimental/ChaosOutfitAsset/ChaosOutfitAsset.uplugin
--- a/Engine/Plugins/Experimental/ChaosOutfitAsset/ChaosOutfitAsset.uplugin
+++ b/Engine/Plugins/Experimental/ChaosOutfitAsset/ChaosOutfitAsset.uplugin
@@ -45,6 +45,11 @@
 		{
 			"Name": "MeshResizing",
 			"Enabled": true
+		},
+		{
+			"Name": "ChaosClothAssetEditor",
+			"Enabled": true,
+			"TargetAllowList" : "Editor"
 		}
 	]
 }
diff --git a/Engine/Plugins/Experimental/ChaosOutfitAsset/Source/ChaosOutfitAssetEditor/ChaosOutfitAssetEditor.Build.cs b/Engine/Plugins/Experimental/ChaosOutfitAsset/Source/ChaosOutfitAssetEditor/ChaosOutfitAssetEditor.Build.cs
--- a/Engine/Plugins/Experimental/ChaosOutfitAsset/Source/ChaosOutfitAssetEditor/ChaosOutfitAssetEditor.Build.cs
+++ b/Engine/Plugins/Experimental/ChaosOutfitAsset/Source/ChaosOutfitAssetEditor/ChaosOutfitAssetEditor.Build.cs
@@ -15,6 +15,7 @@ public class ChaosOutfitAssetEditor : ModuleRules
 				"BaseCharacterFXEditor",
 				"ChaosClothAsset",
 				"ChaosClothAssetEngine",
+				"ChaosClothAssetEditor",
 				"ChaosOutfitAssetEngine",
 				"ContentBrowser",
 				"Core",
diff --git a/Engine/Plugins/Experimental/ChaosOutfitAsset/Source/ChaosOutfitAssetEditor/Private/ChaosOutfitAsset/ChaosOutfitAssetEditorModule.cpp b/Engine/Plugins/Experimental/ChaosOutfitAsset/Source/ChaosOutfitAssetEditor/Private/ChaosOutfitAsset/ChaosOutfitAssetEditorModule.cpp
--- a/Engine/Plugins/Experimental/ChaosOutfitAsset/Source/ChaosOutfitAssetEditor/Private/ChaosOutfitAsset/ChaosOutfitAssetEditorModule.cpp
+++ b/Engine/Plugins/Experimental/ChaosOutfitAsset/Source/ChaosOutfitAssetEditor/Private/ChaosOutfitAsset/ChaosOutfitAssetEditorModule.cpp
@@ -3,6 +3,7 @@
 #include "ChaosOutfitAsset/AssetDefinition_OutfitAsset.h"
 #include "Algo/AllOf.h"
 #include "AssetRegistry/AssetRegistryModule.h"
+#include "ChaosClothAsset/ChaosClothAssetThumbnailRenderer.h"
 #include "ChaosOutfitAsset/OutfitAsset.h"
 #include "ChaosOutfitAsset/OutfitAssetEditorStyle.h"
 #include "ChaosOutfitAsset/OutfitEditorCommands.h"
@@ -37,6 +38,8 @@ namespace UE::Chaos::OutfitAsset
 			StartupCallbackDelegateHandle = UToolMenus::RegisterStartupCallback(
 				FSimpleMulticastDelegate::FDelegate::CreateRaw(this, &FOutfitAssetEditorModule::RegisterMenus));
 
+			UThumbnailManager::Get().RegisterCustomRenderer(UChaosOutfitAsset::StaticClass(), UChaosClothAssetThumbnailRenderer::StaticClass());
+
 			// TODO Register visualizations (see FChaosClothAssetEditorModule::StartupModule() in ChaosClothAssetEditorModule.cpp)
 		}
 
