Update ProxyDeformer to optionally use the RenderTangentU when calculating. This matches the behavior the "default" proxy deformer that is run if no custom proxy deformer node is used. By default, existing ProxyDeformers will NOT do this to preserve existing behavior, but they are deprecated. A new version defaults to including the tangents.

Also updated the FClothGeometryTools::ApplyProxyDeformer to flip the sign of TangentV if the original TangentV is calculated as left-handed. This matches the cloth shader behavior.

Updated MakeResizableOutfitTemplate.uasset to use the new proxy deformer node. This means that any existing Resizable Outfits will need to be re-evaluated to pick up this change.

#virtualized

diff --git a/Engine/Plugins/ChaosClothAsset/Source/ChaosClothAsset/ChaosClothAsset.Build.cs b/Engine/Plugins/ChaosClothAsset/Source/ChaosClothAsset/ChaosClothAsset.Build.cs
--- a/Engine/Plugins/ChaosClothAsset/Source/ChaosClothAsset/ChaosClothAsset.Build.cs
+++ b/Engine/Plugins/ChaosClothAsset/Source/ChaosClothAsset/ChaosClothAsset.Build.cs
@@ -15,6 +15,7 @@ public class ChaosClothAsset : ModuleRules
 				"GeometryCore",
 				"MeshConversion",
 				"ClothingSystemRuntimeCommon",
+				"RenderCore"
 			}
 		);
 		PublicDependencyModuleNames.AddRange(
diff --git a/Engine/Plugins/ChaosClothAsset/Source/ChaosClothAsset/Private/ChaosClothAsset/ClothGeometryTools.cpp b/Engine/Plugins/ChaosClothAsset/Source/ChaosClothAsset/Private/ChaosClothAsset/ClothGeometryTools.cpp
--- a/Engine/Plugins/ChaosClothAsset/Source/ChaosClothAsset/Private/ChaosClothAsset/ClothGeometryTools.cpp
+++ b/Engine/Plugins/ChaosClothAsset/Source/ChaosClothAsset/Private/ChaosClothAsset/ClothGeometryTools.cpp
@@ -14,6 +14,7 @@
 #include "Util/IndexUtil.h"
 #include "Utils/ClothingMeshUtils.h"
 #include "Algo/RemoveIf.h"
+#include "RenderMath.h"
 
 namespace UE::Chaos::ClothAsset
 {
@@ -1786,7 +1787,10 @@ PRAGMA_ENABLE_DEPRECATION_WARNINGS
 								TangentU = (TangentPosition - AveragedSimPosition).GetSafeNormal();
 								Normal = (NormalPosition - AveragedSimPosition).GetSafeNormal();
 
-								TangentV = FVector3f::CrossProduct(TangentU, Normal).GetSafeNormal();
+								TangentV = FVector3f::CrossProduct(Normal, TangentU).GetSafeNormal();
+
+								// GetBasisDeterminantSign to determine if we need to flip TangentV.
+								TangentV *= GetBasisDeterminantSign(FVector(RenderTangentU[Index]), FVector(RenderTangentV[Index]), FVector(RenderNormal[Index]));
 							}
 							else
 							{
diff --git a/Engine/Plugins/ChaosClothAssetEditor/Source/ChaosClothAssetDataflowNodes/Private/ChaosClothAsset/ChaosClothAssetDataflowNodesModule.cpp b/Engine/Plugins/ChaosClothAssetEditor/Source/ChaosClothAssetDataflowNodes/Private/ChaosClothAsset/ChaosClothAssetDataflowNodesModule.cpp
--- a/Engine/Plugins/ChaosClothAssetEditor/Source/ChaosClothAssetDataflowNodes/Private/ChaosClothAsset/ChaosClothAssetDataflowNodesModule.cpp
+++ b/Engine/Plugins/ChaosClothAssetEditor/Source/ChaosClothAssetDataflowNodes/Private/ChaosClothAsset/ChaosClothAssetDataflowNodesModule.cpp
@@ -111,7 +111,7 @@ namespace UE::Chaos::ClothAsset
 			DATAFLOW_NODE_REGISTER_CREATION_FACTORY(FChaosClothAssetMakeClothAssetNode);
 			DATAFLOW_NODE_REGISTER_CREATION_FACTORY(FChaosClothAssetMergeClothCollectionsNode_v2);
 			DATAFLOW_NODE_REGISTER_CREATION_FACTORY(FChaosClothAssetProceduralSelectionNode);
-			DATAFLOW_NODE_REGISTER_CREATION_FACTORY(FChaosClothAssetProxyDeformerNode_v2);
+			DATAFLOW_NODE_REGISTER_CREATION_FACTORY(FChaosClothAssetProxyDeformerNode_v3);
 			DATAFLOW_NODE_REGISTER_CREATION_FACTORY(FChaosClothAssetRemeshNode_v2);
 			DATAFLOW_NODE_REGISTER_CREATION_FACTORY(FChaosClothAssetReverseNormalsNode);
 			DATAFLOW_NODE_REGISTER_CREATION_FACTORY(FChaosClothAssetSelectionNode_v2);
@@ -168,6 +168,7 @@ PRAGMA_DISABLE_DEPRECATION_WARNINGS
 			DATAFLOW_NODE_REGISTER_CREATION_FACTORY(FChaosClothAssetAddWeightMapNode);
 			DATAFLOW_NODE_REGISTER_CREATION_FACTORY(FChaosClothAssetMergeClothCollectionsNode);
 			DATAFLOW_NODE_REGISTER_CREATION_FACTORY(FChaosClothAssetProxyDeformerNode);
+			DATAFLOW_NODE_REGISTER_CREATION_FACTORY(FChaosClothAssetProxyDeformerNode_v2);
 			DATAFLOW_NODE_REGISTER_CREATION_FACTORY(FChaosClothAssetRemeshNode);
 			DATAFLOW_NODE_REGISTER_CREATION_FACTORY(FChaosClothAssetSelectionNode);
 			DATAFLOW_NODE_REGISTER_CREATION_FACTORY(FChaosClothAssetSimulationLongRangeAttachmentConfigNode);
diff --git a/Engine/Plugins/ChaosClothAssetEditor/Source/ChaosClothAssetDataflowNodes/Private/ChaosClothAsset/ProxyDeformerNode.cpp b/Engine/Plugins/ChaosClothAssetEditor/Source/ChaosClothAssetDataflowNodes/Private/ChaosClothAsset/ProxyDeformerNode.cpp
--- a/Engine/Plugins/ChaosClothAssetEditor/Source/ChaosClothAssetDataflowNodes/Private/ChaosClothAsset/ProxyDeformerNode.cpp
+++ b/Engine/Plugins/ChaosClothAssetEditor/Source/ChaosClothAssetDataflowNodes/Private/ChaosClothAsset/ProxyDeformerNode.cpp
@@ -23,6 +23,7 @@ namespace UE::Chaos::ClothAsset::Private
 		TConstArrayView<FIntVector3> SimIndices;
 		TConstArrayView<FVector3f> RenderPositions;
 		TConstArrayView<FVector3f> RenderNormals;
+		TConstArrayView<FVector3f> RenderTangentUs;
 		TConstArrayView<FIntVector3> RenderIndices;
 		FPointWeightMap PointWeightMap;
 		TArray<ClothingMeshUtils::FMeshToMeshFilterSet> MeshToMeshFilterSet;
@@ -67,7 +68,7 @@ namespace UE::Chaos::ClothAsset::Private
 			}
 
 			const ClothingMeshUtils::ClothMeshDesc SimMeshDesc(SimPositions, ScalarSimIndices);
-			const ClothingMeshUtils::ClothMeshDesc RenderMeshDesc(RenderPositions, RenderNormals, ScalarRenderIndices);
+			const ClothingMeshUtils::ClothMeshDesc RenderMeshDesc(RenderPositions, RenderNormals, RenderTangentUs, ScalarRenderIndices);
 
 			TArray<FMeshToMeshVertData> MeshToMeshVertData;
 
@@ -226,8 +227,14 @@ PRAGMA_ENABLE_DEPRECATION_WARNINGS
 
 }  // End namespace UE::Chaos::ClothAsset::Private
 
-FChaosClothAssetProxyDeformerNode_v2::FChaosClothAssetProxyDeformerNode_v2(const UE::Dataflow::FNodeParameters& InParam, FGuid InGuid)
+FChaosClothAssetProxyDeformerNode_v3::FChaosClothAssetProxyDeformerNode_v3(const UE::Dataflow::FNodeParameters& InParam, FGuid InGuid)
+	: FChaosClothAssetProxyDeformerNode_v2(InParam, InGuid, true)
+{
+}
+
+FChaosClothAssetProxyDeformerNode_v2::FChaosClothAssetProxyDeformerNode_v2(const UE::Dataflow::FNodeParameters& InParam, FGuid InGuid, bool bInPreserveRenderTangents)
 	: FDataflowNode(InParam, InGuid)
+	, bPreserveRenderTangents(bInPreserveRenderTangents)
 {
 	RegisterInputConnection(&Collection);
 	RegisterOutputConnection(&Collection)
@@ -270,6 +277,7 @@ void FChaosClothAssetProxyDeformerNode_v2::Evaluate(UE::Dataflow::FContext& Cont
 			DeformerMappingDataGenerator.SimIndices = ClothFacade.GetSimIndices3D();
 			DeformerMappingDataGenerator.RenderPositions = ClothFacade.GetRenderPosition();
 			DeformerMappingDataGenerator.RenderNormals = ClothFacade.GetRenderNormal();
+			DeformerMappingDataGenerator.RenderTangentUs = bPreserveRenderTangents ? ClothFacade.GetRenderTangentU() : TConstArrayView<FVector3f>();
 			DeformerMappingDataGenerator.RenderIndices = ClothFacade.GetRenderIndices();
 			DeformerMappingDataGenerator.PointWeightMap = FPointWeightMap(); // _v2 no longer compute SkinningBlend transitions
 			DeformerMappingDataGenerator.MeshToMeshFilterSet = Private::SelectionsToMeshToMeshFilterSets_v2(ClothCollection, GetSelectionFilterNames(Context));
@@ -477,6 +485,7 @@ void FChaosClothAssetProxyDeformerNode::Evaluate(UE::Dataflow::FContext& Context
 			DeformerMappingDataGenerator.SimIndices = ClothFacade.GetSimIndices3D();
 			DeformerMappingDataGenerator.RenderPositions = ClothFacade.GetRenderPosition();
 			DeformerMappingDataGenerator.RenderNormals = ClothFacade.GetRenderNormal();
+			DeformerMappingDataGenerator.RenderTangentUs = TConstArrayView<FVector3f>();
 			DeformerMappingDataGenerator.RenderIndices = ClothFacade.GetRenderIndices();
 			DeformerMappingDataGenerator.PointWeightMap = Private::SelectionToPointWeightMap(ClothCollection, ClothFacade, SimVertexSelectionName);
 			DeformerMappingDataGenerator.MeshToMeshFilterSet = Private::SelectionsToMeshToMeshFilterSets(ClothCollection, GetSelectionFilterNames(Context));
diff --git a/Engine/Plugins/ChaosClothAssetEditor/Source/ChaosClothAssetDataflowNodes/Public/ChaosClothAsset/ProxyDeformerNode.h b/Engine/Plugins/ChaosClothAssetEditor/Source/ChaosClothAssetDataflowNodes/Public/ChaosClothAsset/ProxyDeformerNode.h
--- a/Engine/Plugins/ChaosClothAssetEditor/Source/ChaosClothAssetDataflowNodes/Public/ChaosClothAsset/ProxyDeformerNode.h
+++ b/Engine/Plugins/ChaosClothAssetEditor/Source/ChaosClothAssetDataflowNodes/Public/ChaosClothAsset/ProxyDeformerNode.h
@@ -27,89 +27,13 @@
 	FChaosClothAssetConnectableIStringValue SimSelection = { TEXT("SelectionSimFilterSet") };
 };
 
-/**
- * Adds the proxy deformer information to this cloth collection's render data.
- * 
- * This node is only required to selectively assign specific areas of the render mesh to the sim mesh in the proxy deformer data.
- * When this node isn't provided to the construction graph, default proxy deformer data are automatically generated.
- * This node can be used to either override or remove the default proxy deformer data.
- * 
- * To enable deformer areas, selection filter sets must be provided.
- * With no selection filter sets, the Cloth Asset is by default fully skinned.
- * When more selection filter sets are needed, right-click on the node and do AddOptionPin to add more sets.
- * Each selection filter set is composed by two selections:
- * * A render mesh selection (natively a RenderVertices selection, but RenderFaces selections can also be used).
- * * A simulation mesh selection (natively a SimFaces selection, but SimVertices3D and SimVertices2D selections can also be used).
- * The render mesh selection of vertices will then only be deformed by the associated simulation mesh selection of faces.
- * The Proxy Deformer node evaluations are usually slower than the equivalent default proxy deformer due to the added processing cost
- * of the selection filter sets. When in use, it is therefore better to place it earlier in the graph.
- * 
- * A default SkinningBlend weight map is also generated by the node that masks the areas for which selection filters aren't provided.
- * The map can then be painted, or replaced by a smoothed transition map generated by the SkinningBlend node.
- * The SkinningBlend node can leverage information about the kinematic areas added by the MaxDistanceConfig node.
- * Unlike the default SkinningBlend weight map provided by the Proxy Deformer node, the SkinningBlend node generates areas of smooth
- * transition where the render mesh vertices can be simultaneously skinned and deformed by the cloth.
- * 
- * To sum up:
- * * No ProxyDeformer node -> The Cloth Asset is fully deformed.
- * * ProxyDeformer node without selection filters -> The Cloth Asset is fully skinned.
- * * ProxyDeformer node with selection filters -> the Cloth Asset is partially deformed/skinned.
- * * The ProxyDeformer can be slow to evaluate, always place it as early as possible in the graph to avoid costly re-evaluations.
- * * A SkinningBlend node placed after the MaxDistanceConfig adds smooth transitions between the skinned and proxy deformed areas.
- */
-USTRUCT(Meta = (DataflowCloth))
-struct FChaosClothAssetProxyDeformerNode_v2 : public FDataflowNode
-{
-	GENERATED_USTRUCT_BODY()
-	DATAFLOW_NODE_DEFINE_INTERNAL(FChaosClothAssetProxyDeformerNode_v2, "ProxyDeformer", "Cloth", "Cloth Simulation Proxy Deformer")
-	DATAFLOW_NODE_RENDER_TYPE("SurfaceRender", FName("FClothCollection"), "Collection")
-
-public:
-
-	UPROPERTY(Meta = (DataflowInput, DataflowOutput, DataflowPassthrough = "Collection"))
-	FManagedArrayCollection Collection;
-
-	/** Selection filter set used to restrict a selection of render vertices to a selection of simulation mesh triangles. Right click and do AddOptionPin to add more selection sets .*/
-	UPROPERTY(EditAnywhere, EditFixedSize, Category = "ProxyDeformer", Meta = (SkipInDisplayNameChain))
-	TArray<FChaosClothAssetSelectionFilterSet> SelectionFilterSets;
-
-	/** Whether using multiple simulation mesh triangles to influence the position of the deformed render vertex. */
-	UPROPERTY(EditAnywhere, Category = "ProxyDeformer")
-	bool bUseMultipleInfluences = false;
-
-	/** The radius around the render vertices to look for all simulation mesh triangles influencing it (AKA SkinningKernelRadius). */
-	UPROPERTY(EditAnywhere, Category = "ProxyDeformer", Meta = (EditCondition = "bUseMultipleInfluences"))
-	float InfluenceRadius = 5.f;
-
-	FChaosClothAssetProxyDeformerNode_v2(const UE::Dataflow::FNodeParameters& InParam, FGuid InGuid = FGuid::NewGuid());
-
-private:
-
-	//~ Begin FDataflowNode interface
-	virtual void Evaluate(UE::Dataflow::FContext& Context, const FDataflowOutput* Out) const override;
-	virtual TArray<UE::Dataflow::FPin> AddPins() override;
-	virtual bool CanAddPin() const override { return true; }
-	virtual bool CanRemovePin() const override { return SelectionFilterSets.Num() > NumInitialSelectionFilterSets; }
-	virtual TArray<UE::Dataflow::FPin> GetPinsToRemove() const override;
-	virtual void OnPinRemoved(const UE::Dataflow::FPin& Pin) override;
-	virtual void PostSerialize(const FArchive& Ar) override;
-	//~ End FDataflowNode interface
-
-	TArray<TPair<FName, FName>> GetSelectionFilterNames(UE::Dataflow::FContext& Context) const;
-	UE::Dataflow::TConnectionReference<FString> GetRenderConnectionReference(int32 Index) const;
-	UE::Dataflow::TConnectionReference<FString> GetSimConnectionReference(int32 Index) const;
-
-	static constexpr int32 NumRequiredInputs = 1; // non-filter set inputs
-	static constexpr int32 NumInitialSelectionFilterSets = 1;
-};
-
 /** Add the proxy deformer information to this cloth collection's render data. */
 USTRUCT(Meta = (DataflowCloth, Deprecated = "5.5"))
 struct UE_DEPRECATED(5.5, "Use the newer version of this node instead.") FChaosClothAssetProxyDeformerNode : public FDataflowNode
 {
 	GENERATED_USTRUCT_BODY()
 	DATAFLOW_NODE_DEFINE_INTERNAL(FChaosClothAssetProxyDeformerNode, "ProxyDeformer", "Cloth", "Cloth Simulation Proxy Deformer")
-	DATAFLOW_NODE_RENDER_TYPE("SurfaceRender", FName("FClothCollection"), "Collection")
+		DATAFLOW_NODE_RENDER_TYPE("SurfaceRender", FName("FClothCollection"), "Collection")
 
 public:
 
@@ -189,31 +113,157 @@
 	UPROPERTY(VisibleAnywhere, Category = "ProxyDeformer", Meta = (DataflowOutput))
 	FString SkinningBlendName;
 
-	FChaosClothAssetProxyDeformerNode(const UE::Dataflow::FNodeParameters& InParam, FGuid InGuid = FGuid::NewGuid());
+	FChaosClothAssetProxyDeformerNode(const UE::Dataflow::FNodeParameters & InParam, FGuid InGuid = FGuid::NewGuid());
 
 private:
 	/** Hardcoded number of FilterSets since it is currently not possible to use arrays for optional inputs. */
 	static constexpr int32 MaxNumFilterSets = 10;
 
 	/** The number of filter sets currently exposed to the node UI. */
-PRAGMA_DISABLE_DEPRECATION_WARNINGS
-	UPROPERTY()
+	PRAGMA_DISABLE_DEPRECATION_WARNINGS
+		UPROPERTY()
 	int32 NumFilterSets = NumInitialOptionalInputs;
-PRAGMA_ENABLE_DEPRECATION_WARNINGS
+	PRAGMA_ENABLE_DEPRECATION_WARNINGS
 
-	//~ Begin FDataflowNode interface
-	virtual void Evaluate(UE::Dataflow::FContext& Context, const FDataflowOutput* Out) const override;
+		//~ Begin FDataflowNode interface
+		virtual void Evaluate(UE::Dataflow::FContext & Context, const FDataflowOutput * Out) const override;
 	virtual TArray<UE::Dataflow::FPin> AddPins() override;
 	virtual bool CanAddPin() const override { return NumFilterSets < MaxNumFilterSets; }
 	virtual bool CanRemovePin() const override { return NumFilterSets > NumInitialOptionalInputs; }
 	virtual TArray<UE::Dataflow::FPin> GetPinsToRemove() const override;
-	virtual void OnPinRemoved(const UE::Dataflow::FPin& Pin) override;
-	virtual void PostSerialize(const FArchive& Ar) override;
+	virtual void OnPinRemoved(const UE::Dataflow::FPin & Pin) override;
+	virtual void PostSerialize(const FArchive & Ar) override;
 	//~ End FDataflowNode interface
 
-	TArray<FName> GetSelectionFilterNames(UE::Dataflow::FContext& Context) const;
+	TArray<FName> GetSelectionFilterNames(UE::Dataflow::FContext & Context) const;
 	TArray<const FChaosClothAssetConnectableStringValue*> Get1To9SelectionFilterSets() const;
 
 	static constexpr int32 NumRequiredInputs = 2; // non-filter set inputs
 	static constexpr int32 NumInitialOptionalInputs = 1; // filter set inputs that are created in the constructor.
 };
+
+/**
+ * Adds the proxy deformer information to this cloth collection's render data.
+ * This version of the node does not try to preserve existing Render Tangents by default to maintain existing behavior.
+ * 
+ * This node is only required to selectively assign specific areas of the render mesh to the sim mesh in the proxy deformer data.
+ * When this node isn't provided to the construction graph, default proxy deformer data are automatically generated.
+ * This node can be used to either override or remove the default proxy deformer data.
+ * 
+ * To enable deformer areas, selection filter sets must be provided.
+ * With no selection filter sets, the Cloth Asset is by default fully skinned.
+ * When more selection filter sets are needed, right-click on the node and do AddOptionPin to add more sets.
+ * Each selection filter set is composed by two selections:
+ * * A render mesh selection (natively a RenderVertices selection, but RenderFaces selections can also be used).
+ * * A simulation mesh selection (natively a SimFaces selection, but SimVertices3D and SimVertices2D selections can also be used).
+ * The render mesh selection of vertices will then only be deformed by the associated simulation mesh selection of faces.
+ * The Proxy Deformer node evaluations are usually slower than the equivalent default proxy deformer due to the added processing cost
+ * of the selection filter sets. When in use, it is therefore better to place it earlier in the graph.
+ * 
+ * A default SkinningBlend weight map is also generated by the node that masks the areas for which selection filters aren't provided.
+ * The map can then be painted, or replaced by a smoothed transition map generated by the SkinningBlend node.
+ * The SkinningBlend node can leverage information about the kinematic areas added by the MaxDistanceConfig node.
+ * Unlike the default SkinningBlend weight map provided by the Proxy Deformer node, the SkinningBlend node generates areas of smooth
+ * transition where the render mesh vertices can be simultaneously skinned and deformed by the cloth.
+ * 
+ * To sum up:
+ * * No ProxyDeformer node -> The Cloth Asset is fully deformed.
+ * * ProxyDeformer node without selection filters -> The Cloth Asset is fully skinned.
+ * * ProxyDeformer node with selection filters -> the Cloth Asset is partially deformed/skinned.
+ * * The ProxyDeformer can be slow to evaluate, always place it as early as possible in the graph to avoid costly re-evaluations.
+ * * A SkinningBlend node placed after the MaxDistanceConfig adds smooth transitions between the skinned and proxy deformed areas.
+ */
+USTRUCT(Meta = (DataflowCloth, Deprecated = "5.6"))
+struct FChaosClothAssetProxyDeformerNode_v2 : public FDataflowNode
+{
+	GENERATED_USTRUCT_BODY()
+	DATAFLOW_NODE_DEFINE_INTERNAL(FChaosClothAssetProxyDeformerNode_v2, "ProxyDeformer", "Cloth", "Cloth Simulation Proxy Deformer")
+	DATAFLOW_NODE_RENDER_TYPE("SurfaceRender", FName("FClothCollection"), "Collection")
+
+public:
+
+	UPROPERTY(Meta = (DataflowInput, DataflowOutput, DataflowPassthrough = "Collection"))
+	FManagedArrayCollection Collection;
+
+	/** Selection filter set used to restrict a selection of render vertices to a selection of simulation mesh triangles. Right click and do AddOptionPin to add more selection sets .*/
+	UPROPERTY(EditAnywhere, EditFixedSize, Category = "ProxyDeformer", Meta = (SkipInDisplayNameChain))
+	TArray<FChaosClothAssetSelectionFilterSet> SelectionFilterSets;
+
+	/** Whether using multiple simulation mesh triangles to influence the position of the deformed render vertex. */
+	UPROPERTY(EditAnywhere, Category = "ProxyDeformer")
+	bool bUseMultipleInfluences = false;
+
+	/** The radius around the render vertices to look for all simulation mesh triangles influencing it (AKA SkinningKernelRadius). */
+	UPROPERTY(EditAnywhere, Category = "ProxyDeformer", Meta = (EditCondition = "bUseMultipleInfluences"))
+	float InfluenceRadius = 5.f;
+
+	/** Whether or not to include the RenderTangents when generating the proxy deformer data. */
+	UPROPERTY(EditAnywhere, Category = "ProxyDeformer")
+	bool bPreserveRenderTangents = false;
+
+	FChaosClothAssetProxyDeformerNode_v2(const UE::Dataflow::FNodeParameters& InParam, FGuid InGuid = FGuid::NewGuid(), bool bInPreserveRenderTangents = false);
+
+protected:
+
+	//~ Begin FDataflowNode interface
+	virtual void Evaluate(UE::Dataflow::FContext& Context, const FDataflowOutput* Out) const override;
+	virtual TArray<UE::Dataflow::FPin> AddPins() override;
+	virtual bool CanAddPin() const override { return true; }
+	virtual bool CanRemovePin() const override { return SelectionFilterSets.Num() > NumInitialSelectionFilterSets; }
+	virtual TArray<UE::Dataflow::FPin> GetPinsToRemove() const override;
+	virtual void OnPinRemoved(const UE::Dataflow::FPin& Pin) override;
+	virtual void PostSerialize(const FArchive& Ar) override;
+	//~ End FDataflowNode interface
+
+private:
+
+	TArray<TPair<FName, FName>> GetSelectionFilterNames(UE::Dataflow::FContext& Context) const;
+	UE::Dataflow::TConnectionReference<FString> GetRenderConnectionReference(int32 Index) const;
+	UE::Dataflow::TConnectionReference<FString> GetSimConnectionReference(int32 Index) const;
+
+	static constexpr int32 NumRequiredInputs = 1; // non-filter set inputs
+	static constexpr int32 NumInitialSelectionFilterSets = 1;
+};
+
+
+/**
+ * Adds the proxy deformer information to this cloth collection's render data.
+ * This version will include the render tangents when generating the proxy deformer data by default (matches the default behavior if no proxy deformer node is added).
+ *
+ * This node is only required to selectively assign specific areas of the render mesh to the sim mesh in the proxy deformer data.
+ * When this node isn't provided to the construction graph, default proxy deformer data are automatically generated.
+ * This node can be used to either override or remove the default proxy deformer data.
+ *
+ * To enable deformer areas, selection filter sets must be provided.
+ * With no selection filter sets, the Cloth Asset is by default fully skinned.
+ * When more selection filter sets are needed, right-click on the node and do AddOptionPin to add more sets.
+ * Each selection filter set is composed by two selections:
+ * * A render mesh selection (natively a RenderVertices selection, but RenderFaces selections can also be used).
+ * * A simulation mesh selection (natively a SimFaces selection, but SimVertices3D and SimVertices2D selections can also be used).
+ * The render mesh selection of vertices will then only be deformed by the associated simulation mesh selection of faces.
+ * The Proxy Deformer node evaluations are usually slower than the equivalent default proxy deformer due to the added processing cost
+ * of the selection filter sets. When in use, it is therefore better to place it earlier in the graph.
+ *
+ * A default SkinningBlend weight map is also generated by the node that masks the areas for which selection filters aren't provided.
+ * The map can then be painted, or replaced by a smoothed transition map generated by the SkinningBlend node.
+ * The SkinningBlend node can leverage information about the kinematic areas added by the MaxDistanceConfig node.
+ * Unlike the default SkinningBlend weight map provided by the Proxy Deformer node, the SkinningBlend node generates areas of smooth
+ * transition where the render mesh vertices can be simultaneously skinned and deformed by the cloth.
+ *
+ * To sum up:
+ * * No ProxyDeformer node -> The Cloth Asset is fully deformed.
+ * * ProxyDeformer node without selection filters -> The Cloth Asset is fully skinned.
+ * * ProxyDeformer node with selection filters -> the Cloth Asset is partially deformed/skinned.
+ * * The ProxyDeformer can be slow to evaluate, always place it as early as possible in the graph to avoid costly re-evaluations.
+ * * A SkinningBlend node placed after the MaxDistanceConfig adds smooth transitions between the skinned and proxy deformed areas.
+ */
+USTRUCT(Meta = (DataflowCloth))
+struct FChaosClothAssetProxyDeformerNode_v3 : public FChaosClothAssetProxyDeformerNode_v2
+{
+	GENERATED_USTRUCT_BODY()
+	DATAFLOW_NODE_DEFINE_INTERNAL(FChaosClothAssetProxyDeformerNode_v3, "ProxyDeformer", "Cloth", "Cloth Simulation Proxy Deformer")
+	DATAFLOW_NODE_RENDER_TYPE("SurfaceRender", FName("FClothCollection"), "Collection")
+
+public:
+	FChaosClothAssetProxyDeformerNode_v3(const UE::Dataflow::FNodeParameters& InParam, FGuid InGuid = FGuid::NewGuid());
+};
