Cloth Editor: Fix crash when selecting an invalid dataflow node (e.g., unregistered) while in a tool.

diff --git a/Engine/Plugins/ChaosClothAssetEditor/Source/ChaosClothAssetEditor/Private/ChaosClothAsset/ClothEditorToolkit.cpp b/Engine/Plugins/ChaosClothAssetEditor/Source/ChaosClothAssetEditor/Private/ChaosClothAsset/ClothEditorToolkit.cpp
--- a/Engine/Plugins/ChaosClothAssetEditor/Source/ChaosClothAssetEditor/Private/ChaosClothAsset/ClothEditorToolkit.cpp
+++ b/Engine/Plugins/ChaosClothAssetEditor/Source/ChaosClothAssetEditor/Private/ChaosClothAsset/ClothEditorToolkit.cpp
@@ -1218,7 +1218,7 @@ void FChaosClothAssetEditorToolkit::OnNodeSelectionChanged(const TSet<UObject*>&
 							}
 						}
 						SelectedDataflowNode = Node->GetDataflowNode();
-						SelectedDataflowNodeGuid = SelectedDataflowNode->GetGuid();
+						SelectedDataflowNodeGuid = SelectedDataflowNode ? SelectedDataflowNode->GetGuid() : FGuid();
 
 						if (SelectedDataflowNode)
 						{
