Cloth Remesh Node: use the DynamicMesh UV layers to determine how many UV channels should be on each render vertex in the output ClothCollection (so that we don't accidentally add 4 channels when the input only has 1, for example)

diff --git a/Engine/Plugins/ChaosClothAssetEditor/Source/ChaosClothAssetDataflowNodes/Private/ChaosClothAsset/RemeshNode.cpp b/Engine/Plugins/ChaosClothAssetEditor/Source/ChaosClothAssetDataflowNodes/Private/ChaosClothAsset/RemeshNode.cpp
--- a/Engine/Plugins/ChaosClothAssetEditor/Source/ChaosClothAssetDataflowNodes/Private/ChaosClothAsset/RemeshNode.cpp
+++ b/Engine/Plugins/ChaosClothAssetEditor/Source/ChaosClothAssetDataflowNodes/Private/ChaosClothAsset/RemeshNode.cpp
@@ -1585,10 +1585,10 @@ namespace UE::Chaos::ClothAsset::Private
 		// UVs
 		//
 
-		TArray<TStaticArray<FVector2f, MAX_TEXCOORDS>> UVs;
+		TArray<TArray<FVector2f, TInlineAllocator<MAX_TEXCOORDS>>> UVs;
 		if (bHasUVs)
 		{
-			UVs.SetNumZeroed(DynamicMesh.VertexCount());
+			UVs.SetNum(DynamicMesh.VertexCount());
 
 			for (int32 UVChannel = 0; UVChannel < DynamicMesh.Attributes()->NumUVLayers(); ++UVChannel)
 			{
@@ -1600,11 +1600,16 @@ namespace UE::Chaos::ClothAsset::Private
 				for (const int TriangleIndex : DynamicMesh.TriangleIndicesItr())
 				{
 					const FIndex3i Tri = DynamicMesh.GetTriangle(TriangleIndex);
+					const FIndex3i OverlayTri = UVOverlay->GetTriangle(TriangleIndex);
 
 					for (int TriangleVertexIndex = 0; TriangleVertexIndex < 3; ++TriangleVertexIndex)
 					{
-						const int VertexIndex = Tri[TriangleVertexIndex];
-						UVs[VertexIndex][UVChannel] = UVOverlay->GetElementAtVertex(TriangleIndex, VertexIndex);
+						if (OverlayTri[TriangleVertexIndex] != FDynamicMesh3::InvalidID)
+						{
+							const int VertexIndex = Tri[TriangleVertexIndex];
+							UVs[VertexIndex].SetNumZeroed(UVChannel + 1);
+							UVs[VertexIndex][UVChannel] = UVOverlay->GetElement(OverlayTri[TriangleVertexIndex]);
+						}
 					}
 				}
 			}
@@ -1711,11 +1716,7 @@ namespace UE::Chaos::ClothAsset::Private
 				RenderPosition[PatternVertexIndex] = FVector3f(DynamicMesh.GetVertex(SourceVertexIndex));
 				if (bHasUVs)
 				{
-					RenderUVs[PatternVertexIndex].SetNum(MAX_TEXCOORDS);
-					for (int32 UVChannel = 0; UVChannel < MAX_TEXCOORDS; ++UVChannel)
-					{
-						RenderUVs[PatternVertexIndex][UVChannel] = UVs[SourceVertexIndex][UVChannel];
-					}
+					RenderUVs[PatternVertexIndex] = UVs[SourceVertexIndex];
 				}
 				if (bHasNormals)
 				{
