Chaos Outfit Asset:
- Fixed UV layers added in the ClothAssetUpdateClothFromDynamicMesh node.
- Fixed category names on the UVResizeController node.

#rnx

diff --git a/Engine/Plugins/ChaosClothAssetEditor/Source/ChaosClothAssetDataflowNodes/Private/ChaosClothAsset/ClothCollectionToDynamicMeshNode.cpp b/Engine/Plugins/ChaosClothAssetEditor/Source/ChaosClothAssetDataflowNodes/Private/ChaosClothAsset/ClothCollectionToDynamicMeshNode.cpp
--- a/Engine/Plugins/ChaosClothAssetEditor/Source/ChaosClothAssetDataflowNodes/Private/ChaosClothAsset/ClothCollectionToDynamicMeshNode.cpp
+++ b/Engine/Plugins/ChaosClothAssetEditor/Source/ChaosClothAssetDataflowNodes/Private/ChaosClothAsset/ClothCollectionToDynamicMeshNode.cpp
@@ -187,7 +187,7 @@ void FChaosClothAssetUpdateClothFromDynamicMeshNode::Evaluate(UE::Dataflow::FCon
 									const int32 RenderMeshVertexID = ClothMapping.GetOriginalVertexID(VertexID);
 									if (RenderUVs.IsValidIndex(RenderMeshVertexID))
 									{
-										if (UVIndex >= RenderUVs[RenderMeshVertexID].Num())
+										while (UVIndex >= RenderUVs[RenderMeshVertexID].Num())
 										{
 											RenderUVs[RenderMeshVertexID].AddZeroed();
 										}
diff --git a/Engine/Plugins/Experimental/MeshResizing/Source/MeshResizingNodes/Private/MeshResizing/UVResizeControllerNode.h b/Engine/Plugins/Experimental/MeshResizing/Source/MeshResizingNodes/Private/MeshResizing/UVResizeControllerNode.h
--- a/Engine/Plugins/Experimental/MeshResizing/Source/MeshResizingNodes/Private/MeshResizing/UVResizeControllerNode.h
+++ b/Engine/Plugins/Experimental/MeshResizing/Source/MeshResizingNodes/Private/MeshResizing/UVResizeControllerNode.h
@@ -24,11 +24,11 @@ public:
 
 private:
 	/** The texture name suffix . */
-	UPROPERTY(EditAnywhere, Category = "UVResizeLogicNode")
+	UPROPERTY(EditAnywhere, Category = "UVResizeController")
 	FString TextureSuffix = TEXT("Texture");
 
 	/** The suffix to replace the texture name with pointing to the UV channel index scalar parameter. */
-	UPROPERTY(EditAnywhere, Category = "UVResizeLogicNode")
+	UPROPERTY(EditAnywhere, Category = "UVResizeController")
 	FString UVChannelSuffix = TEXT("UVIndex");
 
 	/** The input/output Dataflow mesh. */
