Chaos Cloth Asset - Added missing LODs to the SkeletalMesh exporter.

#rnx

diff --git a/Engine/Plugins/ChaosClothAssetEditor/Source/ChaosClothAssetTools/Private/ChaosClothAsset/SkeletalMeshConverter.cpp b/Engine/Plugins/ChaosClothAssetEditor/Source/ChaosClothAssetTools/Private/ChaosClothAsset/SkeletalMeshConverter.cpp
--- a/Engine/Plugins/ChaosClothAssetEditor/Source/ChaosClothAssetTools/Private/ChaosClothAsset/SkeletalMeshConverter.cpp
+++ b/Engine/Plugins/ChaosClothAssetEditor/Source/ChaosClothAssetTools/Private/ChaosClothAsset/SkeletalMeshConverter.cpp
@@ -21,7 +21,9 @@ bool UClothAssetEditorSkeletalMeshConverter::ExportToSkeletalMesh(const UChaosCl
 	// Create one dynamic mesh from the multiple collections
 	TArray<UMaterialInterface*> RenderMaterials;
 	TArray<FDynamicMesh3> DynamicMeshes;
-	DynamicMeshes.AddDefaulted();  // TODO: Do the other LODs, one per Dynamic mesh
+
+	const int32 NumLods = ClothAssetBase.GetLODNum();
+	DynamicMeshes.SetNum(NumLods);
 
 	int32 MaterialOffset = 0;
 
@@ -31,10 +33,13 @@ bool UClothAssetEditorSkeletalMeshConverter::ExportToSkeletalMesh(const UChaosCl
 	{
 		const TArray<TSharedRef<const FManagedArrayCollection>>& ClothCollections = ClothAssetBase.GetCollections(ModelIndex);
 
-		//for (const TSharedRef<const FManagedArrayCollection>& ClothCollection : ClothCollections)
-		if (ClothCollections.Num())
+		for (int32 LodIndex = 0; LodIndex < NumLods; ++LodIndex)
 		{
-			const TSharedRef<const FManagedArrayCollection>& ClothCollection = ClothCollections[0];  // TODO: Do other LODs
+			if (!ClothCollections.IsValidIndex(LodIndex))
+			{
+				continue;
+			}
+			const TSharedRef<const FManagedArrayCollection>& ClothCollection = ClothCollections[LodIndex];
 			FCollectionClothConstFacade ClothFacade(ClothCollection);
 
 			if (ClothFacade.IsValid())
@@ -54,13 +59,13 @@ bool UClothAssetEditorSkeletalMeshConverter::ExportToSkeletalMesh(const UChaosCl
 					bDisableAttributes,
 					MaterialOffset);
 
-				if (ModelIndex == 0)
+				if (!DynamicMeshes[LodIndex].VertexCount() || !DynamicMeshes[LodIndex].TriangleCount())
 				{
-					DynamicMeshes[0] = MoveTemp(DynamicMeshPiece);
+					DynamicMeshes[LodIndex] = MoveTemp(DynamicMeshPiece);
 				}
 				else
 				{
-					DynamicMeshes[0].AppendWithOffsets(DynamicMeshPiece);
+					DynamicMeshes[LodIndex].AppendWithOffsets(DynamicMeshPiece);
 				}
 
 				const TArrayView<const FString> MaterialPaths = ClothFacade.GetRenderMaterialPathName();
