Transfer Skin Weights make LodIndex a DataflowInput. Fallback to last valid LOD if LodIndex is invalid.

Update ResizeOutfitTemplate asset to transfer skin weights using the same LODIndex on the target mesh as the current cloth LODIndex. Can override this to always use LOD 0 from the variables, but default behavior is to sync the LODs.

diff --git a/Engine/Plugins/ChaosClothAssetEditor/Source/ChaosClothAssetDataflowNodes/Private/ChaosClothAsset/TransferSkinWeightsNode.cpp b/Engine/Plugins/ChaosClothAssetEditor/Source/ChaosClothAssetDataflowNodes/Private/ChaosClothAsset/TransferSkinWeightsNode.cpp
--- a/Engine/Plugins/ChaosClothAssetEditor/Source/ChaosClothAssetDataflowNodes/Private/ChaosClothAsset/TransferSkinWeightsNode.cpp
+++ b/Engine/Plugins/ChaosClothAssetEditor/Source/ChaosClothAssetDataflowNodes/Private/ChaosClothAsset/TransferSkinWeightsNode.cpp
@@ -804,6 +804,7 @@ FChaosClothAssetTransferSkinWeightsNode::FChaosClothAssetTransferSkinWeightsNode
 {
 	RegisterInputConnection(&Collection);
 	RegisterInputConnection(&SkeletalMesh);
+	RegisterInputConnection(&LodIndex);
 	RegisterInputConnection(&SimCollection)
 		.SetCanHidePin(true)
 		.SetPinIsHidden(true);
@@ -857,25 +858,42 @@ void FChaosClothAssetTransferSkinWeightsNode::Evaluate(UE::Dataflow::FContext& C
 					return;
 				}
 
-				if (!InSkeletalMesh->IsValidLODIndex(LodIndex))
+				int32 InLodIndex = GetValue(Context, &LodIndex);
+				if (!InSkeletalMesh->IsValidLODIndex(InLodIndex))
 				{
-					FClothDataflowTools::LogAndToastWarning(*this,
-						LOCTEXT("InvalidLodIndexHeadline", "Invalid LOD Index."),
-						FText::Format(
-							LOCTEXT("InvalidLodIndexDetails", "LOD index {0} is not a valid LOD for skeletal mesh {1}."),
-							LodIndex,
-							FText::FromString(InSkeletalMesh->GetName())));
-					SetValue(Context, MoveTemp(*ClothCollection), &Collection);
-					return;
+					const int32 LastLODIndex = InSkeletalMesh->GetLODNum() - 1;
+					if (InSkeletalMesh->IsValidLODIndex(LastLODIndex))
+					{
+						FClothDataflowTools::LogAndToastWarning(*this,
+							LOCTEXT("InvalidLodIndexHeadline", "Invalid LOD Index."),
+							FText::Format(
+								LOCTEXT("InvalidLodIndexSubstituteDetails", "LOD index {0} is not a valid LOD for skeletal mesh {1}. Using LOD index {2} instead."),
+								InLodIndex,
+								FText::FromString(InSkeletalMesh->GetName()),
+								LastLODIndex));
+
+						InLodIndex = LastLODIndex;
+					}
+					else
+					{
+						FClothDataflowTools::LogAndToastWarning(*this,
+							LOCTEXT("InvalidLodIndexHeadline", "Invalid LOD Index."),
+							FText::Format(
+								LOCTEXT("InvalidLodIndexDetails", "LOD index {0} is not a valid LOD for skeletal mesh {1}."),
+								InLodIndex,
+								FText::FromString(InSkeletalMesh->GetName())));
+						SetValue(Context, MoveTemp(*ClothCollection), &Collection);
+						return;
+					}
 				}
 
-				if (!SkeletalMeshToDynamicMesh(InSkeletalMesh.Get(), LodIndex, SourceDynamicMesh))
+				if (!SkeletalMeshToDynamicMesh(InSkeletalMesh.Get(), InLodIndex, SourceDynamicMesh))
 				{
 					FClothDataflowTools::LogAndToastWarning(*this,
 						LOCTEXT("InvalidLodHeadline", "Could not convert LOD to Dynamic Mesh."),
 						FText::Format(
 							LOCTEXT("InvalidLodDetails", "Could not convert LOD index {0} of the skeletal mesh {1} to dyanamic mesh."),
-							LodIndex,
+							InLodIndex,
 							FText::FromString(InSkeletalMesh->GetName())));
 					SetValue(Context, MoveTemp(*ClothCollection), &Collection);
 					return;
diff --git a/Engine/Plugins/ChaosClothAssetEditor/Source/ChaosClothAssetDataflowNodes/Public/ChaosClothAsset/TransferSkinWeightsNode.h b/Engine/Plugins/ChaosClothAssetEditor/Source/ChaosClothAssetDataflowNodes/Public/ChaosClothAsset/TransferSkinWeightsNode.h
--- a/Engine/Plugins/ChaosClothAssetEditor/Source/ChaosClothAssetDataflowNodes/Public/ChaosClothAsset/TransferSkinWeightsNode.h
+++ b/Engine/Plugins/ChaosClothAssetEditor/Source/ChaosClothAssetDataflowNodes/Public/ChaosClothAsset/TransferSkinWeightsNode.h
@@ -98,7 +98,7 @@ public:
 	FManagedArrayCollection SimCollection;
 
 	/** The skeletal mesh LOD to transfer the skin weights from. */
-	UPROPERTY(EditAnywhere, Category = "Transfer Skin Weights", Meta = (DisplayName = "LOD Index", EditCondition = "TargetMeshType != EChaosClothAssetTransferTargetMeshType::Render || RenderMeshSourceType == EChaosClothAssetTransferRenderMeshSource::SkeletalMesh", EditConditionHides))
+	UPROPERTY(EditAnywhere, Category = "Transfer Skin Weights", Meta = (DisplayName = "LOD Index", EditCondition = "TargetMeshType != EChaosClothAssetTransferTargetMeshType::Render || RenderMeshSourceType == EChaosClothAssetTransferRenderMeshSource::SkeletalMesh", EditConditionHides, DataflowInput))
 	int32 LodIndex = 0;
 
 	/** The relative transform between the skeletal mesh and the cloth asset. */
