Add GenerateSimMorphTarget dataflow node to generate morph targets from sim caches (or other ways of populating a cloth sim collection).

diff --git a/Engine/Plugins/ChaosCloth/Source/ChaosCloth/Private/ChaosCloth/ChaosClothVisualization.cpp b/Engine/Plugins/ChaosCloth/Source/ChaosCloth/Private/ChaosCloth/ChaosClothVisualization.cpp
--- a/Engine/Plugins/ChaosCloth/Source/ChaosCloth/Private/ChaosCloth/ChaosClothVisualization.cpp
+++ b/Engine/Plugins/ChaosCloth/Source/ChaosCloth/Private/ChaosCloth/ChaosClothVisualization.cpp
@@ -3067,6 +3067,7 @@ namespace Chaos
 			}
 		}
 	}
+
 	void FClothVisualizationNoGC::DrawTeleportReset(FPrimitiveDrawInterface* PDI) const
 	{
 		if (!Solver)
@@ -3166,6 +3167,16 @@ namespace Chaos
 		return AllNames.Array();
 	}
 
+	TArray<FString> FClothVisualizationNoGC::GetAllMorphTargetNames() const
+	{
+		TSet<FString> AllNames;
+		for (const FClothingSimulationCloth* const Cloth : Solver->GetCloths())
+		{
+			AllNames.Append(Cloth->GetAllMorphTargetNames());
+		}
+		return AllNames.Array();
+	}
+
 	void FClothVisualizationNoGC::DrawKinematicColliderShaded(FPrimitiveDrawInterface* PDI) const
 	{
 		if (!Solver || !FMaterials::GetInstance().GetCollisionMaterial())
@@ -3213,6 +3224,66 @@ namespace Chaos
 		const FMatrix LocalSimSpaceToWorld = GetLocalSpaceToWorld(*Solver);
 		MeshBuilder.Draw(PDI, LocalSimSpaceToWorld, FMaterials::GetInstance().GetCollisionMaterial()->GetRenderProxy(), Chaos::Private::GetDepthPriority(), false, false);
 	}
+
+	void FClothVisualizationNoGC::DrawSimMorphTarget(FPrimitiveDrawInterface* PDI, const FString& Name) const
+	{
+		if (!Solver)
+		{
+			return;
+		}
+		static const FLinearColor Color = FColor::Turquoise;
+		for (const FClothingSimulationCloth* const Cloth : Solver->GetCloths())
+		{
+			const int32 ParticleRangeId = Cloth->GetParticleRangeId(Solver);
+			if (ParticleRangeId == INDEX_NONE)
+			{
+				continue;
+			}
+
+			const int32 LODIndex = Cloth->GetLODIndex(Solver);
+
+			if (const FClothingSimulationMesh* const Mesh = Cloth->GetMesh())
+			{
+				int32 MorphTargetIndex = Mesh->FindMorphTargetByName(LODIndex, Name);
+				if (MorphTargetIndex == INDEX_NONE)
+				{
+					MorphTargetIndex = Cloth->GetCurrentMorphTargetIndex(Solver);
+				}
+
+				if (MorphTargetIndex != INDEX_NONE)
+				{
+					TConstArrayView<FVector3f> PositionDeltas = Mesh->GetMorphTargetPositionDeltas(LODIndex, MorphTargetIndex);
+					TConstArrayView<int32> Indices = Mesh->GetMorphTargetIndices(LODIndex, MorphTargetIndex);
+					TConstArrayView<FVector3f> OrigPositions = Mesh->GetPositions(LODIndex);
+
+					TArray<FVector3f> Positions(OrigPositions.GetData(), OrigPositions.Num());
+					for (int32 Index = 0; Index < Indices.Num(); ++Index)
+					{
+						const int32 Vertex = Indices[Index];
+						Positions[Vertex] += PositionDeltas[Index];
+					}
+
+
+					// Elements are local indexed for new solver
+					const int32 Offset = Solver->IsLegacySolver() ? ParticleRangeId : 0;
+
+					const TConstArrayView<TVec3<int32>> Elements = Cloth->GetTriangleMesh(Solver).GetElements();
+					for (int32 ElementIndex = 0; ElementIndex < Elements.Num(); ++ElementIndex)
+					{
+						const TVec3<int32>& Element = Elements[ElementIndex];
+
+						const FVector Pos0 = GetWorldPosition(*Solver, Positions[Element.X - Offset]); // TODO: Triangle Mesh shouldn't really be solver dependent (ie not use an offset)
+						const FVector Pos1 = GetWorldPosition(*Solver, Positions[Element.Y - Offset]);
+						const FVector Pos2 = GetWorldPosition(*Solver, Positions[Element.Z - Offset]);
+
+						DrawLine(PDI, Pos0, Pos1, Color);
+						DrawLine(PDI, Pos1, Pos2, Color);
+						DrawLine(PDI, Pos2, Pos0, Color);
+					}
+				}
+			}
+		}
+	}
 #endif // #if WITH_EDITOR
 #else  // #if CHAOS_DEBUG_DRAW
 namespace Chaos
diff --git a/Engine/Plugins/ChaosCloth/Source/ChaosCloth/Private/ChaosCloth/ChaosClothingSimulationCloth.cpp b/Engine/Plugins/ChaosCloth/Source/ChaosCloth/Private/ChaosCloth/ChaosClothingSimulationCloth.cpp
--- a/Engine/Plugins/ChaosCloth/Source/ChaosCloth/Private/ChaosCloth/ChaosClothingSimulationCloth.cpp
+++ b/Engine/Plugins/ChaosCloth/Source/ChaosCloth/Private/ChaosCloth/ChaosClothingSimulationCloth.cpp
@@ -814,6 +814,28 @@ int32 FClothingSimulationCloth::GetReferenceBoneIndex() const
 	return Mesh ? Mesh->GetReferenceBoneIndex() : INDEX_NONE;
 }
 
+int32 FClothingSimulationCloth::GetCurrentMorphTargetIndex(const FClothingSimulationSolver* Solver) const
+{
+	const int32 LODIndex = LODIndices.FindChecked(Solver);
+	return LODData.IsValidIndex(LODIndex) ? LODData[LODIndex]->ActiveMorphTarget : INDEX_NONE;
+}
+
+FRealSingle FClothingSimulationCloth::GetCurrentMorphTargetWeight(const FClothingSimulationSolver* Solver) const
+{
+	const int32 LODIndex = LODIndices.FindChecked(Solver);
+	return LODData.IsValidIndex(LODIndex) ? LODData[LODIndex]->ActiveMorphTargetWeight : 0.f;
+}
+
+TSet<FString> FClothingSimulationCloth::GetAllMorphTargetNames() const
+{
+	TSet<FString> Names;
+	for (int32 LODIndex = 0; LODIndex < LODData.Num(); ++LODIndex)
+	{
+		Names.Append(Mesh->GetAllMorphTargetNames(LODIndex));
+	}
+	return Names;
+}
+
 void FClothingSimulationCloth::PreUpdate(FClothingSimulationSolver* Solver)
 {
 	check(Solver);
diff --git a/Engine/Plugins/ChaosCloth/Source/ChaosCloth/Public/ChaosCloth/ChaosClothVisualization.h b/Engine/Plugins/ChaosCloth/Source/ChaosCloth/Public/ChaosCloth/ChaosClothVisualization.h
--- a/Engine/Plugins/ChaosCloth/Source/ChaosCloth/Public/ChaosCloth/ChaosClothVisualization.h
+++ b/Engine/Plugins/ChaosCloth/Source/ChaosCloth/Public/ChaosCloth/ChaosClothVisualization.h
@@ -93,7 +93,6 @@ namespace Chaos
 		void DrawFictitiousAngularForces(FPrimitiveDrawInterface* /*PDI*/ = nullptr) const {}
 		void DrawMultiResConstraint(FPrimitiveDrawInterface* /*PDI*/ = nullptr) const {}
 		void DrawClothClothConstraints(FPrimitiveDrawInterface* /*PDI*/ = nullptr) const {}
-		void DrawTeleportReset(FPrimitiveDrawInterface* /*PDI*/ = nullptr) const {}
 		void DrawExtremlyDeformedEdges(FPrimitiveDrawInterface* /*PDI*/ = nullptr) const {}
 #endif  // #if CHAOS_DEBUG_DRAW
 
@@ -105,7 +104,10 @@ namespace Chaos
 		CHAOSCLOTH_API void DrawInpaintWeightsMatched(FPrimitiveDrawInterface* PDI = nullptr) const;
 		CHAOSCLOTH_API void DrawKinematicColliderShaded(FPrimitiveDrawInterface* PDI) const;
 		CHAOSCLOTH_API void DrawWeightMapWithName(FPrimitiveDrawInterface* PDI, const FString& Name) const;
+		/** Will draw current morph target (if any) if no name or an invalid name is specified */
+		CHAOSCLOTH_API void DrawSimMorphTarget(FPrimitiveDrawInterface* PDI = nullptr, const FString& Name = "") const;
 		CHAOSCLOTH_API TArray<FString> GetAllWeightMapNames() const;
+		CHAOSCLOTH_API TArray<FString> GetAllMorphTargetNames() const;
 #else  // #if WITH_EDITOR && CHAOS_DEBUG_DRAW
 		void DrawPhysMeshShaded(FPrimitiveDrawInterface* /*PDI*/) const {}
 		void DrawKinematicColliderShaded(FPrimitiveDrawInterface* /*PDI*/) const {}
diff --git a/Engine/Plugins/ChaosCloth/Source/ChaosCloth/Public/ChaosCloth/ChaosClothingSimulationCloth.h b/Engine/Plugins/ChaosCloth/Source/ChaosCloth/Public/ChaosCloth/ChaosClothingSimulationCloth.h
--- a/Engine/Plugins/ChaosCloth/Source/ChaosCloth/Public/ChaosCloth/ChaosClothingSimulationCloth.h
+++ b/Engine/Plugins/ChaosCloth/Source/ChaosCloth/Public/ChaosCloth/ChaosClothingSimulationCloth.h
@@ -114,6 +114,9 @@ namespace Chaos
 		CHAOSCLOTH_API int32 GetReferenceBoneIndex() const;
 		// Return the local reference space transform for this cloth.
 		const FRigidTransform3& GetReferenceSpaceTransform() const { return ReferenceSpaceTransform;  }
+		CHAOSCLOTH_API int32 GetCurrentMorphTargetIndex(const FClothingSimulationSolver* Solver) const;
+		CHAOSCLOTH_API FRealSingle GetCurrentMorphTargetWeight(const FClothingSimulationSolver* Solver) const;
+		CHAOSCLOTH_API TSet<FString> GetAllMorphTargetNames() const;
 #if CHAOS_DEBUG_DRAW
 		FRealSingle GetTimeSinceLastTeleport() const { return TimeSinceLastTeleport; }
 		FRealSingle GetTimeSinceLastReset() const { return TimeSinceLastReset; }
diff --git a/Engine/Plugins/ChaosCloth/Source/ChaosCloth/Public/ChaosCloth/ChaosClothingSimulationMesh.h b/Engine/Plugins/ChaosCloth/Source/ChaosCloth/Public/ChaosCloth/ChaosClothingSimulationMesh.h
--- a/Engine/Plugins/ChaosCloth/Source/ChaosCloth/Public/ChaosCloth/ChaosClothingSimulationMesh.h
+++ b/Engine/Plugins/ChaosCloth/Source/ChaosCloth/Public/ChaosCloth/ChaosClothingSimulationMesh.h
@@ -149,16 +149,25 @@ namespace Chaos
 			return INDEX_NONE;
 		}
 
+		/** Get a list of all MorphTargets for a given LODIndex. (Index matches FindMorphTargetByName) */
+		virtual TConstArrayView<FString> GetAllMorphTargetNames(int32 LODIndex) const
+		{
+			return TConstArrayView<FString>();
+		}
+
+		/** Get all Morph Target position deltas for a given MorphTargetIndex (e.g., index returned by FindMorphTargetByName). Deltas index back to Positions via MorphTargetIndices */
 		virtual TConstArrayView<FVector3f> GetMorphTargetPositionDeltas(int32 LODIndex, int32 MorphTargetIndex) const
 		{
 			return TConstArrayView<FVector3f>();
 		}
 
+		/** Get all Morph Target tangent z (normal) deltas for a given MorphTargetIndex (e.g., index returned by FindMorphTargetByName). Deltas index back to Normals via MorphTargetIndices */
 		virtual TConstArrayView<FVector3f> GetMorphTargetTangentZDeltas(int32 LODIndex, int32 MorphTargetIndex) const
 		{
 			return TConstArrayView<FVector3f>();
 		}
 
+		/** Get all Morph Target indices for a given MorphTargetIndex (e.g., index returned by FindMorphTargetByName). These indices can map MorphTargetDeltas back to Positions. */
 		virtual TConstArrayView<int32> GetMorphTargetIndices(int32 LODIndex, int32 MorphTargetIndex) const
 		{
 			return TConstArrayView<int32>();
diff --git a/Engine/Plugins/ChaosClothAsset/Source/ChaosClothAssetEngine/Private/ChaosClothAsset/ClothSimulationMesh.cpp b/Engine/Plugins/ChaosClothAsset/Source/ChaosClothAssetEngine/Private/ChaosClothAsset/ClothSimulationMesh.cpp
--- a/Engine/Plugins/ChaosClothAsset/Source/ChaosClothAssetEngine/Private/ChaosClothAsset/ClothSimulationMesh.cpp
+++ b/Engine/Plugins/ChaosClothAsset/Source/ChaosClothAssetEngine/Private/ChaosClothAsset/ClothSimulationMesh.cpp
@@ -251,6 +251,15 @@ namespace UE::Chaos::ClothAsset
 		return ClothFacades.IsValidIndex(LODIndex) ? ClothFacades[LODIndex]->FindSimMorphTargetIndexByName(Name) : INDEX_NONE;
 	}
 
+	TConstArrayView<FString> FClothSimulationMesh::GetAllMorphTargetNames(int32 LODIndex) const
+	{
+		if (ClothFacades.IsValidIndex(LODIndex))
+		{
+			return ClothFacades[LODIndex]->GetSimMorphTargetName();
+		}
+		return TConstArrayView<FString>();
+	}
+
 	TConstArrayView<FVector3f> FClothSimulationMesh::GetMorphTargetPositionDeltas(int32 LODIndex, int32 MorphTargetIndex) const
 	{
 		if (ClothFacades.IsValidIndex(LODIndex) && (MorphTargetIndex >= 0 && MorphTargetIndex < ClothFacades[LODIndex]->GetNumSimMorphTargets()))
diff --git a/Engine/Plugins/ChaosClothAsset/Source/ChaosClothAssetEngine/Private/ChaosClothAsset/ClothSimulationMesh.h b/Engine/Plugins/ChaosClothAsset/Source/ChaosClothAssetEngine/Private/ChaosClothAsset/ClothSimulationMesh.h
--- a/Engine/Plugins/ChaosClothAsset/Source/ChaosClothAssetEngine/Private/ChaosClothAsset/ClothSimulationMesh.h
+++ b/Engine/Plugins/ChaosClothAsset/Source/ChaosClothAssetEngine/Private/ChaosClothAsset/ClothSimulationMesh.h
@@ -59,6 +59,7 @@ namespace UE::Chaos::ClothAsset
 		virtual TConstArrayView<FMeshToMeshVertData> GetTransitionDownSkinData(int32 LODIndex) const override;
 		virtual TSharedPtr<const FManagedArrayCollection> GetManagedArrayCollection(int32 LODIndex) const override;
 		virtual int32 FindMorphTargetByName(int32 LODIndex, const FString& Name) const override;
+		virtual TConstArrayView<FString> GetAllMorphTargetNames(int32 LODIndex) const override;
 		virtual TConstArrayView<FVector3f> GetMorphTargetPositionDeltas(int32 LODIndex, int32 MorphTargetIndex) const override;
 		virtual TConstArrayView<FVector3f> GetMorphTargetTangentZDeltas(int32 LODIndex, int32 MorphTargetIndex) const override;
 		virtual TConstArrayView<int32> GetMorphTargetIndices(int32 LODIndex, int32 MorphTargetIndex) const override;
diff --git a/Engine/Plugins/ChaosClothAssetEditor/Source/ChaosClothAssetDataflowNodes/Private/ChaosClothAsset/ChaosClothAssetDataflowNodesModule.cpp b/Engine/Plugins/ChaosClothAssetEditor/Source/ChaosClothAssetDataflowNodes/Private/ChaosClothAsset/ChaosClothAssetDataflowNodesModule.cpp
--- a/Engine/Plugins/ChaosClothAssetEditor/Source/ChaosClothAssetDataflowNodes/Private/ChaosClothAsset/ChaosClothAssetDataflowNodesModule.cpp
+++ b/Engine/Plugins/ChaosClothAssetEditor/Source/ChaosClothAssetDataflowNodes/Private/ChaosClothAsset/ChaosClothAssetDataflowNodesModule.cpp
@@ -13,9 +13,10 @@
 #include "ChaosClothAsset/ClothCollectionQueryNode.h"
 #include "ChaosClothAsset/ColorScheme.h"
 #include "ChaosClothAsset/CopySimulationToRenderMeshNode.h"
+#include "ChaosClothAsset/CustomRegionResizingNode.h"
 #include "ChaosClothAsset/DatasmithImportNode.h"
 #include "ChaosClothAsset/DeleteElementNode.h"
-#include "ChaosClothAsset/CustomRegionResizingNode.h"
+#include "ChaosClothAsset/GenerateSimMorphTargetNode.h"
 #include "ChaosClothAsset/ImportFilePathCustomization.h"
 #include "ChaosClothAsset/ImportNode.h"
 #include "ChaosClothAsset/ImportSimulationCacheNode.h"
@@ -112,6 +113,7 @@ namespace UE::Chaos::ClothAsset
 			DATAFLOW_NODE_REGISTER_CREATION_FACTORY(FChaosClothAssetDeleteElementNode);
 			DATAFLOW_NODE_REGISTER_CREATION_FACTORY(FChaosClothAssetExtractWeightMapNode);
 			DATAFLOW_NODE_REGISTER_CREATION_FACTORY(FChaosClothAssetExtractSelectionSetNode);
+			DATAFLOW_NODE_REGISTER_CREATION_FACTORY(FChaosClothAssetGenerateSimMorphTargetNode);
 			DATAFLOW_NODE_REGISTER_CREATION_FACTORY(FChaosClothAssetImportNode);
 			DATAFLOW_NODE_REGISTER_CREATION_FACTORY(FChaosClothAssetImportSimulationCacheNode);
 			DATAFLOW_NODE_REGISTER_CREATION_FACTORY(FChaosClothAssetMakeClothAssetNode);
diff --git a/Engine/Plugins/ChaosClothAssetEditor/Source/ChaosClothAssetDataflowNodes/Private/ChaosClothAsset/GenerateSimMorphTargetNode.cpp b/Engine/Plugins/ChaosClothAssetEditor/Source/ChaosClothAssetDataflowNodes/Private/ChaosClothAsset/GenerateSimMorphTargetNode.cpp
new file mode 100644
--- /dev/null
+++ b/Engine/Plugins/ChaosClothAssetEditor/Source/ChaosClothAssetDataflowNodes/Private/ChaosClothAsset/GenerateSimMorphTargetNode.cpp
@@ -0,0 +1,117 @@
+// Copyright Epic Games, Inc. All Rights Reserved.
+
+#include "ChaosClothAsset/GenerateSimMorphTargetNode.h"
+#include "ChaosClothAsset/CollectionClothFacade.h"
+#include "ChaosClothAsset/ClothDataflowTools.h"
+#include "Dataflow/DataflowInputOutput.h"
+
+#include UE_INLINE_GENERATED_CPP_BY_NAME(GenerateSimMorphTargetNode)
+
+#define LOCTEXT_NAMESPACE "ChaosClothAssetGenerateSimMorphTargetNode"
+
+FChaosClothAssetGenerateSimMorphTargetNode::FChaosClothAssetGenerateSimMorphTargetNode(const UE::Dataflow::FNodeParameters& InParam, FGuid InGuid)
+	: FDataflowNode(InParam, InGuid)
+{
+	RegisterInputConnection(&Collection);
+	RegisterInputConnection(&MorphTargetCollection);
+	RegisterOutputConnection(&Collection, &Collection);
+	RegisterOutputConnection(&MorphTargetName);
+}
+
+void FChaosClothAssetGenerateSimMorphTargetNode::Evaluate(UE::Dataflow::FContext& Context, const FDataflowOutput* Out) const
+{
+	if (Out->IsA<FManagedArrayCollection>(&Collection))
+	{
+		using namespace UE::Chaos::ClothAsset;
+
+		// Evaluate in collection
+		FManagedArrayCollection InCollection = GetValue<FManagedArrayCollection>(Context, &Collection);
+		const TSharedRef<FManagedArrayCollection> ClothCollection = MakeShared<FManagedArrayCollection>(MoveTemp(InCollection));
+
+		FManagedArrayCollection InMorphTargetCollection = GetValue<FManagedArrayCollection>(Context, &MorphTargetCollection);  // Can't use a const reference here sadly since the facade needs a SharedRef to be created
+		const TSharedRef<const FManagedArrayCollection> MorphTargetClothCollection = MakeShared<const FManagedArrayCollection>(MoveTemp(InMorphTargetCollection));
+
+		// Always check for a valid cloth collection/facade to avoid processing non cloth collections
+		FCollectionClothFacade ClothFacade(ClothCollection);
+		const FCollectionClothConstFacade MorphTargetClothFacade(MorphTargetClothCollection);
+
+		if (ClothFacade.IsValid())
+		{
+			if (MorphTargetClothFacade.IsValid())
+			{
+				if (!MorphTargetName.IsEmpty())
+				{
+					const int32 NumSimVertices = ClothFacade.GetNumSimVertices3D();
+					const int32 NumMorphSimVertices = MorphTargetClothFacade.GetNumSimVertices3D();
+					if (NumSimVertices == NumMorphSimVertices)
+					{
+						TConstArrayView<FVector3f> Positions = ClothFacade.GetSimPosition3D();
+						TConstArrayView<FVector3f> MorphPositions = MorphTargetClothFacade.GetSimPosition3D();
+						TConstArrayView<FVector3f> Normals = ClothFacade.GetSimNormal();
+						TConstArrayView<FVector3f> MorphNormals = MorphTargetClothFacade.GetSimNormal();
+
+						TArray<FVector3f> PositionDeltas;
+						TArray<FVector3f> NormalDeltas;
+						TArray<int32> Indices;
+						PositionDeltas.Reserve(NumSimVertices);
+						NormalDeltas.Reserve(NumSimVertices);
+						Indices.Reserve(NumSimVertices);
+
+						for (int32 Index = 0; Index < NumSimVertices; ++Index)
+						{
+							const FVector3f PositionDelta = MorphPositions[Index] - Positions[Index];
+							const FVector3f NormalDelta = bGenerateNormalDeltas ? MorphNormals[Index] - Normals[Index] : FVector3f::ZeroVector;
+							if (PositionDelta.SquaredLength() >= FMath::Square(UE_THRESH_POINTS_ARE_NEAR) ||
+								(bGenerateNormalDeltas && NormalDelta.SizeSquared() > FMath::Square(UE_THRESH_VECTORS_ARE_NEAR)))
+							{
+								PositionDeltas.Add(PositionDelta);
+								NormalDeltas.Add(NormalDelta);
+								Indices.Add(Index);
+							}
+						}
+
+						if (PositionDeltas.Num() > 0)
+						{
+							const int32 ExistingMorphTarget = ClothFacade.FindSimMorphTargetIndexByName(MorphTargetName);
+							if (ExistingMorphTarget != INDEX_NONE)
+							{
+								UE::Chaos::ClothAsset::FClothDataflowTools::LogAndToastWarning(*this,
+									LOCTEXT("DuplicateMorphTargetHeadline", "Duplicate Sim Morph Target"),
+									FText::Format(LOCTEXT("DuplicateMorphTargetDetails", "Existing Sim Morph target with name '{0}' will be replaced."),
+										FText::FromString(MorphTargetName)));
+							}
+
+							FCollectionClothSimMorphTargetFacade MorphTargetFacade = ExistingMorphTarget == INDEX_NONE ? ClothFacade.AddGetSimMorphTarget() : ClothFacade.GetSimMorphTarget(ExistingMorphTarget);
+							MorphTargetFacade.Initialize(MorphTargetName, PositionDeltas, NormalDeltas, Indices);
+						}
+						else
+						{
+							UE::Chaos::ClothAsset::FClothDataflowTools::LogAndToastWarning(*this,
+								LOCTEXT("IdenticalCollectionsHeadline", "Identical sim collections"),
+								LOCTEXT("IdenticalCollectionsDetails", "No morph targets generated because the sim collections are identical.")
+							);
+						}
+					}
+					else
+					{
+						UE::Chaos::ClothAsset::FClothDataflowTools::LogAndToastWarning(*this,
+							LOCTEXT("VertexCountMismatchHeadline", "Failed to generate morph target"),
+							FText::Format(LOCTEXT("VertexCountMismatchDetail", "Vertex count mismatch {0} != {1}."),
+								NumSimVertices, NumMorphSimVertices));
+					}
+				}
+			}
+			else
+			{
+				UE::Chaos::ClothAsset::FClothDataflowTools::LogAndToastWarning(*this,
+					LOCTEXT("InvalidMorphTargetCollectionHeadline", "Invalid Morph Target Collection"),
+					LOCTEXT("InvalidMorphTargetCollectionDetails", "Input Morph Target Collection is not a valid Cloth Collection.")
+				);
+			}
+		}
+
+		SetValue(Context, MoveTemp(*ClothCollection), &Collection);
+	}
+}
+
+#undef LOCTEXT_NAMESPACE
diff --git a/Engine/Plugins/ChaosClothAssetEditor/Source/ChaosClothAssetDataflowNodes/Public/ChaosClothAsset/GenerateSimMorphTargetNode.h b/Engine/Plugins/ChaosClothAssetEditor/Source/ChaosClothAssetDataflowNodes/Public/ChaosClothAsset/GenerateSimMorphTargetNode.h
new file mode 100644
--- /dev/null
+++ b/Engine/Plugins/ChaosClothAssetEditor/Source/ChaosClothAssetDataflowNodes/Public/ChaosClothAsset/GenerateSimMorphTargetNode.h
@@ -0,0 +1,39 @@
+// Copyright Epic Games, Inc. All Rights Reserved.
+
+#pragma once
+
+#include "Dataflow/DataflowNode.h"
+#include "GeometryCollection/ManagedArrayCollection.h"
+#include "GenerateSimMorphTargetNode.generated.h"
+
+/** Generate a Sim Morph target from a cloth collection sim mesh (with matching topology).*/
+USTRUCT(Meta = (DataflowCloth))
+struct FChaosClothAssetGenerateSimMorphTargetNode final : public FDataflowNode 
+{
+	GENERATED_USTRUCT_BODY()
+	DATAFLOW_NODE_DEFINE_INTERNAL(FChaosClothAssetGenerateSimMorphTargetNode, "GenerateSimMorphTarget", "Cloth", "Cloth Generate Sim Morph Target")
+	DATAFLOW_NODE_RENDER_TYPE("SurfaceRender", FName("FClothCollection"), "Collection")
+
+public:
+
+	/** Input/output collection */
+	UPROPERTY(Meta = (Dataflowinput, DataflowOutput, DataflowPassthrough = "Collection"))
+	FManagedArrayCollection Collection;
+
+	/** Collection to generate the morph target from. */
+	UPROPERTY(Meta = (Dataflowinput))
+	FManagedArrayCollection MorphTargetCollection;
+
+	/** Morph target name */
+	UPROPERTY(EditAnywhere, Category = "Sim Morph Target", meta = (DataflowOutput))
+	FString MorphTargetName;
+
+	/** Whether or not to generate normal deltas */
+	UPROPERTY(EditAnywhere, Category = "Sim Morph Target")
+	bool bGenerateNormalDeltas = false;
+
+	FChaosClothAssetGenerateSimMorphTargetNode(const UE::Dataflow::FNodeParameters& InParam, FGuid InGuid = FGuid::NewGuid());
+
+private:
+	virtual void Evaluate(UE::Dataflow::FContext& Context, const FDataflowOutput* Out) const override;
+};
diff --git a/Engine/Plugins/ChaosClothAssetEditor/Source/ChaosClothAssetEditor/Private/ChaosClothAsset/ClothEditorSimulationVisualization.cpp b/Engine/Plugins/ChaosClothAssetEditor/Source/ChaosClothAssetEditor/Private/ChaosClothAsset/ClothEditorSimulationVisualization.cpp
--- a/Engine/Plugins/ChaosClothAssetEditor/Source/ChaosClothAssetEditor/Private/ChaosClothAsset/ClothEditorSimulationVisualization.cpp
+++ b/Engine/Plugins/ChaosClothAssetEditor/Source/ChaosClothAssetEditor/Private/ChaosClothAsset/ClothEditorSimulationVisualization.cpp
@@ -425,6 +425,19 @@ const FVisualizationOption FVisualizationOption::OptionData[] =
 		{
 			EditorVisualization.ExtendViewportShowMenuWeightMapSelector(MenuBuilder);
 		})),
+	FVisualizationOption(
+		FClothVisualizationDebugDraw::CreateLambda([](const FClothEditorSimulationVisualization& EditorVisualization, const ::Chaos::FClothVisualizationNoGC& Visualization, FPrimitiveDrawInterface* PDI)
+		{
+			const FString* const MorphTargetName = EditorVisualization.GetCurrentlySelectedMorphTarget();
+			Visualization.DrawSimMorphTarget(PDI, MorphTargetName ? *MorphTargetName : FString());
+		}),
+		LOCTEXT("ChaosVisName_DrawSimMorphTarget", "Sim Morph Target"), 
+		LOCTEXT("ChaosVisName_DrawSimMorphTarget_ToolTip", "Draw a sim morph target. If none selected, the currently active morph target is displayed."), 
+		/*bDisablesSimulation =*/false, /*bHidesClothSections=*/false, /*bDefaultFlagValue=*/false,
+		FAdditionalMenus::CreateLambda([](FClothEditorSimulationVisualization& EditorVisualization, FMenuBuilder& MenuBuilder)
+		{
+			EditorVisualization.ExtendViewportShowMenuMorphTargetSelector(MenuBuilder);
+		})),
 	FVisualizationOption(
 		FClothVisualizationDebugDraw::CreateLambda([](const FClothEditorSimulationVisualization&, const ::Chaos::FClothVisualizationNoGC& Visualization, FPrimitiveDrawInterface* PDI)
 		{
@@ -537,22 +550,20 @@ void FClothEditorSimulationVisualization::ExtendViewportShowMenu(FMenuBuilder& M
 }
 
 
-void FClothEditorSimulationVisualization::ExtendViewportShowMenuWeightMapSelector(FMenuBuilder& MenuBuilder)
+void FClothEditorSimulationVisualization::ExtendViewportShowMenuNameSelector(FMenuBuilder& MenuBuilder, FNameSelectionData& SelectionData)
 {
-	WeightMapSelector =
+	SelectionData.Selector =
 		SNew(STextComboBox)
-		.OptionsSource(&WeightMapNames)
-		.OnSelectionChanged(STextComboBox::FOnTextSelectionChanged::CreateRaw(this, &FClothEditorSimulationVisualization::WeightMapSelectionChanged));
+		.OptionsSource(&SelectionData.Names)
+		.OnSelectionChanged(STextComboBox::FOnTextSelectionChanged::CreateLambda([&SelectionData](TSharedPtr<FString> Selection, ESelectInfo::Type SelectInfo)
+			{
+				SelectionData.CurrentlySelectedName = Selection;
+			}));
 
-	MenuBuilder.AddMenuEntry(FUIAction(), WeightMapSelector.ToSharedRef(),
+	MenuBuilder.AddMenuEntry(FUIAction(), SelectionData.Selector.ToSharedRef(),
 		NAME_None, LOCTEXT("WeightMapNameSelection", "Select weight map to draw."), EUserInterfaceActionType::None);
 }
 
-void FClothEditorSimulationVisualization::WeightMapSelectionChanged(TSharedPtr<FString> Selection, ESelectInfo::Type SelectInfo)
-{
-	CurrentlySelectedWeightMap = Selection;
-}
-
 void FClothEditorSimulationVisualization::ExtendViewportShowMenuSpinBox(FMenuBuilder& MenuBuilder, float& Value, const float MinValue, const float MaxValue, const float MinSliderValue, const float MaxSliderValue)
 {
 	TSharedRef<SWidget> SpinBox =
@@ -571,38 +582,46 @@ void FClothEditorSimulationVisualization::ExtendViewportShowMenuSpinBox(FMenuBui
 void FClothEditorSimulationVisualization::RefreshMenusForClothComponent(const UChaosClothComponent* ClothComponent)
 {
 	TArray<FString> WeightMaps;
+	TArray<FString> MorphTargets;
 	if (const ::Chaos::FClothVisualizationNoGC* const Visualization = ClothComponent && ClothComponent->GetClothSimulationProxy() ? ClothComponent->GetClothSimulationProxy()->GetClothVisualization() : nullptr)
 	{
 		WeightMaps = Visualization->GetAllWeightMapNames();
+		MorphTargets = Visualization->GetAllMorphTargetNames();
 	}
 
-	bool bFoundCurrentlySelected = false;
-	WeightMapNames.Empty(WeightMaps.Num());
-	for (const FString& Map : WeightMaps)
-	{
-		const int32 Index = WeightMapNames.Emplace(MakeShared<FString>(Map));
-		if (!bFoundCurrentlySelected && CurrentlySelectedWeightMap && Map == *CurrentlySelectedWeightMap)
-		{
-			CurrentlySelectedWeightMap = WeightMapNames[Index];
-			bFoundCurrentlySelected = true;
-		}
-	}
-	if (!bFoundCurrentlySelected)
-	{
-		if (WeightMapNames.IsEmpty())
-		{
-			CurrentlySelectedWeightMap.Reset();
-		}
-		else
+	auto RefreshMenuNames = [](const TArray<FString>& AllNames, FNameSelectionData& SelectionData)
 		{
-			CurrentlySelectedWeightMap = WeightMapNames[0];
-		}
-	}
-	if (WeightMapSelector)
-	{
-		WeightMapSelector->RefreshOptions();
-		WeightMapSelector->SetSelectedItem(CurrentlySelectedWeightMap);
-	}
+			bool bFoundCurrentlySelected = false;
+			SelectionData.Names.Empty(AllNames.Num());
+			for (const FString& Map : AllNames)
+			{
+				const int32 Index = SelectionData.Names.Emplace(MakeShared<FString>(Map));
+				if (!bFoundCurrentlySelected && SelectionData.CurrentlySelectedName && Map == *SelectionData.CurrentlySelectedName)
+				{
+					SelectionData.CurrentlySelectedName = SelectionData.Names[Index];
+					bFoundCurrentlySelected = true;
+				}
+			}
+			if (!bFoundCurrentlySelected)
+			{
+				if (SelectionData.Names.IsEmpty())
+				{
+					SelectionData.CurrentlySelectedName.Reset();
+				}
+				else
+				{
+					SelectionData.CurrentlySelectedName = SelectionData.Names[0];
+				}
+			}
+			if (SelectionData.Selector)
+			{
+				SelectionData.Selector->RefreshOptions();
+				SelectionData.Selector->SetSelectedItem(SelectionData.CurrentlySelectedName);
+			}
+		};
+
+	RefreshMenuNames(WeightMaps, WeightMapSelection);
+	RefreshMenuNames(MorphTargets, MorphTargetSelection);
 }
 
 void FClothEditorSimulationVisualization::DebugDrawSimulation(const UChaosClothComponent* ClothComponent, FPrimitiveDrawInterface* PDI)
diff --git a/Engine/Plugins/ChaosClothAssetEditor/Source/ChaosClothAssetEditor/Private/ChaosClothAsset/ClothEditorSimulationVisualization.h b/Engine/Plugins/ChaosClothAssetEditor/Source/ChaosClothAssetEditor/Private/ChaosClothAsset/ClothEditorSimulationVisualization.h
--- a/Engine/Plugins/ChaosClothAssetEditor/Source/ChaosClothAssetEditor/Private/ChaosClothAsset/ClothEditorSimulationVisualization.h
+++ b/Engine/Plugins/ChaosClothAssetEditor/Source/ChaosClothAssetEditor/Private/ChaosClothAsset/ClothEditorSimulationVisualization.h
@@ -36,9 +36,17 @@ public:
 	void RefreshMenusForClothComponent(const UChaosClothComponent* ClothComponent);
 
 	// WeightMaps 
-	const FString* GetCurrentlySelectedWeightMap() const { return CurrentlySelectedWeightMap.Get(); }
-	void ExtendViewportShowMenuWeightMapSelector(FMenuBuilder& MenuBuilder);
-
+	const FString* GetCurrentlySelectedWeightMap() const { return WeightMapSelection.CurrentlySelectedName.Get(); }
+	void ExtendViewportShowMenuWeightMapSelector(FMenuBuilder& MenuBuilder)
+	{
+		ExtendViewportShowMenuNameSelector(MenuBuilder, WeightMapSelection);
+	}
+	// Morph Targets
+	const FString* GetCurrentlySelectedMorphTarget() const { return MorphTargetSelection.CurrentlySelectedName.Get(); }
+	void ExtendViewportShowMenuMorphTargetSelector(FMenuBuilder& MenuBuilder)
+	{
+		ExtendViewportShowMenuNameSelector(MenuBuilder, MorphTargetSelection);
+	}
 	// Normals
 	void ExtendViewportShowMenuPointNormalsLength(FMenuBuilder& MenuBuilder)
 	{
@@ -57,6 +65,13 @@ public:
 	float GetAerodynamicsLengthScale() const { return AerodynamicsLengthScale; }
 
 private:
+	struct FNameSelectionData
+	{
+		TSharedPtr<STextComboBox> Selector;
+		TArray<TSharedPtr<FString>> Names;
+		TSharedPtr<FString> CurrentlySelectedName;
+	};
+
 
 	void ExtendViewportShowMenu(FMenuBuilder& MenuBuilder, const TFunction<UChaosClothComponent*()>& GetClothComponentFunc);
 
@@ -64,15 +79,14 @@ private:
 	bool ShouldDisableSimulation() const;
 	/** Show/hide all cloth sections for the specified mesh compoment. */
 	void ShowClothSections(UChaosClothComponent* ClothComponent, bool bIsClothSectionsVisible) const;
-	/** Callback for weight map selection*/
-	void WeightMapSelectionChanged(TSharedPtr<FString> Selection, ESelectInfo::Type SelectInfo);
 	void ExtendViewportShowMenuSpinBox(FMenuBuilder& MenuBuilder, float& Value, const float MinValue, const float MaxValue, const float MinSliderValue, const float MaxSliderValue);
+	void ExtendViewportShowMenuNameSelector(FMenuBuilder& MenuBuilder, FNameSelectionData& SelectionData);
+
 private:
 	/** Flags used to store the checked status for the visualization options. */
 	TBitArray<> Flags;
-	TSharedPtr<STextComboBox> WeightMapSelector;
-	TArray<TSharedPtr<FString>> WeightMapNames;
-	TSharedPtr<FString> CurrentlySelectedWeightMap;
+	FNameSelectionData WeightMapSelection;
+	FNameSelectionData MorphTargetSelection;
 	float PointNormalLength = 20.f;
 	float AnimatedNormalLength = 20.f;
 	float AerodynamicsLengthScale = 10.f;
