Update the weight map transfer to work with render meshes.

diff --git a/Engine/Plugins/ChaosClothAssetEditor/Source/ChaosClothAssetDataflowNodes/Private/ChaosClothAsset/WeightMapNode.cpp b/Engine/Plugins/ChaosClothAssetEditor/Source/ChaosClothAssetDataflowNodes/Private/ChaosClothAsset/WeightMapNode.cpp
--- a/Engine/Plugins/ChaosClothAssetEditor/Source/ChaosClothAssetDataflowNodes/Private/ChaosClothAsset/WeightMapNode.cpp
+++ b/Engine/Plugins/ChaosClothAssetEditor/Source/ChaosClothAssetDataflowNodes/Private/ChaosClothAsset/WeightMapNode.cpp
@@ -64,40 +64,66 @@ void FChaosClothAssetWeightMapNode::OnTransfer(UE::Dataflow::FContext& Context)
 		FCollectionClothConstFacade TransferClothFacade(TransferClothCollection);
 		const FName InInputName = GetInputName(Context);
 
-		if (TransferClothFacade.HasWeightMap(InInputName))
+		if (MeshTarget == EChaosClothAssetWeightMapMeshTarget::Simulation)
 		{
-			// Remap the weights
-			TArray<float> RemappedWeights;
-			RemappedWeights.SetNumZeroed(ClothFacade.GetNumSimVertices3D());
+			if (TransferClothFacade.HasWeightMap(InInputName))
+			{
+				// Remap the weights
+				TArray<float> RemappedWeights;
+				RemappedWeights.SetNumZeroed(ClothFacade.GetNumSimVertices3D());
+
+				switch (TransferType)
+				{
+				case EChaosClothAssetWeightMapTransferType::Use2DSimMesh:
+					Private::TransferWeightMap(
+						TransferClothFacade.GetSimPosition2D(),
+						TransferClothFacade.GetSimIndices2D(),
+						TransferClothFacade.GetSimVertex3DLookup(),
+						TransferClothFacade.GetWeightMap(InInputName),
+						ClothFacade.GetSimPosition2D(),
+						ClothFacade.GetSimIndices2D(),
+						ClothFacade.GetSimVertex3DLookup(),
+						RemappedWeights);
+					break;
+				case EChaosClothAssetWeightMapTransferType::Use3DSimMesh:
+					FClothGeometryTools::TransferWeightMap(
+						TransferClothFacade.GetSimPosition3D(),
+						TransferClothFacade.GetSimIndices3D(),
+						TransferClothFacade.GetWeightMap(InInputName),
+						ClothFacade.GetSimPosition3D(),
+						ClothFacade.GetSimNormal(),
+						ClothFacade.GetSimIndices3D(),
+						TArrayView<float>(RemappedWeights));
+					break;
+				default: unimplemented();
+				}
 
-			switch (TransferType)
+				SetVertexWeights(ClothFacade.GetWeightMap(InInputName), RemappedWeights);
+			}
+		}
+		else 
+		{
+			check(MeshTarget == EChaosClothAssetWeightMapMeshTarget::Render);
+
+			// Try to get a render weight map
+			TConstArrayView<float> TransferWeightMap = TransferClothFacade.GetUserDefinedAttribute<float>(InInputName, ClothCollectionGroup::RenderVertices);
+			if (TransferWeightMap.Num() == TransferClothFacade.GetNumRenderVertices())
 			{
-			case EChaosClothAssetWeightMapTransferType::Use2DSimMesh:
-				Private::TransferWeightMap(
-					TransferClothFacade.GetSimPosition2D(),
-					TransferClothFacade.GetSimIndices2D(),
-					TransferClothFacade.GetSimVertex3DLookup(),
-					TransferClothFacade.GetWeightMap(InInputName),
-					ClothFacade.GetSimPosition2D(),
-					ClothFacade.GetSimIndices2D(),
-					ClothFacade.GetSimVertex3DLookup(),
-					RemappedWeights);
-				break;
-			case EChaosClothAssetWeightMapTransferType::Use3DSimMesh:
+				// Remap the weights
+				TArray<float> RemappedWeights;
+				RemappedWeights.SetNumZeroed(ClothFacade.GetNumRenderVertices());
+
 				FClothGeometryTools::TransferWeightMap(
-					TransferClothFacade.GetSimPosition3D(),
-					TransferClothFacade.GetSimIndices3D(),
-					TransferClothFacade.GetWeightMap(InInputName),
-					ClothFacade.GetSimPosition3D(),
-					ClothFacade.GetSimNormal(),
-					ClothFacade.GetSimIndices3D(),
+					TransferClothFacade.GetRenderPosition(),
+					TransferClothFacade.GetRenderIndices(),
+					TransferWeightMap,
+					ClothFacade.GetRenderPosition(),
+					ClothFacade.GetRenderNormal(),
+					ClothFacade.GetRenderIndices(),
 					TArrayView<float>(RemappedWeights));
-				break;
-			default: unimplemented();
-			}
 
-			// TODO: Allow for transferring render weights from InTransferCollection?
-			SetVertexWeights(ClothFacade.GetWeightMap(InInputName), RemappedWeights);
+				SetVertexWeights(ClothFacade.GetUserDefinedAttribute<float>(InInputName, ClothCollectionGroup::RenderVertices), RemappedWeights);
+			}
 		}
 	}
 }
diff --git a/Engine/Plugins/ChaosClothAssetEditor/Source/ChaosClothAssetDataflowNodes/Public/ChaosClothAsset/WeightMapNode.h b/Engine/Plugins/ChaosClothAssetEditor/Source/ChaosClothAssetDataflowNodes/Public/ChaosClothAsset/WeightMapNode.h
--- a/Engine/Plugins/ChaosClothAssetEditor/Source/ChaosClothAssetDataflowNodes/Public/ChaosClothAsset/WeightMapNode.h
+++ b/Engine/Plugins/ChaosClothAssetEditor/Source/ChaosClothAssetDataflowNodes/Public/ChaosClothAsset/WeightMapNode.h
@@ -59,8 +59,9 @@ public:
 	EChaosClothAssetWeightMapOverrideType MapOverrideType = EChaosClothAssetWeightMapOverrideType::ReplaceChanged;
 
 	/**
-	 * The type of transfer used to transfer the weight map when a TransferCollection is connected.
-	 * This property is disabled when no TransferCollection input has been conencted.
+	 * The type of transfer used to transfer the weight map when a TransferCollection is connected and MeshTarget is Simulation.
+	 * Render weight maps always do a 3D transfer.
+	 * This property is disabled when no TransferCollection input has been connected.
 	 */
 	UPROPERTY(EditAnywhere, Category = "Weight Map")
 	EChaosClothAssetWeightMapTransferType TransferType = EChaosClothAssetWeightMapTransferType::Use2DSimMesh;
