Cloth Remesh Node: transfer vertex weight layers to output cloth render mesh

diff --git a/Engine/Plugins/ChaosClothAssetEditor/Source/ChaosClothAssetDataflowNodes/Private/ChaosClothAsset/RemeshNode.cpp b/Engine/Plugins/ChaosClothAssetEditor/Source/ChaosClothAssetDataflowNodes/Private/ChaosClothAsset/RemeshNode.cpp
--- a/Engine/Plugins/ChaosClothAssetEditor/Source/ChaosClothAssetDataflowNodes/Private/ChaosClothAsset/RemeshNode.cpp
+++ b/Engine/Plugins/ChaosClothAssetEditor/Source/ChaosClothAssetDataflowNodes/Private/ChaosClothAsset/RemeshNode.cpp
@@ -1643,6 +1643,28 @@ namespace UE::Chaos::ClothAsset::Private
 			}
 		}
 
+		//
+		// Vertex Weights
+		//
+
+		TArray<TArray<float>> VertexWeights;
+		VertexWeights.SetNum(Attributes->NumWeightLayers());
+
+		TArray<FName> VertexWeightLayerNames;
+		VertexWeightLayerNames.SetNum(Attributes->NumWeightLayers());
+
+		for (int32 WeightLayerIndex = 0; WeightLayerIndex < Attributes->NumWeightLayers(); ++WeightLayerIndex)
+		{
+			const FDynamicMeshWeightAttribute* const WeightLayerAttribute = Attributes->GetWeightLayer(WeightLayerIndex);
+			VertexWeightLayerNames[WeightLayerIndex] = WeightLayerAttribute->GetName();
+
+			VertexWeights[WeightLayerIndex].SetNum(DynamicMesh.MaxVertexID());
+			for (const int VertexIndex : DynamicMesh.VertexIndicesItr())
+			{
+				WeightLayerAttribute->GetValue(VertexIndex, &(VertexWeights[WeightLayerIndex][VertexIndex]));
+			}
+		}
+
 
 		// Find the set of triangles per MaterialID
 
@@ -1674,6 +1696,8 @@ namespace UE::Chaos::ClothAsset::Private
 		FCollectionClothFacade OutClothFacade(OutClothCollection);
 		OutClothFacade.SetNumRenderPatterns(NumMaterials);
 
+		TMap<int, FIntVector2> MeshVertexToPatternAndVertex;
+
 		for (int32 DestPatternID = 0; DestPatternID < MaterialIDs.Num(); ++DestPatternID)
 		{
 			FCollectionClothRenderPatternFacade OutClothPatternFacade = OutClothFacade.GetRenderPattern(DestPatternID);
@@ -1727,6 +1751,11 @@ namespace UE::Chaos::ClothAsset::Private
 				RenderColor[PatternVertexIndex] = FLinearColor::White;
 				RenderBoneIndices[PatternVertexIndex] = BoneIndices[SourceVertexIndex];
 				RenderBoneWeights[PatternVertexIndex] = BoneWeights[SourceVertexIndex];
+
+				if (!MeshVertexToPatternAndVertex.Contains(SourceVertexIndex))
+				{
+					MeshVertexToPatternAndVertex.Add({ SourceVertexIndex, FIntVector2{DestPatternID, PatternVertexIndex} });
+				}
 			}
 
 			OutClothPatternFacade.SetNumRenderFaces(TriangleIDs.Num());
@@ -1747,6 +1776,37 @@ namespace UE::Chaos::ClothAsset::Private
 			FCollectionClothRenderPatternConstFacade InPatternFacade = InClothFacade.GetRenderPattern(SourceMaterialID);
 			OutClothPatternFacade.SetRenderMaterialPathName(InPatternFacade.GetRenderMaterialPathName());
 		}
+
+
+		if (VertexWeightLayerNames.Num() > 0)
+		{
+			TArray<int32> DynamicMeshToClothVertexMap;
+			DynamicMeshToClothVertexMap.SetNum(DynamicMesh.MaxVertexID());
+			for (int InGlobalVertexIndex = 0; InGlobalVertexIndex < DynamicMesh.MaxVertexID(); ++InGlobalVertexIndex)
+			{
+				const int32 PatternID = MeshVertexToPatternAndVertex[InGlobalVertexIndex][0];
+				const int32 VertexID = MeshVertexToPatternAndVertex[InGlobalVertexIndex][1];
+				const int32 ClothGlobalIndex = OutClothFacade.GetRenderPattern(PatternID).GetRenderVerticesOffset() + VertexID;
+				DynamicMeshToClothVertexMap[InGlobalVertexIndex] = ClothGlobalIndex;
+			}
+
+			for (int32 WeightLayerID = 0; WeightLayerID < VertexWeightLayerNames.Num(); ++WeightLayerID)
+			{
+				const bool bAdded = OutClothFacade.AddUserDefinedAttribute<float>(VertexWeightLayerNames[WeightLayerID], ClothCollectionGroup::RenderVertices);
+				ensure(bAdded);
+				TArrayView<float> ClothRenderWeights = OutClothFacade.GetUserDefinedAttribute<float>(VertexWeightLayerNames[WeightLayerID], ClothCollectionGroup::RenderVertices);
+				ensure(ClothRenderWeights.Num() == VertexWeights[WeightLayerID].Num());
+
+				for (int32 MeshVertexIndex = 0; MeshVertexIndex < VertexWeights[WeightLayerID].Num(); ++MeshVertexIndex)
+				{
+					const float VertexWeight = VertexWeights[WeightLayerID][MeshVertexIndex];
+					const int32 ClothVertexIndex = DynamicMeshToClothVertexMap[MeshVertexIndex];
+					ClothRenderWeights[ClothVertexIndex] = VertexWeight;
+				}
+			}
+		}
+
+
 	}
 
 
@@ -1956,3 +2016,4 @@ void FChaosClothAssetRemeshNode::Evaluate(UE::Dataflow::FContext& Context, const
 }
 
 #undef LOCTEXT_NAMESPACE
+
