Move Cloth Resizing methods from ClothGeometryTools to ClothDataflowTools. Doing this to move the MeshResizing plugin dependency from the ChaosClothAsset plugin to the ChaosClothAssetEditor plugin

diff --git a/Engine/Plugins/ChaosClothAsset/ChaosClothAsset.uplugin b/Engine/Plugins/ChaosClothAsset/ChaosClothAsset.uplugin
--- a/Engine/Plugins/ChaosClothAsset/ChaosClothAsset.uplugin
+++ b/Engine/Plugins/ChaosClothAsset/ChaosClothAsset.uplugin
@@ -38,10 +38,6 @@
 		{
 			"Name": "Dataflow",
 			"Enabled": true
-		},
-		{
-			"Name": "MeshResizing",
-			"Enabled": true
 		}
 	]
 }
diff --git a/Engine/Plugins/ChaosClothAsset/Source/ChaosClothAsset/ChaosClothAsset.Build.cs b/Engine/Plugins/ChaosClothAsset/Source/ChaosClothAsset/ChaosClothAsset.Build.cs
--- a/Engine/Plugins/ChaosClothAsset/Source/ChaosClothAsset/ChaosClothAsset.Build.cs
+++ b/Engine/Plugins/ChaosClothAsset/Source/ChaosClothAsset/ChaosClothAsset.Build.cs
@@ -14,7 +14,6 @@ public class ChaosClothAsset : ModuleRules
 				"Core",
 				"GeometryCore",
 				"MeshConversion",
-				"MeshResizingCore",
 				"ClothingSystemRuntimeCommon",
 				"RenderCore"
 			}
diff --git a/Engine/Plugins/ChaosClothAsset/Source/ChaosClothAsset/Private/ChaosClothAsset/ClothGeometryTools.cpp b/Engine/Plugins/ChaosClothAsset/Source/ChaosClothAsset/Private/ChaosClothAsset/ClothGeometryTools.cpp
--- a/Engine/Plugins/ChaosClothAsset/Source/ChaosClothAsset/Private/ChaosClothAsset/ClothGeometryTools.cpp
+++ b/Engine/Plugins/ChaosClothAsset/Source/ChaosClothAsset/Private/ChaosClothAsset/ClothGeometryTools.cpp
@@ -12,8 +12,6 @@
 #include "DynamicMesh/NonManifoldMappingSupport.h"
 #include "Math/OrientedBox.h"
 #include "Math/Vector.h"
-#include "MeshResizing/CustomRegionResizing.h"
-#include "MeshResizing/RBFInterpolation.h"
 #include "Util/IndexUtil.h"
 #include "Utils/ClothingMeshUtils.h"
 #include "Algo/RemoveIf.h"
@@ -470,69 +468,6 @@ namespace UE::Chaos::ClothAsset
 		}
 	}// namespace Private::SimMeshBuilder
 
-	namespace Private::Resizing
-	{
-		static void ApplyGroupResizing(FCollectionClothFacade& ClothFacade, const FMeshDescription& TargetMeshDescription, const FMeshResizingRBFInterpolationData& InterpolationData, const TArray<FMeshResizingCustomRegion>& ResizingGroupData, TArrayView<FVector3f> Positions, TConstArrayView<float> ResizingBlend)
-		{
-			if (ClothFacade.IsValid(EClothCollectionExtendedSchemas::Resizing))
-			{
-				TConstArrayView<int32> ClothSetTypes = ClothFacade.GetCustomResizingRegionType();
-				check(ClothSetTypes.Num() == ResizingGroupData.Num());
-
-				// Gather trilinear interpolation data.
-				TArray<int32> TrilinearInterpolationGroups;
-				TArray<FVector3d> BoundCorners;
-				int32 BoundCornersIndex = 0;
-
-				for (int32 GroupIndex = 0; GroupIndex < ClothSetTypes.Num(); ++GroupIndex)
-				{
-					if ((EMeshResizingCustomRegionType)ClothSetTypes[GroupIndex] == EMeshResizingCustomRegionType::TrilinearInterpolation)
-					{
-						const FMeshResizingCustomRegion& ResizingData = ResizingGroupData[GroupIndex];
-
-						if (ResizingData.IsValid())
-						{
-							TrilinearInterpolationGroups.Add(GroupIndex);
-
-							const FMatrix TriangleMatrix(FVector3d(ResizingData.SourceAxis0), FVector3d(ResizingData.SourceAxis1), FVector3d(ResizingData.SourceAxis2), ResizingData.SourceOrigin);
-							FOrientedBox BoundBox;
-							BoundBox.Center = TriangleMatrix.TransformPosition(FVector3d(ResizingData.RegionBoundsCentroid));
-							BoundBox.AxisX = FVector3d(ResizingData.SourceAxis0);
-							BoundBox.AxisY = FVector3d(ResizingData.SourceAxis1);
-							BoundBox.AxisZ = FVector3d(ResizingData.SourceAxis2);
-							BoundBox.ExtentX = ResizingData.RegionBoundsExtents.X;
-							BoundBox.ExtentY = ResizingData.RegionBoundsExtents.Y;
-							BoundBox.ExtentZ = ResizingData.RegionBoundsExtents.Z;
-							BoundCorners.SetNum((BoundCornersIndex + 1) * 8);
-							BoundBox.CalcVertices(&BoundCorners[BoundCornersIndex]);
-							BoundCornersIndex += 8;
-						}
-					}
-				}
-				if (!TrilinearInterpolationGroups.IsEmpty())
-				{
-					UE::MeshResizing::FRBFInterpolation::DeformPoints(TargetMeshDescription, InterpolationData, BoundCorners);
-					BoundCornersIndex = 0;
-					check(BoundCorners.Num() == TrilinearInterpolationGroups.Num() * 8);
-
-					const TArray<FVector3f> OrigPositions(Positions.GetData(), Positions.Num());
-					for (const int32 GroupIndex : TrilinearInterpolationGroups)
-					{
-						const FMeshResizingCustomRegion& ResizingData = ResizingGroupData[GroupIndex];
-
-						UE::MeshResizing::FCustomRegionResizing::InterpolateCustomRegionPoints(ResizingData, TConstArrayView<FVector3d>(BoundCorners.GetData() + 8 * BoundCornersIndex, 8), Positions);
-						BoundCornersIndex += 8;
-					}
-
-					for (int32 VertexIndex = 0; VertexIndex < ResizingBlend.Num(); ++VertexIndex)
-					{
-						Positions[VertexIndex] = FMath::Lerp(OrigPositions[VertexIndex], Positions[VertexIndex], ResizingBlend[VertexIndex]);
-					}
-				}
-			}
-		}
-	} // namespace Private::Resizing
-
 	// TODO: Move these functions to the cloth facade?
 	bool FClothGeometryTools::HasSimMesh(const TSharedRef<const FManagedArrayCollection>& ClothCollection)
 	{
@@ -1945,129 +1880,4 @@ PRAGMA_ENABLE_DEPRECATION_WARNINGS
 			}
 		}
 	}
-
-	void FClothGeometryTools::SetGroupResizingData(const TSharedRef<FManagedArrayCollection>& ClothCollection, const TConstArrayView<FName>& SetNames, const TConstArrayView<int32>& SetTypes)
-	{
-		check(SetNames.Num() == SetTypes.Num());
-		FCollectionClothFacade ClothFacade(ClothCollection);
-
-		ClothFacade.DefineSchema(EClothCollectionExtendedSchemas::Resizing);
-
-		// Clear the sim and render blends
-		TArrayView<float> SimBlend = ClothFacade.GetSimCustomResizingBlend();
-		TArrayView<float> RenderBlend = ClothFacade.GetRenderCustomResizingBlend();
-		for (int32 Index = 0; Index < SimBlend.Num(); ++Index)
-		{
-			SimBlend[Index] = 0.f;
-		}
-		for (int32 Index = 0; Index < RenderBlend.Num(); ++Index)
-		{
-			RenderBlend[Index] = 0.f;
-		}
-
-		// Filter out invalid sets
-		TArray<FName> FilteredSetNames;
-		TArray<int32> FilteredSetTypes;
-		FilteredSetNames.Reserve(SetNames.Num());
-		FilteredSetTypes.Reserve(SetNames.Num());
-		TSet<int32> SetValues;
-		for (int32 Index = 0; Index < SetNames.Num(); ++Index)
-		{
-			// Try to get the set as a SimVertex3D set
-			if (ConvertSelectionToNewGroupType(ClothCollection, SetNames[Index], ClothCollectionGroup::SimVertices3D, SetValues))
-			{
-				FilteredSetNames.Add(SetNames[Index]);
-				FilteredSetTypes.Add(SetTypes[Index]);
-
-				for (const int32 Vertex : SetValues)
-				{
-					if (SimBlend.IsValidIndex(Vertex))
-					{
-						SimBlend[Vertex] = 1.f;
-					}
-				}
-			}
-			else if (ConvertSelectionToNewGroupType(ClothCollection, SetNames[Index], ClothCollectionGroup::RenderVertices, SetValues))
-			{
-				FilteredSetNames.Add(SetNames[Index]);
-				FilteredSetTypes.Add(SetTypes[Index]);
-
-				for (const int32 Vertex : SetValues)
-				{
-					if (RenderBlend.IsValidIndex(Vertex))
-					{
-						RenderBlend[Vertex] = 1.f;
-					}
-				}
-			}
-		}
-
-		check(FilteredSetNames.Num() == FilteredSetTypes.Num());
-		ClothFacade.SetNumCustomResizingRegions(FilteredSetNames.Num());
-		TArrayView<FString> ClothSetNames = ClothFacade.GetCustomResizingRegionSet();
-		TArrayView<int32> ClothSetTypes = ClothFacade.GetCustomResizingRegionType();
-		for (int32 Index = 0; Index < FilteredSetNames.Num(); ++Index)
-		{
-			ClothSetNames[Index] = FilteredSetNames[Index].ToString();
-			ClothSetTypes[Index] = FilteredSetTypes[Index];
-		}
-	}
-
-	void FClothGeometryTools::GenerateSimMeshResizingGroupData(const TSharedRef<FManagedArrayCollection>& ClothCollection, const FMeshDescription& SourceMeshDescription, TArray<FMeshResizingCustomRegion>& OutResizingGroupData)
-	{
-		FCollectionClothFacade ClothFacade(ClothCollection);
-		if (ClothFacade.IsValid(EClothCollectionExtendedSchemas::Resizing))
-		{
-			UE::Geometry::FDynamicMesh3 SourceMesh;
-			FMeshDescriptionToDynamicMesh Converter;
-			Converter.Convert(&SourceMeshDescription, SourceMesh, true);
-
-			TConstArrayView<FString> ClothSetNames = ClothFacade.GetCustomResizingRegionSet();
-			TConstArrayView<int32> ClothSetTypes = ClothFacade.GetCustomResizingRegionType();
-			OutResizingGroupData.SetNum(ClothSetNames.Num());
-			for (int32 GroupIndex = 0; GroupIndex < ClothSetNames.Num(); ++GroupIndex)
-			{
-				TSet<int32> SetValues;
-				if (ConvertSelectionToNewGroupType(ClothCollection, FName(ClothSetNames[GroupIndex]), ClothCollectionGroup::SimVertices3D, SetValues))
-				{
-					UE::MeshResizing::FCustomRegionResizing::GenerateCustomRegion(ClothFacade.GetSimPosition3D(), SourceMesh, SetValues, OutResizingGroupData[GroupIndex]);
-				}
-			}
-		}
-	}
-
-	void FClothGeometryTools::GenerateRenderMeshResizingGroupData(const TSharedRef<FManagedArrayCollection>& ClothCollection, const FMeshDescription& SourceMeshDescription, TArray<FMeshResizingCustomRegion>& OutResizingGroupData)
-	{
-		FCollectionClothFacade ClothFacade(ClothCollection);
-		if (ClothFacade.IsValid(EClothCollectionExtendedSchemas::Resizing))
-		{
-			UE::Geometry::FDynamicMesh3 SourceMesh;
-			FMeshDescriptionToDynamicMesh Converter;
-			Converter.Convert(&SourceMeshDescription, SourceMesh, true);
-
-			TConstArrayView<FString> ClothSetNames = ClothFacade.GetCustomResizingRegionSet();
-			TConstArrayView<int32> ClothSetTypes = ClothFacade.GetCustomResizingRegionType();
-			OutResizingGroupData.SetNum(ClothSetNames.Num());
-			for (int32 GroupIndex = 0; GroupIndex < ClothSetNames.Num(); ++GroupIndex)
-			{
-				TSet<int32> SetValues;
-				if (ConvertSelectionToNewGroupType(ClothCollection, FName(ClothSetNames[GroupIndex]), ClothCollectionGroup::RenderVertices, SetValues))
-				{
-					UE::MeshResizing::FCustomRegionResizing::GenerateCustomRegion(ClothFacade.GetRenderPosition(), SourceMesh, SetValues, OutResizingGroupData[GroupIndex]);
-				}
-			}
-		}
-	}
-
-	void FClothGeometryTools::ApplySimGroupResizing(const TSharedRef<FManagedArrayCollection>& ClothCollection, const FMeshDescription& TargetMeshDescription, const FMeshResizingRBFInterpolationData& InterpolationData, const TArray<FMeshResizingCustomRegion>& ResizingGroupData)
-	{
-		FCollectionClothFacade ClothFacade(ClothCollection);
-		Private::Resizing::ApplyGroupResizing(ClothFacade, TargetMeshDescription, InterpolationData, ResizingGroupData, ClothFacade.GetSimPosition3D(), ClothFacade.GetSimCustomResizingBlend());
-	}
-
-	void FClothGeometryTools::ApplyRenderGroupResizing(const TSharedRef<FManagedArrayCollection>& ClothCollection, const FMeshDescription& TargetMeshDescription, const FMeshResizingRBFInterpolationData& InterpolationData, const TArray<FMeshResizingCustomRegion>& ResizingGroupData)
-	{
-		FCollectionClothFacade ClothFacade(ClothCollection);
-		Private::Resizing::ApplyGroupResizing(ClothFacade, TargetMeshDescription, InterpolationData, ResizingGroupData, ClothFacade.GetRenderPosition(), ClothFacade.GetRenderCustomResizingBlend());
-	}
 }  // End namespace UE::Chaos::ClothAsset
diff --git a/Engine/Plugins/ChaosClothAsset/Source/ChaosClothAsset/Public/ChaosClothAsset/ClothGeometryTools.h b/Engine/Plugins/ChaosClothAsset/Source/ChaosClothAsset/Public/ChaosClothAsset/ClothGeometryTools.h
--- a/Engine/Plugins/ChaosClothAsset/Source/ChaosClothAsset/Public/ChaosClothAsset/ClothGeometryTools.h
+++ b/Engine/Plugins/ChaosClothAsset/Source/ChaosClothAsset/Public/ChaosClothAsset/ClothGeometryTools.h
@@ -9,8 +9,6 @@
 struct FManagedArrayCollection;
 namespace UE::Geometry { class FDynamicMesh3; }
 struct FMeshDescription;
-struct FMeshResizingRBFInterpolationData;
-struct FMeshResizingCustomRegion;
 
 namespace UE::Chaos::ClothAsset
 {
@@ -181,31 +179,6 @@ namespace UE::Chaos::ClothAsset
 		 */
 		static UE_API void ApplyProxyDeformer(const TSharedRef<FManagedArrayCollection>& ClothCollection, bool bIgnoreSkinningBlend);
 		static UE_API void ApplyProxyDeformer(FCollectionClothFacade& ClothFacade, bool bIgnoreSkinningBlend, const TConstArrayView<FVector3f>& SimPositions, const TConstArrayView<FVector3f>& SimNormals);
-
-		/**
-		 * Set Group Resizing data from supplied sets and types.
-		 */
-		static UE_API void SetGroupResizingData(const TSharedRef<FManagedArrayCollection>& ClothCollection, const TConstArrayView<FName>& SetNames, const TConstArrayView<int32>& SetTypes);
-
-		/**
-		 * Generate FMeshResizingCustomRegion data from Sim Mesh Group Resizing data
-		 */
-		static UE_API void GenerateSimMeshResizingGroupData(const TSharedRef<FManagedArrayCollection>& ClothCollection, const FMeshDescription& SourceMesh, TArray<FMeshResizingCustomRegion>& OutResizingGroupData);
-
-		/**
-		 * Generate FMeshResizingCustomRegion data from Render Mesh Group Resizing data
-		 */
-		static UE_API void GenerateRenderMeshResizingGroupData(const TSharedRef<FManagedArrayCollection>& ClothCollection, const FMeshDescription& SourceMesh, TArray<FMeshResizingCustomRegion>& OutResizingGroupData);
-
-		/**
-		 * Apply Group Resizing to the sim mesh.
-		 */
-		static UE_API void ApplySimGroupResizing(const TSharedRef<FManagedArrayCollection>& ClothCollection, const FMeshDescription& TargetMesh, const FMeshResizingRBFInterpolationData& InterpolationData, const TArray<FMeshResizingCustomRegion>& ResizingGroupData);
-
-		/**
-		 * Apply Group Resizing to the render mesh.
-		 */
-		static UE_API void ApplyRenderGroupResizing(const TSharedRef<FManagedArrayCollection>& ClothCollection, const FMeshDescription& TargetMesh, const FMeshResizingRBFInterpolationData& InterpolationData, const TArray<FMeshResizingCustomRegion>& ResizingGroupData);
 	};
 }  // End namespace UE::Chaos::ClothAsset
 
diff --git a/Engine/Plugins/ChaosClothAssetEditor/Source/ChaosClothAssetDataflowNodes/Private/ChaosClothAsset/ApplyResizingNode.cpp b/Engine/Plugins/ChaosClothAssetEditor/Source/ChaosClothAssetDataflowNodes/Private/ChaosClothAsset/ApplyResizingNode.cpp
--- a/Engine/Plugins/ChaosClothAssetEditor/Source/ChaosClothAssetDataflowNodes/Private/ChaosClothAsset/ApplyResizingNode.cpp
+++ b/Engine/Plugins/ChaosClothAssetEditor/Source/ChaosClothAssetDataflowNodes/Private/ChaosClothAsset/ApplyResizingNode.cpp
@@ -2,7 +2,7 @@
 
 #include "ChaosClothAsset/ApplyResizingNode.h"
 #include "Chaos/CollectionPropertyFacade.h"
-#include "ChaosClothAsset/ClothGeometryTools.h"
+#include "ChaosClothAsset/ClothDataflowTools.h"
 #include "ChaosClothAsset/CollectionClothFacade.h"
 #include "Dataflow/DataflowInputOutput.h"
 #include "Engine/SkeletalMesh.h"
@@ -78,7 +78,7 @@ void FChaosClothAssetApplyResizingNode::EvaluateClothCollection(UE::Dataflow::FC
 								{
 									if (const FMeshDescription* const SourceMeshDescription = InTargetSKM->GetMeshDescription(InLODIndex))
 									{
-										FClothGeometryTools::GenerateSimMeshResizingGroupData(ClothCollection, *SourceMeshDescription, ResizingGroupData);
+										FClothDataflowTools::GenerateSimMeshResizingGroupData(ClothCollection, *SourceMeshDescription, ResizingGroupData);
 									}
 								}
 							}
@@ -87,7 +87,7 @@ void FChaosClothAssetApplyResizingNode::EvaluateClothCollection(UE::Dataflow::FC
 
 							if (!ResizingGroupData.IsEmpty())
 							{
-								FClothGeometryTools::ApplySimGroupResizing(ClothCollection, *TargetMeshDescription, InInterpolationData, ResizingGroupData);
+								FClothDataflowTools::ApplySimGroupResizing(ClothCollection, *TargetMeshDescription, InInterpolationData, ResizingGroupData);
 							}
 						}
 						else if (ClothFacade.HasValidRenderData())
@@ -99,7 +99,7 @@ void FChaosClothAssetApplyResizingNode::EvaluateClothCollection(UE::Dataflow::FC
 								{
 									if (const FMeshDescription* const SourceMeshDescription = InTargetSKM->GetMeshDescription(InLODIndex))
 									{
-										FClothGeometryTools::GenerateRenderMeshResizingGroupData(ClothCollection, *SourceMeshDescription, ResizingGroupData);
+										FClothDataflowTools::GenerateRenderMeshResizingGroupData(ClothCollection, *SourceMeshDescription, ResizingGroupData);
 									}
 								}
 							}
@@ -108,7 +108,7 @@ void FChaosClothAssetApplyResizingNode::EvaluateClothCollection(UE::Dataflow::FC
 
 							if (!ResizingGroupData.IsEmpty())
 							{
-								FClothGeometryTools::ApplyRenderGroupResizing(ClothCollection, *TargetMeshDescription, InInterpolationData, ResizingGroupData);
+								FClothDataflowTools::ApplyRenderGroupResizing(ClothCollection, *TargetMeshDescription, InInterpolationData, ResizingGroupData);
 							}
 						}
 					}
diff --git a/Engine/Plugins/ChaosClothAssetEditor/Source/ChaosClothAssetDataflowNodes/Private/ChaosClothAsset/ClothDataflowTools.cpp b/Engine/Plugins/ChaosClothAssetEditor/Source/ChaosClothAssetDataflowNodes/Private/ChaosClothAsset/ClothDataflowTools.cpp
--- a/Engine/Plugins/ChaosClothAssetEditor/Source/ChaosClothAssetDataflowNodes/Private/ChaosClothAsset/ClothDataflowTools.cpp
+++ b/Engine/Plugins/ChaosClothAssetEditor/Source/ChaosClothAssetDataflowNodes/Private/ChaosClothAsset/ClothDataflowTools.cpp
@@ -3,6 +3,7 @@
 #include "ChaosClothAsset/ClothDataflowTools.h"
 #include "ChaosClothAsset/ClothGeometryTools.h"
 #include "ChaosClothAsset/CollectionClothFacade.h"
+#include "ChaosClothAsset/ClothCollectionGroup.h"
 #include "Dataflow/DataflowNode.h"
 #include "DynamicMesh/NonManifoldMappingSupport.h"
 #include "Framework/Notifications/NotificationManager.h"
@@ -12,6 +13,8 @@
 #include "Widgets/Notifications/SNotificationList.h"
 #include "MeshDescriptionToDynamicMesh.h"
 #include "MeshUtilities.h"
+#include "MeshResizing/CustomRegionResizing.h"
+#include "MeshResizing/RBFInterpolation.h"
 #include "PropertyHandle.h"
 #include "ReferenceSkeleton.h"
 #include "SkeletalMeshAttributes.h"
@@ -371,6 +374,69 @@ namespace UE::Chaos::ClothAsset
 			TArray<int32> EmptyArray;
 
 		};
+
+		namespace Resizing
+		{
+			static void ApplyGroupResizing(FCollectionClothFacade& ClothFacade, const FMeshDescription& TargetMeshDescription, const FMeshResizingRBFInterpolationData& InterpolationData, const TArray<FMeshResizingCustomRegion>& ResizingGroupData, TArrayView<FVector3f> Positions, TConstArrayView<float> ResizingBlend)
+			{
+				if (ClothFacade.IsValid(EClothCollectionExtendedSchemas::Resizing))
+				{
+					TConstArrayView<int32> ClothSetTypes = ClothFacade.GetCustomResizingRegionType();
+					check(ClothSetTypes.Num() == ResizingGroupData.Num());
+
+					// Gather trilinear interpolation data.
+					TArray<int32> TrilinearInterpolationGroups;
+					TArray<FVector3d> BoundCorners;
+					int32 BoundCornersIndex = 0;
+
+					for (int32 GroupIndex = 0; GroupIndex < ClothSetTypes.Num(); ++GroupIndex)
+					{
+						if ((EMeshResizingCustomRegionType)ClothSetTypes[GroupIndex] == EMeshResizingCustomRegionType::TrilinearInterpolation)
+						{
+							const FMeshResizingCustomRegion& ResizingData = ResizingGroupData[GroupIndex];
+
+							if (ResizingData.IsValid())
+							{
+								TrilinearInterpolationGroups.Add(GroupIndex);
+
+								const FMatrix TriangleMatrix(FVector3d(ResizingData.SourceAxis0), FVector3d(ResizingData.SourceAxis1), FVector3d(ResizingData.SourceAxis2), ResizingData.SourceOrigin);
+								FOrientedBox BoundBox;
+								BoundBox.Center = TriangleMatrix.TransformPosition(FVector3d(ResizingData.RegionBoundsCentroid));
+								BoundBox.AxisX = FVector3d(ResizingData.SourceAxis0);
+								BoundBox.AxisY = FVector3d(ResizingData.SourceAxis1);
+								BoundBox.AxisZ = FVector3d(ResizingData.SourceAxis2);
+								BoundBox.ExtentX = ResizingData.RegionBoundsExtents.X;
+								BoundBox.ExtentY = ResizingData.RegionBoundsExtents.Y;
+								BoundBox.ExtentZ = ResizingData.RegionBoundsExtents.Z;
+								BoundCorners.SetNum((BoundCornersIndex + 1) * 8);
+								BoundBox.CalcVertices(&BoundCorners[BoundCornersIndex]);
+								BoundCornersIndex += 8;
+							}
+						}
+					}
+					if (!TrilinearInterpolationGroups.IsEmpty())
+					{
+						UE::MeshResizing::FRBFInterpolation::DeformPoints(TargetMeshDescription, InterpolationData, BoundCorners);
+						BoundCornersIndex = 0;
+						check(BoundCorners.Num() == TrilinearInterpolationGroups.Num() * 8);
+
+						const TArray<FVector3f> OrigPositions(Positions.GetData(), Positions.Num());
+						for (const int32 GroupIndex : TrilinearInterpolationGroups)
+						{
+							const FMeshResizingCustomRegion& ResizingData = ResizingGroupData[GroupIndex];
+
+							UE::MeshResizing::FCustomRegionResizing::InterpolateCustomRegionPoints(ResizingData, TConstArrayView<FVector3d>(BoundCorners.GetData() + 8 * BoundCornersIndex, 8), Positions);
+							BoundCornersIndex += 8;
+						}
+
+						for (int32 VertexIndex = 0; VertexIndex < ResizingBlend.Num(); ++VertexIndex)
+						{
+							Positions[VertexIndex] = FMath::Lerp(OrigPositions[VertexIndex], Positions[VertexIndex], ResizingBlend[VertexIndex]);
+						}
+					}
+				}
+			}
+		} // namespace Private::Resizing
 	} // Private
 
 	void FClothDataflowTools::AddRenderPatternFromSkeletalMeshSection(const TSharedRef<FManagedArrayCollection>& ClothCollection, const FSkeletalMeshLODModel& SkeletalMeshModel, const int32 SectionIndex, const FString& RenderMaterialPathName)
@@ -892,4 +958,129 @@ namespace UE::Chaos::ClothAsset
 
 		return bHasDuplicateStitches;
 	}
+
+	void FClothDataflowTools::SetGroupResizingData(const TSharedRef<FManagedArrayCollection>& ClothCollection, const TConstArrayView<FName>& SetNames, const TConstArrayView<int32>& SetTypes)
+	{
+		check(SetNames.Num() == SetTypes.Num());
+		FCollectionClothFacade ClothFacade(ClothCollection);
+
+		ClothFacade.DefineSchema(EClothCollectionExtendedSchemas::Resizing);
+
+		// Clear the sim and render blends
+		TArrayView<float> SimBlend = ClothFacade.GetSimCustomResizingBlend();
+		TArrayView<float> RenderBlend = ClothFacade.GetRenderCustomResizingBlend();
+		for (int32 Index = 0; Index < SimBlend.Num(); ++Index)
+		{
+			SimBlend[Index] = 0.f;
+		}
+		for (int32 Index = 0; Index < RenderBlend.Num(); ++Index)
+		{
+			RenderBlend[Index] = 0.f;
+		}
+
+		// Filter out invalid sets
+		TArray<FName> FilteredSetNames;
+		TArray<int32> FilteredSetTypes;
+		FilteredSetNames.Reserve(SetNames.Num());
+		FilteredSetTypes.Reserve(SetNames.Num());
+		TSet<int32> SetValues;
+		for (int32 Index = 0; Index < SetNames.Num(); ++Index)
+		{
+			// Try to get the set as a SimVertex3D set
+			if (FClothGeometryTools::ConvertSelectionToNewGroupType(ClothCollection, SetNames[Index], ClothCollectionGroup::SimVertices3D, SetValues))
+			{
+				FilteredSetNames.Add(SetNames[Index]);
+				FilteredSetTypes.Add(SetTypes[Index]);
+
+				for (const int32 Vertex : SetValues)
+				{
+					if (SimBlend.IsValidIndex(Vertex))
+					{
+						SimBlend[Vertex] = 1.f;
+					}
+				}
+			}
+			else if (FClothGeometryTools::ConvertSelectionToNewGroupType(ClothCollection, SetNames[Index], ClothCollectionGroup::RenderVertices, SetValues))
+			{
+				FilteredSetNames.Add(SetNames[Index]);
+				FilteredSetTypes.Add(SetTypes[Index]);
+
+				for (const int32 Vertex : SetValues)
+				{
+					if (RenderBlend.IsValidIndex(Vertex))
+					{
+						RenderBlend[Vertex] = 1.f;
+					}
+				}
+			}
+		}
+
+		check(FilteredSetNames.Num() == FilteredSetTypes.Num());
+		ClothFacade.SetNumCustomResizingRegions(FilteredSetNames.Num());
+		TArrayView<FString> ClothSetNames = ClothFacade.GetCustomResizingRegionSet();
+		TArrayView<int32> ClothSetTypes = ClothFacade.GetCustomResizingRegionType();
+		for (int32 Index = 0; Index < FilteredSetNames.Num(); ++Index)
+		{
+			ClothSetNames[Index] = FilteredSetNames[Index].ToString();
+			ClothSetTypes[Index] = FilteredSetTypes[Index];
+		}
+	}
+
+	void FClothDataflowTools::GenerateSimMeshResizingGroupData(const TSharedRef<FManagedArrayCollection>& ClothCollection, const FMeshDescription& SourceMeshDescription, TArray<FMeshResizingCustomRegion>& OutResizingGroupData)
+	{
+		FCollectionClothFacade ClothFacade(ClothCollection);
+		if (ClothFacade.IsValid(EClothCollectionExtendedSchemas::Resizing))
+		{
+			UE::Geometry::FDynamicMesh3 SourceMesh;
+			FMeshDescriptionToDynamicMesh Converter;
+			Converter.Convert(&SourceMeshDescription, SourceMesh, true);
+
+			TConstArrayView<FString> ClothSetNames = ClothFacade.GetCustomResizingRegionSet();
+			TConstArrayView<int32> ClothSetTypes = ClothFacade.GetCustomResizingRegionType();
+			OutResizingGroupData.SetNum(ClothSetNames.Num());
+			for (int32 GroupIndex = 0; GroupIndex < ClothSetNames.Num(); ++GroupIndex)
+			{
+				TSet<int32> SetValues;
+				if (FClothGeometryTools::ConvertSelectionToNewGroupType(ClothCollection, FName(ClothSetNames[GroupIndex]), ClothCollectionGroup::SimVertices3D, SetValues))
+				{
+					UE::MeshResizing::FCustomRegionResizing::GenerateCustomRegion(ClothFacade.GetSimPosition3D(), SourceMesh, SetValues, OutResizingGroupData[GroupIndex]);
+				}
+			}
+		}
+	}
+
+	void FClothDataflowTools::GenerateRenderMeshResizingGroupData(const TSharedRef<FManagedArrayCollection>& ClothCollection, const FMeshDescription& SourceMeshDescription, TArray<FMeshResizingCustomRegion>& OutResizingGroupData)
+	{
+		FCollectionClothFacade ClothFacade(ClothCollection);
+		if (ClothFacade.IsValid(EClothCollectionExtendedSchemas::Resizing))
+		{
+			UE::Geometry::FDynamicMesh3 SourceMesh;
+			FMeshDescriptionToDynamicMesh Converter;
+			Converter.Convert(&SourceMeshDescription, SourceMesh, true);
+
+			TConstArrayView<FString> ClothSetNames = ClothFacade.GetCustomResizingRegionSet();
+			TConstArrayView<int32> ClothSetTypes = ClothFacade.GetCustomResizingRegionType();
+			OutResizingGroupData.SetNum(ClothSetNames.Num());
+			for (int32 GroupIndex = 0; GroupIndex < ClothSetNames.Num(); ++GroupIndex)
+			{
+				TSet<int32> SetValues;
+				if (FClothGeometryTools::ConvertSelectionToNewGroupType(ClothCollection, FName(ClothSetNames[GroupIndex]), ClothCollectionGroup::RenderVertices, SetValues))
+				{
+					UE::MeshResizing::FCustomRegionResizing::GenerateCustomRegion(ClothFacade.GetRenderPosition(), SourceMesh, SetValues, OutResizingGroupData[GroupIndex]);
+				}
+			}
+		}
+	}
+
+	void FClothDataflowTools::ApplySimGroupResizing(const TSharedRef<FManagedArrayCollection>& ClothCollection, const FMeshDescription& TargetMeshDescription, const FMeshResizingRBFInterpolationData& InterpolationData, const TArray<FMeshResizingCustomRegion>& ResizingGroupData)
+	{
+		FCollectionClothFacade ClothFacade(ClothCollection);
+		Private::Resizing::ApplyGroupResizing(ClothFacade, TargetMeshDescription, InterpolationData, ResizingGroupData, ClothFacade.GetSimPosition3D(), ClothFacade.GetSimCustomResizingBlend());
+	}
+
+	void FClothDataflowTools::ApplyRenderGroupResizing(const TSharedRef<FManagedArrayCollection>& ClothCollection, const FMeshDescription& TargetMeshDescription, const FMeshResizingRBFInterpolationData& InterpolationData, const TArray<FMeshResizingCustomRegion>& ResizingGroupData)
+	{
+		FCollectionClothFacade ClothFacade(ClothCollection);
+		Private::Resizing::ApplyGroupResizing(ClothFacade, TargetMeshDescription, InterpolationData, ResizingGroupData, ClothFacade.GetRenderPosition(), ClothFacade.GetRenderCustomResizingBlend());
+	}
 }  // End namespace UE::Chaos::ClothAsset
diff --git a/Engine/Plugins/ChaosClothAssetEditor/Source/ChaosClothAssetDataflowNodes/Private/ChaosClothAsset/ClothDataflowTools.h b/Engine/Plugins/ChaosClothAssetEditor/Source/ChaosClothAssetDataflowNodes/Private/ChaosClothAsset/ClothDataflowTools.h
--- a/Engine/Plugins/ChaosClothAssetEditor/Source/ChaosClothAssetDataflowNodes/Private/ChaosClothAsset/ClothDataflowTools.h
+++ b/Engine/Plugins/ChaosClothAssetEditor/Source/ChaosClothAssetDataflowNodes/Private/ChaosClothAsset/ClothDataflowTools.h
@@ -13,6 +13,8 @@ struct FDataflowNode;
 class FSkeletalMeshLODModel;
 class FString;
 class IPropertyHandle;
+struct FMeshResizingRBFInterpolationData;
+struct FMeshResizingCustomRegion;
 
 namespace UE::Chaos::ClothAsset
 {
@@ -89,6 +91,31 @@ namespace UE::Chaos::ClothAsset
 
 		static bool RemoveDuplicateStitches(TArray<TArray<FIntVector2>>& SeamStitches);
 
+		/**
+		 * Set Group Resizing data from supplied sets and types.
+		 */
+		static void SetGroupResizingData(const TSharedRef<FManagedArrayCollection>& ClothCollection, const TConstArrayView<FName>& SetNames, const TConstArrayView<int32>& SetTypes);
+
+		/**
+		 * Generate FMeshResizingCustomRegion data from Sim Mesh Group Resizing data
+		 */
+		static void GenerateSimMeshResizingGroupData(const TSharedRef<FManagedArrayCollection>& ClothCollection, const FMeshDescription& SourceMesh, TArray<FMeshResizingCustomRegion>& OutResizingGroupData);
+
+		/**
+		 * Generate FMeshResizingCustomRegion data from Render Mesh Group Resizing data
+		 */
+		static void GenerateRenderMeshResizingGroupData(const TSharedRef<FManagedArrayCollection>& ClothCollection, const FMeshDescription& SourceMesh, TArray<FMeshResizingCustomRegion>& OutResizingGroupData);
+
+		/**
+		 * Apply Group Resizing to the sim mesh.
+		 */
+		static void ApplySimGroupResizing(const TSharedRef<FManagedArrayCollection>& ClothCollection, const FMeshDescription& TargetMesh, const FMeshResizingRBFInterpolationData& InterpolationData, const TArray<FMeshResizingCustomRegion>& ResizingGroupData);
+
+		/**
+		 * Apply Group Resizing to the render mesh.
+		 */
+		static void ApplyRenderGroupResizing(const TSharedRef<FManagedArrayCollection>& ClothCollection, const FMeshDescription& TargetMesh, const FMeshResizingRBFInterpolationData& InterpolationData, const TArray<FMeshResizingCustomRegion>& ResizingGroupData);
+
 	private:
 		/** Return the Dataflow node owning by this property. */
 		static FDataflowNode* GetPropertyOwnerDataflowNode(const TSharedPtr<IPropertyHandle>& PropertyHandle, const UStruct* DataflowNodeStruct);
diff --git a/Engine/Plugins/ChaosClothAssetEditor/Source/ChaosClothAssetDataflowNodes/Private/ChaosClothAsset/CustomRegionResizingNode.cpp b/Engine/Plugins/ChaosClothAssetEditor/Source/ChaosClothAssetDataflowNodes/Private/ChaosClothAsset/CustomRegionResizingNode.cpp
--- a/Engine/Plugins/ChaosClothAssetEditor/Source/ChaosClothAssetDataflowNodes/Private/ChaosClothAsset/CustomRegionResizingNode.cpp
+++ b/Engine/Plugins/ChaosClothAssetEditor/Source/ChaosClothAssetDataflowNodes/Private/ChaosClothAsset/CustomRegionResizingNode.cpp
@@ -2,7 +2,7 @@
 
 #include "ChaosClothAsset/CustomRegionResizingNode.h"
 #include "ChaosClothAsset/ClothCollectionAttribute.h"
-#include "ChaosClothAsset/ClothGeometryTools.h"
+#include "ChaosClothAsset/ClothDataflowTools.h"
 #include "Dataflow/DataflowInputOutput.h"
 
 #include UE_INLINE_GENERATED_CPP_BY_NAME(CustomRegionResizingNode)
@@ -50,7 +50,7 @@ void FChaosClothAssetCustomRegionResizingNode::Evaluate(UE::Dataflow::FContext&
 			InputTypes[Index] = (int32)InputGroupData[Index].ResizingType;
 		}
 
-		FClothGeometryTools::SetGroupResizingData(ClothCollection, InputSets, InputTypes);
+		FClothDataflowTools::SetGroupResizingData(ClothCollection, InputSets, InputTypes);
 		SetValue(Context, MoveTemp(*ClothCollection), &Collection);
 	}
 	else if(Out->IsA<FString>(&SimCustomResizingBlendName))
