Chaos Cloth Asset - Fixed USD render mesh import when the Prim name ends weith a number.

#rnx

diff --git a/Engine/Plugins/ChaosClothAssetEditor/Source/ChaosClothAssetDataflowNodes/Private/ChaosClothAsset/USDImportNode_v2.cpp b/Engine/Plugins/ChaosClothAssetEditor/Source/ChaosClothAssetDataflowNodes/Private/ChaosClothAsset/USDImportNode_v2.cpp
--- a/Engine/Plugins/ChaosClothAssetEditor/Source/ChaosClothAssetDataflowNodes/Private/ChaosClothAsset/USDImportNode_v2.cpp
+++ b/Engine/Plugins/ChaosClothAssetEditor/Source/ChaosClothAssetDataflowNodes/Private/ChaosClothAsset/USDImportNode_v2.cpp
@@ -25,6 +25,7 @@
 #include "StaticMeshAttributes.h"
 #include "UObject/Package.h"
 #include "UnrealUSDWrapper.h"
+#include "USDAssetUserData.h"
 #include "USDConversionUtils.h"
 #include "USDProjectSettings.h"
 #include "USDStageImportContext.h"
@@ -112,7 +113,7 @@ namespace UE::Chaos::ClothAsset::Private
 		}
 	}
 
-	static void ImportStaticMeshesFromUsdStage(const FUsdStage& UsdStage, const FString& UsdFilePath, const FString& PackagePath)
+	static TArray<TObjectPtr<UObject>> ImportStaticMeshesFromUsdStage(const FUsdStage& UsdStage, const FString& UsdFilePath, const FString& PackagePath)
 	{
 		// Import recognised assets
 		FUsdStageImportContext ImportContext;
@@ -170,6 +171,8 @@ namespace UE::Chaos::ClothAsset::Private
 
 		// Restore Original USD Materials
 		OverrideUsdImportMaterials_v2(OriginalUsdMaterials);
+
+		return ImportContext.ImportedAssets;
 	}
 
 	static FUsdPrim FindClothPrim(const FUsdPrim& RootPrim)
@@ -707,12 +710,11 @@ bool FChaosClothAssetUSDImportNode_v2::ImportUsdFile(const FString& UsdFilePath,
 
 	// Import the stage
 	SlowTask.EnterProgressFrame(1.f);
-	ImportStaticMeshesFromUsdStage(UsdStage, UsdFilePath, PackagePath);
+	TArray<TObjectPtr<UObject>> ImportedAssets = ImportStaticMeshesFromUsdStage(UsdStage, UsdFilePath, PackagePath);
 
-	// Fill up the asset list from the imported USD assets
-	const FString SimMeshName = SimMeshPrim ? FString::Printf(TEXT("SM_%s"), *SimMeshPrim.GetName().ToString()) : FString();
-	const FString RenderMeshName = RenderMeshPrim ? FString::Printf(TEXT("SM_%s"), *RenderMeshPrim.GetName().ToString()) : FString();
-	UpdateImportedAssets(SimMeshName, RenderMeshName);
+	// Find the imported static meshes in the imported USD assets
+	ImportedSimStaticMesh = FindImportedStaticMesh(ImportedAssets, SimMeshPrim.GetPrimPath().GetString());
+	ImportedRenderStaticMesh = FindImportedStaticMesh(ImportedAssets, RenderMeshPrim.GetPrimPath().GetString());
 
 	// Import sim mesh from the static mesh
 	if (bImportSimMesh)
@@ -810,6 +812,25 @@ void FChaosClothAssetUSDImportNode_v2::UpdateImportedAssets(const FString& SimMe
 	}
 }
 
+TObjectPtr<UStaticMesh> FChaosClothAssetUSDImportNode_v2::FindImportedStaticMesh(const TArrayView<TObjectPtr<UObject>> ImportedAssets, const FString& MeshPrimPath)
+{
+	for (TObjectPtr<UObject>& ImportedAsset : ImportedAssets)
+	{
+		if (const TObjectPtr<UStaticMesh> ImportedStaticMesh = Cast<UStaticMesh>(ImportedAsset))
+		{
+			if (const UUsdMeshAssetUserData* const AssetUserData = Cast<UUsdMeshAssetUserData>(
+				ImportedStaticMesh->GetAssetUserDataOfClass(UUsdMeshAssetUserData::StaticClass())))
+			{
+				if (AssetUserData->PrimPaths.Find(MeshPrimPath) != INDEX_NONE)
+				{
+					return ImportedStaticMesh;
+				}
+			}
+		}
+	}
+	return nullptr;
+}
+
 TArray<TObjectPtr<UObject>> FChaosClothAssetUSDImportNode_v2::GetImportedAssetDependencies(const UObject* StaticMesh)
 {
 	using namespace UE::Chaos::ClothAsset::Private;
diff --git a/Engine/Plugins/ChaosClothAssetEditor/Source/ChaosClothAssetDataflowNodes/Public/ChaosClothAsset/USDImportNode_v2.h b/Engine/Plugins/ChaosClothAssetEditor/Source/ChaosClothAssetDataflowNodes/Public/ChaosClothAsset/USDImportNode_v2.h
--- a/Engine/Plugins/ChaosClothAssetEditor/Source/ChaosClothAssetDataflowNodes/Public/ChaosClothAsset/USDImportNode_v2.h
+++ b/Engine/Plugins/ChaosClothAssetEditor/Source/ChaosClothAssetDataflowNodes/Public/ChaosClothAsset/USDImportNode_v2.h
@@ -108,8 +108,12 @@ private:
 	// Import the given static mesh as a render mesh into the cloth collection
 	bool ImportRenderStaticMesh(const TSharedRef<FManagedArrayCollection> ClothCollection, class FText& OutErrorText);
 
-	// Find the two imported static mesh and update the dependencies
+	// Find the two imported static meshes
 	void UpdateImportedAssets(const FString& SimMeshName, const FString& RenderMeshName);
+
+	// Find an imported static mesh from the given usd prim path
+	TObjectPtr<UStaticMesh> FindImportedStaticMesh(const TArrayView<TObjectPtr<UObject>> ImportedAssetsconst, const FString& MeshPrimPath);
+
 	// Find the asset dependencies for the specified static mesh
 	TArray<TObjectPtr<UObject>> GetImportedAssetDependencies(const UObject* StaticMesh);
 
