Updated headers using UnrealCodeFixup to make sure dllstorage is on methods/staticvars instead of types. Part 1/n

diff --git a/Engine/Plugins/Animation/TweeningUtils/Source/TweeningUtilsEditor/Public/Widgets/MVC/TweenToolbarController.h b/Engine/Plugins/Animation/TweeningUtils/Source/TweeningUtilsEditor/Public/Widgets/MVC/TweenToolbarController.h
--- a/Engine/Plugins/Animation/TweeningUtils/Source/TweeningUtilsEditor/Public/Widgets/MVC/TweenToolbarController.h
+++ b/Engine/Plugins/Animation/TweeningUtils/Source/TweeningUtilsEditor/Public/Widgets/MVC/TweenToolbarController.h
@@ -1,4 +1,4 @@
-﻿// Copyright Epic Games, Inc. All Rights Reserved.
+// Copyright Epic Games, Inc. All Rights Reserved.
 
 #pragma once
 
@@ -9,6 +9,8 @@
 #include "Templates/SharedPointer.h"
 #include "Templates/UnrealTemplate.h"
 
+#define UE_API TWEENINGUTILSEDITOR_API
+
 class FText;
 class FToolBarBuilder;
 class FUICommandList;
@@ -32,7 +34,7 @@ class FTweenModel;
  * - Slider: Adjusts the blend value and updates the FTweenModel
  * - Overshoot: Toggle button for changing the overshoot property on FTweenModel
  */
-class TWEENINGUTILSEDITOR_API FTweenToolbarController : public FNoncopyable
+class FTweenToolbarController : public FNoncopyable
 {
 public:
 	
@@ -41,12 +43,12 @@ public:
 	 * @param InTweenFunctions The functions that can be selected
 	 * @param InInitialTweenModelIndex Index of the tween model to initially be selected.
 	 */
-	explicit FTweenToolbarController(
+	UE_API explicit FTweenToolbarController(
 		const TSharedRef<FUICommandList>& InCommandList,
 		const TSharedRef<ITweenModelContainer>& InTweenFunctions,
 		int32 InInitialTweenModelIndex = 0
 		);
-	~FTweenToolbarController();
+	UE_API ~FTweenToolbarController();
 
 	struct FMakeWidgetArgs
 	{
@@ -69,10 +71,10 @@ public:
 	 * @param ToolbarBuilder The toolbar to add the widgets to
 	 * @param InArgs Optional additional args to control how the widget should be constructed.
 	 */
-	void AddToToolbar(FToolBarBuilder& ToolbarBuilder, FMakeWidgetArgs InArgs = {}) const;
+	UE_API void AddToToolbar(FToolBarBuilder& ToolbarBuilder, FMakeWidgetArgs InArgs = {}) const;
 
 	/** Sets the currently selected tween model. Asserts that the tween model is contained in TweenModels. */
-	void SetSelectedTweenModel(FTweenModel& InTweenModel);
+	UE_API void SetSelectedTweenModel(FTweenModel& InTweenModel);
 
 	/** @return Gets the currently selected tween model */
 	FTweenModel& GetSelectedTweenModel() const { return *SelectedTweenModel; }
@@ -99,19 +101,21 @@ private:
 	FOnTweenFunctionChanged OnTweenFunctionChangedDelegate;
 
 	/** Creates the drop-down list for selecting the tween mode. */
-	TSharedRef<SWidget> MakeTweenModeMenu() const;
-	FText GetLabelForComboBox() const;
-	FText GetToolTipForComboBox() const;
-	FSlateIcon GetIconForComboBox() const;
+	UE_API TSharedRef<SWidget> MakeTweenModeMenu() const;
+	UE_API FText GetLabelForComboBox() const;
+	UE_API FText GetToolTipForComboBox() const;
+	UE_API FSlateIcon GetIconForComboBox() const;
 	
-	void SetTweenModel(FTweenModel* InTweenModel);
-	bool IsTweenModelSelected(FTweenModel* InTweenModel) const;
+	UE_API void SetTweenModel(FTweenModel* InTweenModel);
+	UE_API bool IsTweenModelSelected(FTweenModel* InTweenModel) const;
 	
-	void ToggleOvershootMode() const;
-	bool IsOvershootModeEnabled() const;
-	FSlateIcon GetOvershootModeIcon() const;
+	UE_API void ToggleOvershootMode() const;
+	UE_API bool IsOvershootModeEnabled() const;
+	UE_API FSlateIcon GetOvershootModeIcon() const;
 	
-	const FSlateBrush* GetIconForSlider() const;
-	FLinearColor GetColorForSlider() const;
+	UE_API const FSlateBrush* GetIconForSlider() const;
+	UE_API FLinearColor GetColorForSlider() const;
 };
 }
+
+#undef UE_API
diff --git a/Engine/Plugins/Animation/TweeningUtils/Source/TweeningUtilsEditor/Public/Widgets/STweenSlider.h b/Engine/Plugins/Animation/TweeningUtils/Source/TweeningUtilsEditor/Public/Widgets/STweenSlider.h
--- a/Engine/Plugins/Animation/TweeningUtils/Source/TweeningUtilsEditor/Public/Widgets/STweenSlider.h
+++ b/Engine/Plugins/Animation/TweeningUtils/Source/TweeningUtilsEditor/Public/Widgets/STweenSlider.h
@@ -9,6 +9,8 @@
 #include "Widgets/DeclarativeSyntaxSupport.h"
 #include "Widgets/SLeafWidget.h"
 
+#define UE_API TWEENINGUTILSEDITOR_API
+
 template<typename T> struct TOptional;
 
 namespace UE::TweeningUtilsEditor
@@ -25,7 +27,7 @@ struct FSliderWidgetData;
  *	- I is the Icon and the bit behind it is the SliderIconBackground. Together, they form a button.
  *	- + is a point on the scale.
  */
-class TWEENINGUTILSEDITOR_API STweenSlider : public SLeafWidget
+class STweenSlider : public SLeafWidget
 {
 public:
 	
@@ -70,7 +72,7 @@ public:
 		
 	SLATE_END_ARGS()
 
-	void Construct(const FArguments& InArgs);
+	UE_API void Construct(const FArguments& InArgs);
 
 private:
 
@@ -95,30 +97,32 @@ private:
 	FMapSliderValueToBlendValue MapSliderValueToBlendValueDelegate;
 	
 	//~ Begin SLeafWidget Interface
-	virtual FReply OnMouseButtonDoubleClick(const FGeometry& InMyGeometry, const FPointerEvent& InMouseEvent) override;
-	virtual FReply OnMouseButtonDown(const FGeometry& InGeometry, const FPointerEvent& InMouseEvent) override;
-	virtual FReply OnMouseButtonUp(const FGeometry& InGeometry, const FPointerEvent& InMouseEvent) override;
-	virtual void OnMouseLeave(const FPointerEvent& InMouseEvent) override;
-	virtual FReply OnMouseMove(const FGeometry& MyGeometry, const FPointerEvent& InMouseEvent) override;
-	virtual void OnFinishedPointerInput() override;
-	virtual FVector2D ComputeDesiredSize(float) const override;
-	virtual void Tick(const FGeometry& AllottedGeometry, const double InCurrentTime, const float InDeltaTime) override;
-	virtual int32 OnPaint(const FPaintArgs& InArgs, const FGeometry& InAllottedGeometry, const FSlateRect& InCullingRect, FSlateWindowElementList& OutDrawElements, int32 InLayerId, const FWidgetStyle& InWidgetStyle, bool bParentEnabled) const override;
-	virtual TSharedPtr<IToolTip> GetToolTip() override;
+	UE_API virtual FReply OnMouseButtonDoubleClick(const FGeometry& InMyGeometry, const FPointerEvent& InMouseEvent) override;
+	UE_API virtual FReply OnMouseButtonDown(const FGeometry& InGeometry, const FPointerEvent& InMouseEvent) override;
+	UE_API virtual FReply OnMouseButtonUp(const FGeometry& InGeometry, const FPointerEvent& InMouseEvent) override;
+	UE_API virtual void OnMouseLeave(const FPointerEvent& InMouseEvent) override;
+	UE_API virtual FReply OnMouseMove(const FGeometry& MyGeometry, const FPointerEvent& InMouseEvent) override;
+	UE_API virtual void OnFinishedPointerInput() override;
+	UE_API virtual FVector2D ComputeDesiredSize(float) const override;
+	UE_API virtual void Tick(const FGeometry& AllottedGeometry, const double InCurrentTime, const float InDeltaTime) override;
+	UE_API virtual int32 OnPaint(const FPaintArgs& InArgs, const FGeometry& InAllottedGeometry, const FSlateRect& InCullingRect, FSlateWindowElementList& OutDrawElements, int32 InLayerId, const FWidgetStyle& InWidgetStyle, bool bParentEnabled) const override;
+	UE_API virtual TSharedPtr<IToolTip> GetToolTip() override;
 	//~ End SLeafWidget Interface
 
 	/** Updates LastMousePositionOnSliderBar. */
-	void UpdateLastMousePositionOnSlider(const FGeometry& InGeometry, const FVector2D& InMousePos);
+	UE_API void UpdateLastMousePositionOnSlider(const FGeometry& InGeometry, const FVector2D& InMousePos);
 
-	bool IsMouseDown() const;
-	bool IsDragging() const;
+	UE_API bool IsMouseDown() const;
+	UE_API bool IsDragging() const;
 
 	void OnStartDragOp() const;
 
 	/** If the user is currently hovering a point, set the blend value to that point's value.  */
-	void HandlePickValueOfCurrentlyHoveredPoint(const FGeometry& InGeometry, const FPointerEvent& InMouseEvent) const;
+	UE_API void HandlePickValueOfCurrentlyHoveredPoint(const FGeometry& InGeometry, const FPointerEvent& InMouseEvent) const;
 
 	/** Gets the tooltip text for what is currently being hovered. */
-	FText GetToolTipText() const;
+	UE_API FText GetToolTipText() const;
 };
 }
+
+#undef UE_API
diff --git a/Engine/Plugins/Animation/TweeningUtils/Source/TweeningUtilsEditor/Public/Widgets/TweenSliderStyle.h b/Engine/Plugins/Animation/TweeningUtils/Source/TweeningUtilsEditor/Public/Widgets/TweenSliderStyle.h
--- a/Engine/Plugins/Animation/TweeningUtils/Source/TweeningUtilsEditor/Public/Widgets/TweenSliderStyle.h
+++ b/Engine/Plugins/Animation/TweeningUtils/Source/TweeningUtilsEditor/Public/Widgets/TweenSliderStyle.h
@@ -8,13 +8,15 @@
 #include "Styling/SlateWidgetStyle.h"
 #include "TweenSliderStyle.generated.h"
 
+#define UE_API TWEENINGUTILSEDITOR_API
+
 USTRUCT()
-struct TWEENINGUTILSEDITOR_API FTweenPointStyle
+struct FTweenPointStyle
 {
 	GENERATED_BODY()
 	
-	FTweenPointStyle();
-	explicit FTweenPointStyle(
+	UE_API FTweenPointStyle();
+	UE_API explicit FTweenPointStyle(
 		const FVector2D& InNormalSize, const FVector2D& InHoveredSize, const FVector2D& InPressedSize, const FVector2D& InHitSize
 		);
 	
@@ -43,25 +45,25 @@ struct TWEENINGUTILSEDITOR_API FTweenPointStyle
 	FVector2D HitTestSize = { 10.0, 12.0 };
 	FTweenPointStyle& SetHitTestSize( const FVector2D& InSize ){ HitTestSize = InSize; return *this; }
 	
-	void GetResources(TArray<const FSlateBrush*>& OutBrushes) const;
+	UE_API void GetResources(TArray<const FSlateBrush*>& OutBrushes) const;
 };
 
 /** Style for STweenSlider */
 USTRUCT()
-struct TWEENINGUTILSEDITOR_API FTweenSliderStyle : public FSlateWidgetStyle
+struct FTweenSliderStyle : public FSlateWidgetStyle
 {
 	GENERATED_BODY()
 
-	FTweenSliderStyle();
+	UE_API FTweenSliderStyle();
 
 	//~ Begin FSlateWidgetStyle Interface
-	virtual void GetResources(TArray<const FSlateBrush*>& OutBrushes) const override;
+	UE_API virtual void GetResources(TArray<const FSlateBrush*>& OutBrushes) const override;
 	//~ End FSlateWidgetStyle Interface
 	
-	static const FName TypeName;
+	static UE_API const FName TypeName;
 	virtual const FName GetTypeName() const override { return TypeName; }
 	
-	static const FTweenSliderStyle& GetDefault();
+	static UE_API const FTweenSliderStyle& GetDefault();
 	
 	/** The dimensions of the slider. */
 	UPROPERTY(EditAnywhere, Category=Appearance)
@@ -137,4 +139,6 @@ struct TWEENINGUTILSEDITOR_API FTweenSliderStyle : public FSlateWidgetStyle
 	UPROPERTY(EditAnywhere, Category=Appearance)
 	FMargin IconPadding;
 	FTweenSliderStyle& SetIconPadding( const FMargin& InIconPadding ){ IconPadding = InIconPadding; return *this; }
-};
\ No newline at end of file
+};
+
+#undef UE_API
diff --git a/Engine/Plugins/AudioInsights/Source/AudioInsights/Public/AudioInsightsDataSource.h b/Engine/Plugins/AudioInsights/Source/AudioInsights/Public/AudioInsightsDataSource.h
--- a/Engine/Plugins/AudioInsights/Source/AudioInsights/Public/AudioInsightsDataSource.h
+++ b/Engine/Plugins/AudioInsights/Source/AudioInsights/Public/AudioInsightsDataSource.h
@@ -6,7 +6,7 @@
 
 namespace UE::Audio::Insights
 {
-	class AUDIOINSIGHTS_API IDashboardDataViewEntry : public TSharedFromThis<IDashboardDataViewEntry>
+	class IDashboardDataViewEntry : public TSharedFromThis<IDashboardDataViewEntry>
 	{
 	public:
 		virtual ~IDashboardDataViewEntry() = default;
@@ -14,7 +14,7 @@ namespace UE::Audio::Insights
 		virtual bool IsValid() const = 0;
 	};
 
-	class AUDIOINSIGHTS_API IDashboardDataTreeViewEntry : public TSharedFromThis<IDashboardDataTreeViewEntry>
+	class IDashboardDataTreeViewEntry : public TSharedFromThis<IDashboardDataTreeViewEntry>
 	{
 	public:
 		virtual ~IDashboardDataTreeViewEntry() = default;
diff --git a/Engine/Plugins/AudioInsights/Source/AudioInsights/Public/AudioInsightsStyle.h b/Engine/Plugins/AudioInsights/Source/AudioInsights/Public/AudioInsightsStyle.h
--- a/Engine/Plugins/AudioInsights/Source/AudioInsights/Public/AudioInsightsStyle.h
+++ b/Engine/Plugins/AudioInsights/Source/AudioInsights/Public/AudioInsightsStyle.h
@@ -14,7 +14,7 @@
 
 namespace UE::Audio::Insights
 {
-	class AUDIOINSIGHTS_API FSlateStyle final : public FSlateStyleSet
+	class FSlateStyle final : public FSlateStyleSet
 	{
 	public:
 		static FSlateStyle& Get()
diff --git a/Engine/Plugins/AudioInsights/Source/AudioInsights/Public/AudioInsightsTraceModule.h b/Engine/Plugins/AudioInsights/Source/AudioInsights/Public/AudioInsightsTraceModule.h
--- a/Engine/Plugins/AudioInsights/Source/AudioInsights/Public/AudioInsightsTraceModule.h
+++ b/Engine/Plugins/AudioInsights/Source/AudioInsights/Public/AudioInsightsTraceModule.h
@@ -11,6 +11,8 @@
 #include "TraceServices/Model/AnalysisSession.h"
 #include "UObject/NameTypes.h"
 
+#define UE_API AUDIOINSIGHTS_API
+
 namespace UE::Audio::Insights
 {
 	class FRewindDebuggerAudioInsightsRuntime : public IRewindDebuggerRuntimeExtension
@@ -19,17 +21,17 @@ namespace UE::Audio::Insights
 		virtual void RecordingStarted() override;
 	};
 
-	class AUDIOINSIGHTS_API FTraceModule : public IAudioInsightsTraceModule
+	class FTraceModule : public IAudioInsightsTraceModule
 	{
 	public:
-		FTraceModule();
+		UE_API FTraceModule();
 		virtual ~FTraceModule() = default;
 
 		//~ Begin TraceServices::IModule interface
-		virtual void GetModuleInfo(TraceServices::FModuleInfo& OutModuleInfo) override;
-		virtual void OnAnalysisBegin(TraceServices::IAnalysisSession& Session) override;
-		virtual void GetLoggers(TArray<const TCHAR *>& OutLoggers) override;
-		virtual void GenerateReports(const TraceServices::IAnalysisSession& Session, const TCHAR* CmdLine, const TCHAR* OutputDirectory) override;
+		UE_API virtual void GetModuleInfo(TraceServices::FModuleInfo& OutModuleInfo) override;
+		UE_API virtual void OnAnalysisBegin(TraceServices::IAnalysisSession& Session) override;
+		UE_API virtual void GetLoggers(TArray<const TCHAR *>& OutLoggers) override;
+		UE_API virtual void GenerateReports(const TraceServices::IAnalysisSession& Session, const TCHAR* CmdLine, const TCHAR* OutputDirectory) override;
 		virtual const TCHAR* GetCommandLineArgument() override { return TEXT("audiotrace"); }
 		//~ End TraceServices::IModule interface
 
@@ -39,42 +41,42 @@ namespace UE::Audio::Insights
 			return StaticCastSharedPtr<TraceProviderType>(TraceProviders.FindRef(TraceProviderType::GetName_Static()));
 		}
 
-		virtual void AddTraceProvider(TSharedPtr<FTraceProviderBase> TraceProvider) override;
+		UE_API virtual void AddTraceProvider(TSharedPtr<FTraceProviderBase> TraceProvider) override;
 
-		virtual void StartTraceAnalysis(const bool bOnlyTraceAudioChannels) override;
-		virtual bool IsTraceAnalysisActive() const override;
-		virtual void StopTraceAnalysis() override;
-		virtual void OnOnlyTraceAudioChannelsStateChanged(const bool bOnlyTraceAudioChannels) override;
+		UE_API virtual void StartTraceAnalysis(const bool bOnlyTraceAudioChannels) override;
+		UE_API virtual bool IsTraceAnalysisActive() const override;
+		UE_API virtual void StopTraceAnalysis() override;
+		UE_API virtual void OnOnlyTraceAudioChannelsStateChanged(const bool bOnlyTraceAudioChannels) override;
 
-		virtual bool AudioChannelsCanBeManuallyEnabled() const override;
+		UE_API virtual bool AudioChannelsCanBeManuallyEnabled() const override;
 
 #if !WITH_EDITOR
-		virtual void InitializeSessionInfo(const TraceServices::FSessionInfo& SessionInfo) override;
-		virtual void RequestChannelUpdate() override;
-		virtual void ResetTicker() override;
-		virtual bool TraceControllerIsAvailable() const override;
+		UE_API virtual void InitializeSessionInfo(const TraceServices::FSessionInfo& SessionInfo) override;
+		UE_API virtual void RequestChannelUpdate() override;
+		UE_API virtual void ResetTicker() override;
+		UE_API virtual bool TraceControllerIsAvailable() const override;
 #endif
 
 		DECLARE_MULTICAST_DELEGATE_OneParam(FOnAnalysisStarting, const double /*Timestamp*/);
 		FOnAnalysisStarting OnAnalysisStarting;
 
 	private:
-		static const FName GetName();
+		static UE_API const FName GetName();
 
-		static void DisableAllTraceChannels();
-		bool EnableAudioInsightsTraceChannels();
+		static UE_API void DisableAllTraceChannels();
+		UE_API bool EnableAudioInsightsTraceChannels();
 
-		void DisableAudioInsightsTraceChannels() const;
+		UE_API void DisableAudioInsightsTraceChannels() const;
 
-		void CacheCurrentlyEnabledTraceChannels();
-		void RestoreCachedChannels() const;
+		UE_API void CacheCurrentlyEnabledTraceChannels();
+		UE_API void RestoreCachedChannels() const;
 
 #if !WITH_EDITOR
-		const UE::Trace::FStoreClient::FSessionInfo* GetCurrentAnalysisSessionInfo() const;
-		bool GetAudioTracesAreEnabled() const;
-		void SendDiscoveryRequestToTraceController() const;
+		UE_API const UE::Trace::FStoreClient::FSessionInfo* GetCurrentAnalysisSessionInfo() const;
+		UE_API bool GetAudioTracesAreEnabled() const;
+		UE_API void SendDiscoveryRequestToTraceController() const;
 
-		bool Tick(float DeltaTime);
+		UE_API bool Tick(float DeltaTime);
 #endif
 
 		TMap<FName, TSharedPtr<FTraceProviderBase>> TraceProviders;
@@ -95,3 +97,5 @@ namespace UE::Audio::Insights
 #endif
 	};
 } // namespace UE::Audio::Insights
+
+#undef UE_API
diff --git a/Engine/Plugins/AudioInsights/Source/AudioInsights/Public/AudioInsightsTraceProviderBase.h b/Engine/Plugins/AudioInsights/Source/AudioInsights/Public/AudioInsightsTraceProviderBase.h
--- a/Engine/Plugins/AudioInsights/Source/AudioInsights/Public/AudioInsightsTraceProviderBase.h
+++ b/Engine/Plugins/AudioInsights/Source/AudioInsights/Public/AudioInsightsTraceProviderBase.h
@@ -20,16 +20,16 @@
 
 namespace UE::Audio::Insights
 {
-	class AUDIOINSIGHTS_API FTraceProviderBase : public TraceServices::IProvider, public TraceServices::IEditableProvider
+	class FTraceProviderBase : public TraceServices::IProvider, public TraceServices::IEditableProvider
 	{
 	public:
 		FTraceProviderBase() = delete;
-		explicit FTraceProviderBase(FName InName);
+		AUDIOINSIGHTS_API explicit FTraceProviderBase(FName InName);
 
-		virtual ~FTraceProviderBase();
+		AUDIOINSIGHTS_API virtual ~FTraceProviderBase();
 
 		virtual Trace::IAnalyzer* ConstructAnalyzer(TraceServices::IAnalysisSession& InSession) = 0;
-		FName GetName() const;
+		AUDIOINSIGHTS_API FName GetName() const;
 
 		virtual void Reset()
 		{
@@ -74,17 +74,17 @@ namespace UE::Audio::Insights
 #endif // !WITH_EDITOR
 
 	protected:
-		class AUDIOINSIGHTS_API FTraceAnalyzerBase : public Trace::IAnalyzer
+		class FTraceAnalyzerBase : public Trace::IAnalyzer
 		{
 		public:
-			FTraceAnalyzerBase(TSharedRef<FTraceProviderBase> InProvider);
+			AUDIOINSIGHTS_API FTraceAnalyzerBase(TSharedRef<FTraceProviderBase> InProvider);
 			virtual ~FTraceAnalyzerBase() = default;
 
-			virtual void OnAnalysisBegin(const FOnAnalysisContext& Context) override;
+			AUDIOINSIGHTS_API virtual void OnAnalysisBegin(const FOnAnalysisContext& Context) override;
 
 		protected:
-			virtual bool OnEventSuccess(uint16 RouteId, EStyle Style, const FOnEventContext& Context);
-			virtual bool OnEventFailure(uint16 RouteId, EStyle Style, const FOnEventContext& Context);
+			AUDIOINSIGHTS_API virtual bool OnEventSuccess(uint16 RouteId, EStyle Style, const FOnEventContext& Context);
+			AUDIOINSIGHTS_API virtual bool OnEventFailure(uint16 RouteId, EStyle Style, const FOnEventContext& Context);
 
 			template <typename TProviderType>
 			TProviderType& GetProvider()
diff --git a/Engine/Plugins/AudioInsights/Source/AudioInsights/Public/IAudioInsightsDashboardFactory.h b/Engine/Plugins/AudioInsights/Source/AudioInsights/Public/IAudioInsightsDashboardFactory.h
--- a/Engine/Plugins/AudioInsights/Source/AudioInsights/Public/IAudioInsightsDashboardFactory.h
+++ b/Engine/Plugins/AudioInsights/Source/AudioInsights/Public/IAudioInsightsDashboardFactory.h
@@ -8,7 +8,7 @@
 
 namespace UE::Audio::Insights
 {
-	class AUDIOINSIGHTS_API IDashboardFactory
+	class IDashboardFactory
 	{
 	public:
 		virtual ~IDashboardFactory() = default;
diff --git a/Engine/Plugins/AudioInsights/Source/AudioInsights/Public/IAudioInsightsModule.h b/Engine/Plugins/AudioInsights/Source/AudioInsights/Public/IAudioInsightsModule.h
--- a/Engine/Plugins/AudioInsights/Source/AudioInsights/Public/IAudioInsightsModule.h
+++ b/Engine/Plugins/AudioInsights/Source/AudioInsights/Public/IAudioInsightsModule.h
@@ -9,6 +9,8 @@
 #include "UObject/NameTypes.h"
 #include "Widgets/Docking/SDockTab.h"
 
+#define UE_API AUDIOINSIGHTS_API
+
 class IAudioInsightsTraceModule;
 
 namespace UE::Audio::Insights
@@ -17,7 +19,7 @@ namespace UE::Audio::Insights
 	class FTraceProviderBase;
 } // namespace UE::Audio::Insights
 
-class AUDIOINSIGHTS_API IAudioInsightsModule : public IAudioInsightsModuleInterface
+class IAudioInsightsModule : public IAudioInsightsModuleInterface
 {
 public:
 	virtual void RegisterDashboardViewFactory(TSharedRef<UE::Audio::Insights::IDashboardViewFactory> InDashboardFactory) = 0;
@@ -25,14 +27,16 @@ public:
 
 	virtual ::Audio::FDeviceId GetDeviceId() const = 0;
 
-	virtual IAudioInsightsTraceModule& GetTraceModule() override;
+	UE_API virtual IAudioInsightsTraceModule& GetTraceModule() override;
 
 	virtual TSharedRef<SDockTab> CreateDashboardTabWidget(const FSpawnTabArgs& Args) = 0;
 
-	static IAudioInsightsModule& GetChecked();
+	static UE_API IAudioInsightsModule& GetChecked();
 
 #ifdef WITH_EDITOR
-	static IAudioInsightsModule& GetEditorChecked();
+	static UE_API IAudioInsightsModule& GetEditorChecked();
 #endif
 
 };
+
+#undef UE_API
diff --git a/Engine/Plugins/AudioInsights/Source/AudioInsights/Public/IAudioInsightsModuleInterface.h b/Engine/Plugins/AudioInsights/Source/AudioInsights/Public/IAudioInsightsModuleInterface.h
--- a/Engine/Plugins/AudioInsights/Source/AudioInsights/Public/IAudioInsightsModuleInterface.h
+++ b/Engine/Plugins/AudioInsights/Source/AudioInsights/Public/IAudioInsightsModuleInterface.h
@@ -14,7 +14,7 @@ namespace UE::Audio::Insights
 	class FTraceProviderBase;
 } // namespace UE::Audio::Insights
 
-class AUDIOINSIGHTS_API IAudioInsightsModuleInterface : public IModuleInterface
+class IAudioInsightsModuleInterface : public IModuleInterface
 {
 public:
 	virtual void RegisterDashboardViewFactory(TSharedRef<UE::Audio::Insights::IDashboardViewFactory> InDashboardFactory) = 0;
diff --git a/Engine/Plugins/AudioInsights/Source/AudioInsights/Public/IAudioInsightsTraceModule.h b/Engine/Plugins/AudioInsights/Source/AudioInsights/Public/IAudioInsightsTraceModule.h
--- a/Engine/Plugins/AudioInsights/Source/AudioInsights/Public/IAudioInsightsTraceModule.h
+++ b/Engine/Plugins/AudioInsights/Source/AudioInsights/Public/IAudioInsightsTraceModule.h
@@ -14,7 +14,7 @@ namespace TraceServices
 	struct FSessionInfo;
 } // namespace TraceServices
 
-class AUDIOINSIGHTS_API IAudioInsightsTraceModule : public TraceServices::IModule
+class IAudioInsightsTraceModule : public TraceServices::IModule
 {
 public:
 	virtual void AddTraceProvider(TSharedPtr<UE::Audio::Insights::FTraceProviderBase> TraceProvider) = 0;
diff --git a/Engine/Plugins/AudioInsights/Source/AudioInsights/Public/Views/MixerSourceDashboardViewFactory.h b/Engine/Plugins/AudioInsights/Source/AudioInsights/Public/Views/MixerSourceDashboardViewFactory.h
--- a/Engine/Plugins/AudioInsights/Source/AudioInsights/Public/Views/MixerSourceDashboardViewFactory.h
+++ b/Engine/Plugins/AudioInsights/Source/AudioInsights/Public/Views/MixerSourceDashboardViewFactory.h
@@ -7,99 +7,100 @@
 #include "Views/TableDashboardViewFactory.h"
 #include "Widgets/Input/SCheckBox.h"
 
+#define UE_API AUDIOINSIGHTS_API
+
 namespace UE::Audio::Insights
 {
-	class AUDIOINSIGHTS_API FMixerSourceDashboardViewFactory : public FTraceObjectTableDashboardViewFactory
+	class FMixerSourceDashboardViewFactory : public FTraceObjectTableDashboardViewFactory
 	{
 	public:
-		FMixerSourceDashboardViewFactory();
+		UE_API FMixerSourceDashboardViewFactory();
 		struct FPlotColumnInfo
 		{
 			const TFunctionRef<const ::Audio::TCircularAudioBuffer<FDataPoint>&(const IDashboardDataViewEntry& InData)> DataFunc;
 			const FNumberFormattingOptions* FormatOptions;
 		};
 
-		virtual ~FMixerSourceDashboardViewFactory();
+		UE_API virtual ~FMixerSourceDashboardViewFactory();
 
-		virtual FName GetName() const override;
-		virtual FText GetDisplayName() const override;
-		virtual FSlateIcon GetIcon() const override;
-		virtual EDefaultDashboardTabStack GetDefaultTabStack() const override;
-		virtual void ProcessEntries(FTraceTableDashboardViewFactory::EProcessReason Reason) override;
-		virtual TSharedRef<SWidget> MakeWidget(TSharedRef<SDockTab> OwnerTab, const FSpawnTabArgs& SpawnTabArgs) override;
+		UE_API virtual FName GetName() const override;
+		UE_API virtual FText GetDisplayName() const override;
+		UE_API virtual FSlateIcon GetIcon() const override;
+		UE_API virtual EDefaultDashboardTabStack GetDefaultTabStack() const override;
+		UE_API virtual void ProcessEntries(FTraceTableDashboardViewFactory::EProcessReason Reason) override;
+		UE_API virtual TSharedRef<SWidget> MakeWidget(TSharedRef<SDockTab> OwnerTab, const FSpawnTabArgs& SpawnTabArgs) override;
 
 		// Maximum amount of data history kept for plots (in seconds)
-		static const double MaxPlotHistorySeconds;
+		static UE_API const double MaxPlotHistorySeconds;
 		// Maximum number of sources to plot at once 
-		static const int32 MaxPlotSources; 
+		static UE_API const int32 MaxPlotSources; 
 
 #if WITH_EDITOR
 		DECLARE_MULTICAST_DELEGATE_ThreeParams(FOnUpdateMuteSoloState, ECheckBoxState /*MuteState*/, ECheckBoxState /*SoloState*/, const FString& /*CurrentFilterString*/);
-		inline static FOnUpdateMuteSoloState OnUpdateMuteSoloState;
+		UE_API inline static FOnUpdateMuteSoloState OnUpdateMuteSoloState;
 #endif // WITH_EDITOR
 
 	protected:
-		virtual const TMap<FName, FTraceTableDashboardViewFactory::FColumnData>& GetColumns() const override;
-		virtual void SortTable() override;
+		UE_API virtual const TMap<FName, FTraceTableDashboardViewFactory::FColumnData>& GetColumns() const override;
+		UE_API virtual void SortTable() override;
 
-		TSharedRef<SWidget> MakePlotsWidget();
+		UE_API TSharedRef<SWidget> MakePlotsWidget();
 
 	private:
 		using FPlotCurvePoint = SAudioCurveView::FCurvePoint;
 		using FPointDataPerCurveMap = TMap<int32, TArray<FPlotCurvePoint>>; // Map of source id to data point array 
 		using FPlotCurveMetadata = SAudioCurveView::FCurveMetadata;
 
-		void OnAnalysisStarting(const double Timestamp);
+		UE_API void OnAnalysisStarting(const double Timestamp);
 
 #if WITH_EDITOR
-		void OnPIEStarted(bool bSimulating);
-		void OnPIEStopped(bool bSimulating);
-		void OnPIEPaused(bool bSimulating);
-		void OnPIEResumed(bool bSimulating);
+		UE_API void OnPIEStarted(bool bSimulating);
+		UE_API void OnPIEStopped(bool bSimulating);
+		UE_API void OnPIEPaused(bool bSimulating);
+		UE_API void OnPIEResumed(bool bSimulating);
 #else
-		void OnAudioInsightsComponentTabSpawn();
-		void OnSessionAnalysisCompleted();
-		void OnTimingViewTimeMarkerChanged(double InTimeMarker);
+		UE_API void OnAudioInsightsComponentTabSpawn();
+		UE_API void OnSessionAnalysisCompleted();
+		UE_API void OnTimingViewTimeMarkerChanged(double InTimeMarker);
 #endif // WITH_EDITOR
 
-		void ResetPlots();
-		void UpdatePlotsWidgetsData();
+		UE_API void ResetPlots();
+		UE_API void UpdatePlotsWidgetsData();
 
 #if !WITH_EDITOR
-		void FilterOfflinePlots(const FString& InFilterString, TMap<int32, FPlotCurveMetadata>& InPlotWidgetMetadataPerCurve);
+		UE_API void FilterOfflinePlots(const FString& InFilterString, TMap<int32, FPlotCurveMetadata>& InPlotWidgetMetadataPerCurve);
 #endif // !WITH_EDITOR
 
 #if WITH_EDITOR
-		void UpdateMuteSoloState();
+		UE_API void UpdateMuteSoloState();
 #endif // WITH_EDITOR
 
 		// Column information used by plot widgets, keyed by column name. These keys should be a subset of the keys in GetColumns(). 
-		const TMap<FName, FPlotColumnInfo>& GetPlotColumnInfo();
-		const TFunctionRef<const ::Audio::TCircularAudioBuffer<FDataPoint>&(const IDashboardDataViewEntry& InData)> GetPlotColumnDataFunc(const FName& ColumnName);
-		const FNumberFormattingOptions* GetPlotColumnNumberFormat(const FName& ColumnName);
-		const FText GetPlotColumnDisplayName(const FName& ColumnName);
+		UE_API const TMap<FName, FPlotColumnInfo>& GetPlotColumnInfo();
+		UE_API const TFunctionRef<const ::Audio::TCircularAudioBuffer<FDataPoint>&(const IDashboardDataViewEntry& InData)> GetPlotColumnDataFunc(const FName& ColumnName);
+		UE_API const FNumberFormattingOptions* GetPlotColumnNumberFormat(const FName& ColumnName);
+		UE_API const FText GetPlotColumnDisplayName(const FName& ColumnName);
 
 #if WITH_EDITOR
-		TSharedRef<SWidget> MakeMuteSoloWidget();
+		UE_API TSharedRef<SWidget> MakeMuteSoloWidget();
 
-		void ToggleMuteForAllItems(ECheckBoxState NewState);
-		void ToggleSoloForAllItems(ECheckBoxState NewState);
+		UE_API void ToggleMuteForAllItems(ECheckBoxState NewState);
+		UE_API void ToggleSoloForAllItems(ECheckBoxState NewState);
 
-		TSharedRef<FTabManager::FLayout> LoadLayoutFromConfig();
-		void SaveLayoutToConfig();
+		UE_API TSharedRef<FTabManager::FLayout> LoadLayoutFromConfig();
+		UE_API void SaveLayoutToConfig();
 #endif // WITH EDITOR
 
-		TSharedRef<SWidget> MakeAmpDisplayModeWidget();
-		TSharedRef<SWidget> MakePlotsButtonWidget();
-		void TogglePlotsTabVisibility(ECheckBoxState InCheckboxState);
+		UE_API TSharedRef<SWidget> MakePlotsButtonWidget();
+		UE_API void TogglePlotsTabVisibility(ECheckBoxState InCheckboxState);
 
-		TSharedRef<SDockTab> CreateMixerSourcesTab(const FSpawnTabArgs& Args);
-		TSharedRef<SDockTab> CreatePlotsTab(const FSpawnTabArgs& Args);
+		UE_API TSharedRef<SDockTab> CreateMixerSourcesTab(const FSpawnTabArgs& Args);
+		UE_API TSharedRef<SDockTab> CreatePlotsTab(const FSpawnTabArgs& Args);
 
-		void RegisterTabSpawners();
-		void UnregisterTabSpawners();
+		UE_API void RegisterTabSpawners();
+		UE_API void UnregisterTabSpawners();
 
-		TSharedRef<FTabManager::FLayout> GetDefaultTabLayout();
+		UE_API TSharedRef<FTabManager::FLayout> GetDefaultTabLayout();
 
 		FCheckBoxStyle MuteToggleButtonStyle;
 		FCheckBoxStyle SoloToggleButtonStyle;
@@ -138,8 +139,6 @@ namespace UE::Audio::Insights
 		TSharedPtr<FTabManager> MixerSourcesTabManager;
 		TSharedPtr<FWorkspaceItem> MixerSourcesWorkspace;
 
-		bool bDisplayAmpPeakInDb = true;
-
 #if WITH_EDITOR
 		// State of the mute and solo buttons
 		ECheckBoxState MuteState = ECheckBoxState::Unchecked;
@@ -151,3 +150,5 @@ namespace UE::Audio::Insights
 #endif // WITH_EDITOR
 	};
 } // namespace UE::Audio::Insights
+
+#undef UE_API
diff --git a/Engine/Plugins/AudioInsights/Source/AudioInsights/Public/Views/SoundDashboardViewFactory.h b/Engine/Plugins/AudioInsights/Source/AudioInsights/Public/Views/SoundDashboardViewFactory.h
--- a/Engine/Plugins/AudioInsights/Source/AudioInsights/Public/Views/SoundDashboardViewFactory.h
+++ b/Engine/Plugins/AudioInsights/Source/AudioInsights/Public/Views/SoundDashboardViewFactory.h
@@ -5,6 +5,8 @@
 
 #include "Filters/GenericFilter.h"
 
+#define UE_API AUDIOINSIGHTS_API
+
 class FUICommandList;
 
 namespace UE::Audio::Insights
@@ -68,23 +70,23 @@ namespace UE::Audio::Insights
 	*	- Contains a weak handle to the original entry (OriginalDataEntry) which is updated from the trace provider
 	*	- Copies updated params to PinnedSectionEntry for display
 	*/
-	class AUDIOINSIGHTS_API FPinnedSoundEntryWrapper
+	class FPinnedSoundEntryWrapper
 	{
 	public:
 		FPinnedSoundEntryWrapper() = delete;
-		FPinnedSoundEntryWrapper(const TSharedPtr<IDashboardDataTreeViewEntry>& OriginalEntry);
+		UE_API FPinnedSoundEntryWrapper(const TSharedPtr<IDashboardDataTreeViewEntry>& OriginalEntry);
 
 		TSharedPtr<IDashboardDataTreeViewEntry> GetPinnedSectionEntry() const { return PinnedSectionEntry; }
 		TSharedPtr<IDashboardDataTreeViewEntry> GetOriginalDataEntry() const { return OriginalDataEntry.IsValid() ? OriginalDataEntry.Pin() : nullptr; }
 
-		TSharedPtr<FPinnedSoundEntryWrapper> AddChildEntry(const TSharedPtr<IDashboardDataTreeViewEntry> Child);
+		UE_API TSharedPtr<FPinnedSoundEntryWrapper> AddChildEntry(const TSharedPtr<IDashboardDataTreeViewEntry> Child);
 
-		void UpdateParams();
+		UE_API void UpdateParams();
 
-		void CleanUp();
-		void MarkToDelete();
+		UE_API void CleanUp();
+		UE_API void MarkToDelete();
 
-		bool EntryIsValid() const;
+		UE_API bool EntryIsValid() const;
 
 		TArray<TSharedPtr<FPinnedSoundEntryWrapper>> PinnedWrapperChildren;
 
@@ -93,96 +95,78 @@ namespace UE::Audio::Insights
 		TWeakPtr<IDashboardDataTreeViewEntry> OriginalDataEntry;
 	};
 
-	class AUDIOINSIGHTS_API FSoundDashboardViewFactory : public FTraceTreeDashboardViewFactory
+	class FSoundDashboardViewFactory : public FTraceTreeDashboardViewFactory
 	{
 	public:
-		FSoundDashboardViewFactory();
-		virtual ~FSoundDashboardViewFactory();
+		UE_API FSoundDashboardViewFactory();
+		UE_API virtual ~FSoundDashboardViewFactory();
 
-		virtual FName GetName() const override;
-		virtual FText GetDisplayName() const override;
-		virtual FSlateIcon GetIcon() const override;
-		virtual EDefaultDashboardTabStack GetDefaultTabStack() const override;
-		virtual TSharedRef<SWidget> MakeWidget(TSharedRef<SDockTab> OwnerTab, const FSpawnTabArgs& SpawnTabArgs) override;
-		virtual void ProcessEntries(FTraceTreeDashboardViewFactory::EProcessReason Reason) override;
+		UE_API virtual FName GetName() const override;
+		UE_API virtual FText GetDisplayName() const override;
+		UE_API virtual FSlateIcon GetIcon() const override;
+		UE_API virtual EDefaultDashboardTabStack GetDefaultTabStack() const override;
+		UE_API virtual TSharedRef<SWidget> MakeWidget(TSharedRef<SDockTab> OwnerTab, const FSpawnTabArgs& SpawnTabArgs) override;
+		UE_API virtual void ProcessEntries(FTraceTreeDashboardViewFactory::EProcessReason Reason) override;
 
 	private:
-		void BindCommands();
+		UE_API void BindCommands();
 
-		virtual TSharedPtr<SWidget> GetFilterBarWidget() override;
-		virtual TSharedPtr<SWidget> GetFilterBarButtonWidget() override;
+		UE_API virtual TSharedPtr<SWidget> GetFilterBarWidget() override;
+		UE_API virtual TSharedPtr<SWidget> GetFilterBarButtonWidget() override;
 
-		bool IsRootItem(const TSharedRef<IDashboardDataTreeViewEntry>& InEntry) const;
+		UE_API bool IsRootItem(const TSharedRef<IDashboardDataTreeViewEntry>& InEntry) const;
+		UE_API bool EntryCanHaveChildren(const TSharedRef<IDashboardDataTreeViewEntry>& InEntry) const;
+		UE_API bool IsDescendant(const TSharedPtr<IDashboardDataTreeViewEntry>& InEntry, const TSharedPtr<IDashboardDataTreeViewEntry>& InChildCandidate) const;
 
-		TSharedRef<SWidget> GenerateWidgetForRootColumn(const TSharedRef<FTraceTreeDashboardViewFactory::SRowWidget>& InRowWidget, const TSharedRef<IDashboardDataTreeViewEntry>& InRowData, const FName& InColumn, const FText& InValueText);
-		virtual TSharedRef<SWidget> GenerateWidgetForColumn(const TSharedRef<FTraceTreeDashboardViewFactory::SRowWidget>& InRowWidget, const TSharedRef<IDashboardDataTreeViewEntry>& InRowData, const FName& InColumn) override;
+		UE_API TSharedRef<SWidget> GenerateWidgetForRootColumn(const TSharedRef<FTraceTreeDashboardViewFactory::SRowWidget>& InRowWidget, const TSharedRef<IDashboardDataTreeViewEntry>& InRowData, const FName& InColumn, const FText& InValueText);
+		UE_API virtual TSharedRef<SWidget> GenerateWidgetForColumn(const TSharedRef<FTraceTreeDashboardViewFactory::SRowWidget>& InRowWidget, const TSharedRef<IDashboardDataTreeViewEntry>& InRowData, const FName& InColumn) override;
 		
-		virtual TSharedRef<ITableRow> OnGenerateRow(TSharedPtr<IDashboardDataTreeViewEntry> Item, const TSharedRef<STableViewBase>& OwnerTable) override;
+		UE_API virtual TSharedRef<ITableRow> OnGenerateRow(TSharedPtr<IDashboardDataTreeViewEntry> Item, const TSharedRef<STableViewBase>& OwnerTable) override;
 
-		virtual TSharedPtr<SWidget> OnConstructContextMenu() override;
-		virtual FReply OnDataRowKeyInput(const FGeometry& InGeometry, const FKeyEvent& InKeyEvent) const override;
+		UE_API virtual TSharedPtr<SWidget> OnConstructContextMenu() override;
+		UE_API virtual FReply OnDataRowKeyInput(const FGeometry& InGeometry, const FKeyEvent& InKeyEvent) const override;
 
-		virtual const TMap<FName, FTraceTreeDashboardViewFactory::FHeaderRowColumnData>& GetHeaderRowColumns() const override;
-		virtual const TMap<FName, FTraceTreeDashboardViewFactory::FColumnData>& GetColumns() const override;
+		UE_API virtual const TMap<FName, FTraceTreeDashboardViewFactory::FHeaderRowColumnData>& GetHeaderRowColumns() const override;
+		UE_API virtual const TMap<FName, FTraceTreeDashboardViewFactory::FColumnData>& GetColumns() const override;
 
-		virtual void SortTable() override;
+		UE_API virtual void SortTable() override;
 		virtual TArray<TSharedPtr<IDashboardDataTreeViewEntry>>& GetTreeItemsSource() override { return FullTree; }
-		virtual bool ResetTreeData() override;
-
-		virtual int32 GetDisplayedTreeDepth() const override { return static_cast<int32>(TreeViewingMode); }
-
-		enum class ETreeViewingOptions : uint8
-		{
-			FullTree = 0,
-			ActiveSounds,
-			FlatList
-		};
+		UE_API virtual bool ResetTreeData() override;
 
-		void ChangeTreeViewingMode(const ETreeViewingOptions SelectedMode);
-		void BuildViewOptionsMenuContent(FMenuBuilder& MenuBuilder);
+		UE_API void RecursiveSort(TArray<TSharedPtr<IDashboardDataTreeViewEntry>>& OutTree, TFunctionRef<bool(const FSoundDashboardEntry&, const FSoundDashboardEntry&)> Predicate);
+		UE_API void SortByPredicate(TFunctionRef<bool(const FSoundDashboardEntry&, const FSoundDashboardEntry&)> Predicate);
 
-		void FilterText();
-		void RecursiveSort(TArray<TSharedPtr<IDashboardDataTreeViewEntry>>& OutTree, TFunctionRef<bool(const FSoundDashboardEntry&, const FSoundDashboardEntry&)> Predicate);
-		void SortByPredicate(TFunctionRef<bool(const FSoundDashboardEntry&, const FSoundDashboardEntry&)> Predicate);
-
-		TSharedRef<SWidget> OnGetSettingsMenuContent();
-		TSharedRef<SWidget> MakeSettingsButtonWidget();
+		UE_API TSharedRef<SWidget> MakeShowRecentlyStoppedSoundsWidget();
 
 #if WITH_EDITOR
-		TSharedRef<SWidget> MakeMuteSoloWidget();
+		UE_API TSharedRef<SWidget> MakeMuteSoloWidget();
 
-		TSharedRef<SWidget> CreateMuteSoloButton(const TSharedRef<FTraceTreeDashboardViewFactory::SRowWidget>& InRowWidget, 
-			const TSharedRef<IDashboardDataTreeViewEntry>& InRowData,
+		UE_API TSharedRef<SWidget> CreateMuteSoloButton(const TSharedRef<IDashboardDataTreeViewEntry>& InRowData,
 			const FName& InColumn,
 			TFunction<void(const TArray<TSharedPtr<IDashboardDataTreeViewEntry>>& /*InEntries*/)> MuteSoloToggleFunc,
 			TFunctionRef<bool(const IDashboardDataTreeViewEntry& /*InEntry*/, const bool /*bInCheckChildren*/)> IsMuteSoloFunc);
 
-		void ToggleMuteSoloEntries(const TArray<TSharedPtr<IDashboardDataTreeViewEntry>>& InEntries, const EMuteSoloMode InMuteSoloMode);
+		UE_API void ToggleMuteSoloEntries(const TArray<TSharedPtr<IDashboardDataTreeViewEntry>>& InEntries, const EMuteSoloMode InMuteSoloMode);
 
-		TArray<TObjectPtr<UObject>> GetSelectedEditableAssets() const;
+		UE_API TArray<TObjectPtr<UObject>> GetSelectedEditableAssets() const;
 #endif // WITH_EDITOR
 
-		bool SelectedItemsIncludesAnAsset() const;
-		bool SelectionIncludesUnpinnedItem() const;
-		void ClearSelection();
-
-		void PinSound();
-		void UnpinSound();
+		UE_API bool SelectedItemsIncludesAnAsset() const;
+		UE_API void PinSound();
+		UE_API void UnpinSound();
+		UE_API bool SelectionIncludesUnpinnedItem() const;
 
-		void PinSelectedItems(const TSharedPtr<IDashboardDataTreeViewEntry>& Entry, const TArray<TSharedPtr<IDashboardDataTreeViewEntry>>& SelectedItems);
-		void UnpinSelectedItems(const TSharedPtr<FPinnedSoundEntryWrapper>& PinnedWrapperEntry, const TArray<TSharedPtr<IDashboardDataTreeViewEntry>>& SelectedItems, const bool bSelectionContainsAssets);
+		UE_API void PinSelectedItems(const TSharedPtr<IDashboardDataTreeViewEntry>& Entry, const TArray<TSharedPtr<IDashboardDataTreeViewEntry>>& SelectedItems);
+		UE_API void UnpinSelectedItems(const TSharedPtr<FPinnedSoundEntryWrapper>& PinnedWrapperEntry, const TArray<TSharedPtr<IDashboardDataTreeViewEntry>>& SelectedItems, const bool bSelectionContainsAssets);
 
-		void RecreatePinnedEntries(const TSharedPtr<IDashboardDataTreeViewEntry>& Entry);
-
-		void MarkBranchAsPinned(const TSharedPtr<IDashboardDataTreeViewEntry> Entry, const bool bIsPinned);
-		void InitPinnedItemEntries();
-		void CreatePinnedEntry(TSharedPtr<IDashboardDataTreeViewEntry> Entry);
-		void UpdatePinnedSection();
-		void RebuildPinnedSection();
+		UE_API void MarkBranchAsPinned(const TSharedPtr<IDashboardDataTreeViewEntry> Entry, const bool bIsPinned);
+		UE_API void InitPinnedItemEntries();
+		UE_API void CreatePinnedEntry(TSharedPtr<IDashboardDataTreeViewEntry> Entry);
+		UE_API void UpdatePinnedSection();
 
 #if WITH_EDITOR
-		void BrowseSoundAsset() const;
-		void OpenSoundAsset() const;
+		UE_API void BrowseSoundAsset() const;
+		UE_API void OpenSoundAsset() const;
 #endif
 		// @TODO UE-250399: Hide category pending to implement
 		//void HideSound();
@@ -204,11 +188,10 @@ namespace UE::Audio::Insights
 		TSharedPtr<SWidget> SoundsFilterBar;
 		TSharedPtr<SWidget> SoundsFilterBarButton;
 
-		ETreeViewingOptions TreeViewingMode = ETreeViewingOptions::FullTree;
 		ESoundDashboardFilterFlags SelectedFilterFlags = AllFilterFlags;
 		bool bIsPinnedCategoryFilterEnabled = true;
 		bool bShowRecentlyStoppedSounds = false;
-		bool bTreeViewModeChanged = false;
-		bool bDisplayAmpPeakInDb = true;
 	};
 } // namespace UE::Audio::Insights
+
+#undef UE_API
diff --git a/Engine/Plugins/AudioInsights/Source/AudioInsights/Public/Views/TableDashboardViewFactory.h b/Engine/Plugins/AudioInsights/Source/AudioInsights/Public/Views/TableDashboardViewFactory.h
--- a/Engine/Plugins/AudioInsights/Source/AudioInsights/Public/Views/TableDashboardViewFactory.h
+++ b/Engine/Plugins/AudioInsights/Source/AudioInsights/Public/Views/TableDashboardViewFactory.h
@@ -9,10 +9,12 @@
 #include "Widgets/Views/SHeaderRow.h"
 #include "Widgets/Views/SListView.h"
 
+#define UE_API AUDIOINSIGHTS_API
+
 
 namespace UE::Audio::Insights
 {
-	class AUDIOINSIGHTS_API IObjectDashboardEntry : public IDashboardDataViewEntry
+	class IObjectDashboardEntry : public IDashboardDataViewEntry
 	{
 	public:
 		virtual ~IObjectDashboardEntry() = default;
@@ -22,14 +24,14 @@ namespace UE::Audio::Insights
 		virtual UObject* GetObject() = 0;
 	};
 
-	struct AUDIOINSIGHTS_API FSoundAssetDashboardEntry : public IObjectDashboardEntry
+	struct FSoundAssetDashboardEntry : public IObjectDashboardEntry
 	{
 		virtual ~FSoundAssetDashboardEntry() = default;
 
-		virtual FText GetDisplayName() const override;
-		virtual const UObject* GetObject() const override;
-		virtual UObject* GetObject() override;
-		virtual bool IsValid() const override;
+		UE_API virtual FText GetDisplayName() const override;
+		UE_API virtual const UObject* GetObject() const override;
+		UE_API virtual UObject* GetObject() override;
+		UE_API virtual bool IsValid() const override;
 
 		::Audio::FDeviceId DeviceId = INDEX_NONE;
 		uint32 PlayOrder = INDEX_NONE;
@@ -38,11 +40,11 @@ namespace UE::Audio::Insights
 		FString Name;
 	};
 
-	class AUDIOINSIGHTS_API FTraceTableDashboardViewFactory : public FTraceDashboardViewFactoryBase, public TSharedFromThis<FTraceTableDashboardViewFactory>
+	class FTraceTableDashboardViewFactory : public FTraceDashboardViewFactoryBase, public TSharedFromThis<FTraceTableDashboardViewFactory>
 	{
 	public:
-		FTraceTableDashboardViewFactory();
-		virtual ~FTraceTableDashboardViewFactory();
+		UE_API FTraceTableDashboardViewFactory();
+		UE_API virtual ~FTraceTableDashboardViewFactory();
 
 	protected:
 		struct SRowWidget : public SMultiColumnTableRow<TSharedPtr<IDashboardDataViewEntry>>
@@ -58,10 +60,10 @@ namespace UE::Audio::Insights
 		};
 
 	public:
-		virtual TSharedRef<SWidget> GenerateWidgetForColumn(TSharedRef<IDashboardDataViewEntry> InRowData, const FName& Column);
-		virtual FReply OnDataRowKeyInput(const FGeometry& Geometry, const FKeyEvent& KeyEvent) const;
-		virtual TSharedRef<SWidget> MakeWidget(TSharedRef<SDockTab> OwnerTab, const FSpawnTabArgs& SpawnTabArgs) override;
-		virtual void RefreshFilteredEntriesListView();
+		UE_API virtual TSharedRef<SWidget> GenerateWidgetForColumn(TSharedRef<IDashboardDataViewEntry> InRowData, const FName& Column);
+		UE_API virtual FReply OnDataRowKeyInput(const FGeometry& Geometry, const FKeyEvent& KeyEvent) const;
+		UE_API virtual TSharedRef<SWidget> MakeWidget(TSharedRef<SDockTab> OwnerTab, const FSpawnTabArgs& SpawnTabArgs) override;
+		UE_API virtual void RefreshFilteredEntriesListView();
 
 		enum class EProcessReason : uint8
 		{
@@ -71,7 +73,7 @@ namespace UE::Audio::Insights
 		};
 
 	protected:
-		void Tick(float InElapsed);
+		UE_API void Tick(float InElapsed);
 
 		template<typename TableProviderType>
 		bool FilterEntries(TFunctionRef<bool(const IDashboardDataViewEntry&)> IsFiltered)
@@ -136,13 +138,13 @@ namespace UE::Audio::Insights
 		virtual const TMap<FName, FColumnData>& GetColumns() const = 0;
 		virtual void ProcessEntries(EProcessReason Reason) = 0;
 
-		const FText& GetSearchFilterText() const;
+		UE_API const FText& GetSearchFilterText() const;
 
 		virtual void SortTable() = 0;
 
-		virtual TSharedPtr<SWidget> OnConstructContextMenu();
-		virtual void OnSelectionChanged(TSharedPtr<IDashboardDataViewEntry> SelectedItem, ESelectInfo::Type SelectInfo);
-		virtual FSlateColor GetRowColor(const TSharedPtr<IDashboardDataViewEntry>& InRowDataPtr);
+		UE_API virtual TSharedPtr<SWidget> OnConstructContextMenu();
+		UE_API virtual void OnSelectionChanged(TSharedPtr<IDashboardDataViewEntry> SelectedItem, ESelectInfo::Type SelectInfo);
+		UE_API virtual FSlateColor GetRowColor(const TSharedPtr<IDashboardDataViewEntry>& InRowDataPtr);
 
 		EProcessReason UpdateFilterReason = EProcessReason::None;
 		FTSTicker::FDelegateHandle TickerHandle;
@@ -160,37 +162,39 @@ namespace UE::Audio::Insights
 		EColumnSortMode::Type SortMode = EColumnSortMode::None;
 
 	private:
-		TSharedRef<SHeaderRow> MakeHeaderRowWidget();
-		void SetSearchBoxFilterText(const FText& NewText);
+		UE_API TSharedRef<SHeaderRow> MakeHeaderRowWidget();
+		UE_API void SetSearchBoxFilterText(const FText& NewText);
 
-		EColumnSortMode::Type GetColumnSortMode(const FName InColumnId) const;
-		void RequestSort();
-		void OnColumnSortModeChanged(const EColumnSortPriority::Type InSortPriority, const FName& InColumnId, const EColumnSortMode::Type InSortMode);
+		UE_API EColumnSortMode::Type GetColumnSortMode(const FName InColumnId) const;
+		UE_API void RequestSort();
+		UE_API void OnColumnSortModeChanged(const EColumnSortPriority::Type InSortPriority, const FName& InColumnId, const EColumnSortMode::Type InSortMode);
 
 		TSharedPtr<SSearchBox> SearchBoxWidget;
 		FText SearchBoxFilterText;
 	};
 
-	class AUDIOINSIGHTS_API FTraceObjectTableDashboardViewFactory : public FTraceTableDashboardViewFactory
+	class FTraceObjectTableDashboardViewFactory : public FTraceTableDashboardViewFactory
 	{
 	public:
 		virtual ~FTraceObjectTableDashboardViewFactory() = default;
 
-		virtual TSharedRef<SWidget> GenerateWidgetForColumn(TSharedRef<IDashboardDataViewEntry> InRowData, const FName& Column) override;
-		virtual FReply OnDataRowKeyInput(const FGeometry& Geometry, const FKeyEvent& KeyEvent) const override;
-		virtual TSharedRef<SWidget> MakeWidget(TSharedRef<SDockTab> OwnerTab, const FSpawnTabArgs& SpawnTabArgs) override;
+		UE_API virtual TSharedRef<SWidget> GenerateWidgetForColumn(TSharedRef<IDashboardDataViewEntry> InRowData, const FName& Column) override;
+		UE_API virtual FReply OnDataRowKeyInput(const FGeometry& Geometry, const FKeyEvent& KeyEvent) const override;
+		UE_API virtual TSharedRef<SWidget> MakeWidget(TSharedRef<SDockTab> OwnerTab, const FSpawnTabArgs& SpawnTabArgs) override;
 
 	protected:
 #if WITH_EDITOR
-		virtual TSharedRef<SWidget> MakeAssetMenuBar() const;
+		UE_API virtual TSharedRef<SWidget> MakeAssetMenuBar() const;
 #endif // WITH_EDITOR
 
 	private:
 #if WITH_EDITOR
-		TArray<UObject*> GetSelectedEditableAssets() const;
+		UE_API TArray<UObject*> GetSelectedEditableAssets() const;
 
-		bool OpenAsset() const;
-		bool BrowseToAsset() const;
+		UE_API bool OpenAsset() const;
+		UE_API bool BrowseToAsset() const;
 #endif // WITH_EDITOR
 	};
 } // namespace UE::Audio::Insights
+
+#undef UE_API
diff --git a/Engine/Plugins/AudioInsights/Source/AudioInsights/Public/Views/TreeDashboardViewFactory.h b/Engine/Plugins/AudioInsights/Source/AudioInsights/Public/Views/TreeDashboardViewFactory.h
--- a/Engine/Plugins/AudioInsights/Source/AudioInsights/Public/Views/TreeDashboardViewFactory.h
+++ b/Engine/Plugins/AudioInsights/Source/AudioInsights/Public/Views/TreeDashboardViewFactory.h
@@ -8,6 +8,8 @@
 #include "Widgets/Views/STableRow.h"
 #include "Widgets/Views/STreeView.h"
 
+#define UE_API AUDIOINSIGHTS_API
+
 class SHeaderRow;
 class SSearchBox;
 class STableViewBase;
@@ -17,7 +19,7 @@ namespace UE::Audio::Insights
 	/**
 	 * Tree view entries can inherit from this class to implement extra UObject functionality (ex: open, browse, edit, etc.) 
 	 */
-	class AUDIOINSIGHTS_API IObjectTreeDashboardEntry : public IDashboardDataTreeViewEntry
+	class IObjectTreeDashboardEntry : public IDashboardDataTreeViewEntry
 	{
 	public:
 		virtual ~IObjectTreeDashboardEntry() = default;
@@ -33,11 +35,11 @@ namespace UE::Audio::Insights
 	 * It contains a search textbox, filters can be optionally be implemented via GetFilterBarWidget and GetFilterBarButtonWidget.
 	 * Item actions can be done via OnSelectionChanged, OnDataRowKeyInput, OnConstructContextMenu (for right click)
 	 */
-	class AUDIOINSIGHTS_API FTraceTreeDashboardViewFactory : public FTraceDashboardViewFactoryBase, public TSharedFromThis<FTraceTreeDashboardViewFactory>
+	class FTraceTreeDashboardViewFactory : public FTraceDashboardViewFactoryBase, public TSharedFromThis<FTraceTreeDashboardViewFactory>
 	{
 	public:
-		FTraceTreeDashboardViewFactory();
-		virtual ~FTraceTreeDashboardViewFactory();
+		UE_API FTraceTreeDashboardViewFactory();
+		UE_API virtual ~FTraceTreeDashboardViewFactory();
 
 		enum class EProcessReason : uint8
 		{
@@ -81,23 +83,23 @@ namespace UE::Audio::Insights
 			const TFunction<FSlateColor(const IDashboardDataTreeViewEntry&)> GetTextColorValue;
 		};
 
-		virtual TSharedPtr<SWidget> GetFilterBarWidget();
-		virtual TSharedPtr<SWidget> GetFilterBarButtonWidget();
+		UE_API virtual TSharedPtr<SWidget> GetFilterBarWidget();
+		UE_API virtual TSharedPtr<SWidget> GetFilterBarButtonWidget();
 
-		virtual TSharedRef<SWidget> MakeWidget(TSharedRef<SDockTab> OwnerTab, const FSpawnTabArgs& SpawnTabArgs);
+		UE_API virtual TSharedRef<SWidget> MakeWidget(TSharedRef<SDockTab> OwnerTab, const FSpawnTabArgs& SpawnTabArgs);
 
 		virtual TSharedRef<SWidget> GenerateWidgetForColumn(const TSharedRef<FTraceTreeDashboardViewFactory::SRowWidget>& InRowWidget, const TSharedRef<IDashboardDataTreeViewEntry>& InRowData, const FName& Column) = 0;
 		
-		virtual TSharedRef<ITableRow> OnGenerateRow(TSharedPtr<IDashboardDataTreeViewEntry> Item, const TSharedRef<STableViewBase>& OwnerTable);
+		UE_API virtual TSharedRef<ITableRow> OnGenerateRow(TSharedPtr<IDashboardDataTreeViewEntry> Item, const TSharedRef<STableViewBase>& OwnerTable);
 
-		virtual TSharedPtr<SWidget> OnConstructContextMenu();
+		UE_API virtual TSharedPtr<SWidget> OnConstructContextMenu();
 
-		virtual void OnSelectionChanged(TSharedPtr<IDashboardDataTreeViewEntry> SelectedItem, ESelectInfo::Type SelectInfo);
-		virtual FReply OnDataRowKeyInput(const FGeometry& Geometry, const FKeyEvent& KeyEvent) const;
+		UE_API virtual void OnSelectionChanged(TSharedPtr<IDashboardDataTreeViewEntry> SelectedItem, ESelectInfo::Type SelectInfo);
+		UE_API virtual FReply OnDataRowKeyInput(const FGeometry& Geometry, const FKeyEvent& KeyEvent) const;
 
-		const FText& GetSearchFilterText() const;
+		UE_API const FText& GetSearchFilterText() const;
 
-		virtual void RefreshFilteredEntriesListView();
+		UE_API virtual void RefreshFilteredEntriesListView();
 
 		virtual const TMap<FName, FHeaderRowColumnData>& GetHeaderRowColumns() const = 0;
 		virtual const TMap<FName, FColumnData>& GetColumns() const = 0;
@@ -120,7 +122,7 @@ namespace UE::Audio::Insights
 			return false;
 		}
 
-		void Tick(float InElapsed);
+		UE_API void Tick(float InElapsed);
 
 #if WITH_EDITOR
 		virtual bool IsDebugDrawEnabled() const { return false; }
@@ -189,13 +191,13 @@ namespace UE::Audio::Insights
 		EColumnSortMode::Type SortMode = EColumnSortMode::None;
 
 	private:
-		TSharedRef<SHeaderRow> MakeHeaderRowWidget();
-		void SetSearchBoxFilterText(const FText& NewText);
+		UE_API TSharedRef<SHeaderRow> MakeHeaderRowWidget();
+		UE_API void SetSearchBoxFilterText(const FText& NewText);
 
-		EColumnSortMode::Type GetColumnSortMode(const FName InColumnId) const;
-		void RequestSort();
-		void OnColumnSortModeChanged(const EColumnSortPriority::Type InSortPriority, const FName& InColumnId, const EColumnSortMode::Type InSortMode);
-		void HandleRecursiveExpansion(TSharedPtr<IDashboardDataTreeViewEntry> Item, bool bIsItemExpanded);
+		UE_API EColumnSortMode::Type GetColumnSortMode(const FName InColumnId) const;
+		UE_API void RequestSort();
+		UE_API void OnColumnSortModeChanged(const EColumnSortPriority::Type InSortPriority, const FName& InColumnId, const EColumnSortMode::Type InSortMode);
+		UE_API void HandleRecursiveExpansion(TSharedPtr<IDashboardDataTreeViewEntry> Item, bool bIsItemExpanded);
 
 		void TransformDataRecursive(const TSharedPtr<IDashboardDataTreeViewEntry>& OriginalEntry, TFunctionRef<bool(IDashboardDataTreeViewEntry&)> FilterPredicate, const int32 DepthToDisplayFrom, const bool bParentPassedFilter = true);
 
@@ -203,3 +205,5 @@ namespace UE::Audio::Insights
 		FText SearchBoxFilterText;
 	};
 } // namespace UE::Audio::Insights
+
+#undef UE_API
diff --git a/Engine/Plugins/AudioInsights/Source/AudioInsights/Public/Views/VirtualLoopDashboardViewFactory.h b/Engine/Plugins/AudioInsights/Source/AudioInsights/Public/Views/VirtualLoopDashboardViewFactory.h
--- a/Engine/Plugins/AudioInsights/Source/AudioInsights/Public/Views/VirtualLoopDashboardViewFactory.h
+++ b/Engine/Plugins/AudioInsights/Source/AudioInsights/Public/Views/VirtualLoopDashboardViewFactory.h
@@ -3,34 +3,38 @@
 
 #include "Views/TableDashboardViewFactory.h"
 
+#define UE_API AUDIOINSIGHTS_API
+
 
 namespace UE::Audio::Insights
 {
-	class AUDIOINSIGHTS_API FVirtualLoopDashboardViewFactory : public FTraceObjectTableDashboardViewFactory
+	class FVirtualLoopDashboardViewFactory : public FTraceObjectTableDashboardViewFactory
 	{
 	public:
-		FVirtualLoopDashboardViewFactory();
+		UE_API FVirtualLoopDashboardViewFactory();
 		virtual ~FVirtualLoopDashboardViewFactory() = default;
 
-		virtual FName GetName() const override;
-		virtual FText GetDisplayName() const override;
-		virtual FSlateIcon GetIcon() const override;
-		virtual EDefaultDashboardTabStack GetDefaultTabStack() const override;
+		UE_API virtual FName GetName() const override;
+		UE_API virtual FText GetDisplayName() const override;
+		UE_API virtual FSlateIcon GetIcon() const override;
+		UE_API virtual EDefaultDashboardTabStack GetDefaultTabStack() const override;
 
 #if WITH_EDITOR
 		DECLARE_MULTICAST_DELEGATE_ThreeParams(FOnDebugDrawEntries, float /*InElapsed*/, const TArray<TSharedPtr<IDashboardDataViewEntry>>& /*InSelectedItems*/, ::Audio::FDeviceId /*InAudioDeviceId*/);
-		inline static FOnDebugDrawEntries OnDebugDrawEntries;
+		UE_API inline static FOnDebugDrawEntries OnDebugDrawEntries;
 #endif // WITH_EDITOR
 
 	protected:
-		virtual void ProcessEntries(FTraceTableDashboardViewFactory::EProcessReason Reason) override;
-		virtual const TMap<FName, FTraceTableDashboardViewFactory::FColumnData>& GetColumns() const override;
+		UE_API virtual void ProcessEntries(FTraceTableDashboardViewFactory::EProcessReason Reason) override;
+		UE_API virtual const TMap<FName, FTraceTableDashboardViewFactory::FColumnData>& GetColumns() const override;
 
-		virtual void SortTable() override;
+		UE_API virtual void SortTable() override;
 
 #if WITH_EDITOR
-		virtual bool IsDebugDrawEnabled() const override;
-		virtual void DebugDraw(float InElapsed, const TArray<TSharedPtr<IDashboardDataViewEntry>>& InSelectedItems, ::Audio::FDeviceId InAudioDeviceId) const;
+		UE_API virtual bool IsDebugDrawEnabled() const override;
+		UE_API virtual void DebugDraw(float InElapsed, const TArray<TSharedPtr<IDashboardDataViewEntry>>& InSelectedItems, ::Audio::FDeviceId InAudioDeviceId) const;
 #endif // WITH_EDITOR
 	};
 } // namespace UE::Audio::Insights
+
+#undef UE_API
diff --git a/Engine/Plugins/AudioInsights/Source/AudioInsightsEditor/Public/AudioInsightsEditorSettings.h b/Engine/Plugins/AudioInsights/Source/AudioInsightsEditor/Public/AudioInsightsEditorSettings.h
--- a/Engine/Plugins/AudioInsights/Source/AudioInsightsEditor/Public/AudioInsightsEditorSettings.h
+++ b/Engine/Plugins/AudioInsights/Source/AudioInsightsEditor/Public/AudioInsightsEditorSettings.h
@@ -8,15 +8,17 @@
 
 #include "AudioInsightsEditorSettings.generated.h"
 
-UCLASS(config = EditorPerProjectUserSettings)
-class AUDIOINSIGHTSEDITOR_API UAudioInsightsEditorSettings : public UDeveloperSettings
+#define UE_API AUDIOINSIGHTSEDITOR_API
+
+UCLASS(MinimalAPI, config = EditorPerProjectUserSettings)
+class UAudioInsightsEditorSettings : public UDeveloperSettings
 {
 	GENERATED_BODY()
 	
 public:
-	virtual FName GetCategoryName() const override;
-	virtual FText GetSectionText() const override;
-	virtual FText GetSectionDescription() const override;
+	UE_API virtual FName GetCategoryName() const override;
+	UE_API virtual FText GetSectionText() const override;
+	UE_API virtual FText GetSectionDescription() const override;
 
 	/** Whether to automatically set the first PIE client in Audio Insights World Filter. */
 	UPROPERTY(Config, EditAnywhere, Category="World Filter")
@@ -30,3 +32,5 @@ public:
 	UPROPERTY(EditAnywhere, config, Category = SpectrumAnalyzer, meta = (ShowOnlyInnerProperties))
 	FSpectrumAnalyzerRackUnitSettings SpectrumAnalyzerSettings;
 };
+
+#undef UE_API
diff --git a/Engine/Plugins/AudioInsights/Source/AudioInsightsEditor/Public/IAudioInsightsEditorModule.h b/Engine/Plugins/AudioInsights/Source/AudioInsightsEditor/Public/IAudioInsightsEditorModule.h
--- a/Engine/Plugins/AudioInsights/Source/AudioInsightsEditor/Public/IAudioInsightsEditorModule.h
+++ b/Engine/Plugins/AudioInsights/Source/AudioInsightsEditor/Public/IAudioInsightsEditorModule.h
@@ -21,14 +21,14 @@ namespace UE::Audio::Insights
 class IAudioInsightsEditorModule : public IAudioInsightsModuleInterface
 {
 public:
-	AUDIOINSIGHTSEDITOR_API virtual void RegisterDashboardViewFactory(TSharedRef<UE::Audio::Insights::IDashboardViewFactory> InDashboardFactory) = 0;
-	AUDIOINSIGHTSEDITOR_API virtual void UnregisterDashboardViewFactory(FName InName) = 0;
+	virtual void RegisterDashboardViewFactory(TSharedRef<UE::Audio::Insights::IDashboardViewFactory> InDashboardFactory) = 0;
+	virtual void UnregisterDashboardViewFactory(FName InName) = 0;
 
-	AUDIOINSIGHTSEDITOR_API virtual ::Audio::FDeviceId GetDeviceId() const = 0;
+	virtual ::Audio::FDeviceId GetDeviceId() const = 0;
 
 	AUDIOINSIGHTSEDITOR_API virtual IAudioInsightsTraceModule& GetTraceModule() override;
 
-	AUDIOINSIGHTSEDITOR_API virtual TSharedRef<SDockTab> CreateDashboardTabWidget(const FSpawnTabArgs& Args) = 0;
+	virtual TSharedRef<SDockTab> CreateDashboardTabWidget(const FSpawnTabArgs& Args) = 0;
 
 	AUDIOINSIGHTSEDITOR_API static bool IsModuleLoaded();
 
diff --git a/Engine/Plugins/ChaosClothAssetEditor/Source/ChaosClothAssetEditor/Private/ChaosClothAsset/ClothEditor.h b/Engine/Plugins/ChaosClothAssetEditor/Source/ChaosClothAssetEditor/Private/ChaosClothAsset/ClothEditor.h
--- a/Engine/Plugins/ChaosClothAssetEditor/Source/ChaosClothAssetEditor/Private/ChaosClothAsset/ClothEditor.h
+++ b/Engine/Plugins/ChaosClothAssetEditor/Source/ChaosClothAssetEditor/Private/ChaosClothAsset/ClothEditor.h
@@ -5,6 +5,8 @@
 #include "BaseCharacterFXEditor.h"
 #include "ClothEditor.generated.h"
 
+#define UE_API CHAOSCLOTHASSETEDITOR_API
+
 /** 
  * The actual asset editor class doesn't have that much in it, intentionally. 
  * 
@@ -18,14 +20,16 @@
  * on the left.
  */
 
-UCLASS()
-class CHAOSCLOTHASSETEDITOR_API UChaosClothAssetEditor : public UBaseCharacterFXEditor
+UCLASS(MinimalAPI)
+class UChaosClothAssetEditor : public UBaseCharacterFXEditor
 {
 	GENERATED_BODY()
 
 public:
 
-	virtual TSharedPtr<FBaseAssetToolkit> CreateToolkit() override;
+	UE_API virtual TSharedPtr<FBaseAssetToolkit> CreateToolkit() override;
 };
 
-DECLARE_LOG_CATEGORY_EXTERN(LogChaosClothAssetEditor, Log, All);
\ No newline at end of file
+DECLARE_LOG_CATEGORY_EXTERN(LogChaosClothAssetEditor, Log, All);
+
+#undef UE_API
diff --git a/Engine/Plugins/ChaosClothAssetEditor/Source/ChaosClothAssetEditor/Private/ChaosClothAsset/ClothEditor3DViewportClient.h b/Engine/Plugins/ChaosClothAssetEditor/Source/ChaosClothAssetEditor/Private/ChaosClothAsset/ClothEditor3DViewportClient.h
--- a/Engine/Plugins/ChaosClothAssetEditor/Source/ChaosClothAssetEditor/Private/ChaosClothAsset/ClothEditor3DViewportClient.h
+++ b/Engine/Plugins/ChaosClothAssetEditor/Source/ChaosClothAssetEditor/Private/ChaosClothAsset/ClothEditor3DViewportClient.h
@@ -9,6 +9,8 @@
 #include "IPreviewLODController.h"
 #include "BaseBehaviors/BehaviorTargetInterfaces.h"
 
+#define UE_API CHAOSCLOTHASSETEDITOR_API
+
 class UChaosClothComponent;
 class UChaosClothAssetEditorMode;
 class UTransformProxy;
@@ -26,7 +28,7 @@ class FClothEditorSimulationVisualization;
  * Viewport client for the 3d sim preview in the cloth editor. Currently same as editor viewport
  * client but doesn't allow editor gizmos/widgets.
  */
-class CHAOSCLOTHASSETEDITOR_API FChaosClothAssetEditor3DViewportClient : public FEditorViewportClient, 
+class FChaosClothAssetEditor3DViewportClient : public FEditorViewportClient, 
 	public TSharedFromThis<FChaosClothAssetEditor3DViewportClient>, 
 	public IClickBehaviorTarget,
 	public IClickDragBehaviorTarget,
@@ -35,53 +37,53 @@ class CHAOSCLOTHASSETEDITOR_API FChaosClothAssetEditor3DViewportClient : public
 {
 public:
 
-	FChaosClothAssetEditor3DViewportClient(FEditorModeTools* InModeTools, TSharedPtr<FChaosClothPreviewScene> InPreviewScene, 
+	UE_API FChaosClothAssetEditor3DViewportClient(FEditorModeTools* InModeTools, TSharedPtr<FChaosClothPreviewScene> InPreviewScene, 
 		TSharedPtr<FClothEditorSimulationVisualization> InVisualization,
 		const TWeakPtr<SEditorViewport>& InEditorViewportWidget = nullptr);
 
 	// Call this after construction to initialize callbacks when settings change
-	void RegisterDelegates();
+	UE_API void RegisterDelegates();
 
-	virtual ~FChaosClothAssetEditor3DViewportClient();
+	UE_API virtual ~FChaosClothAssetEditor3DViewportClient();
 
 	// Delete the viewport gizmo and transform proxy
-	void DeleteViewportGizmo();
+	UE_API void DeleteViewportGizmo();
 
-	void ClearSelectedComponents();
+	UE_API void ClearSelectedComponents();
 
 	void EnableSimMeshWireframe(bool bEnable) { bSimMeshWireframe = bEnable; }
 	bool SimMeshWireframeEnabled() const { return bSimMeshWireframe; }
 
-	void EnableRenderMeshWireframe(bool bEnable);
+	UE_API void EnableRenderMeshWireframe(bool bEnable);
 	bool RenderMeshWireframeEnabled() const { return bRenderMeshWireframe; }
 
-	void SetClothEdMode(TObjectPtr<UChaosClothAssetEditorMode> ClothEdMode);
-	void SetClothEditorToolkit(TWeakPtr<const FChaosClothAssetEditorToolkit> ClothToolkit);
+	UE_API void SetClothEdMode(TObjectPtr<UChaosClothAssetEditorMode> ClothEdMode);
+	UE_API void SetClothEditorToolkit(TWeakPtr<const FChaosClothAssetEditorToolkit> ClothToolkit);
 
-	void SoftResetSimulation();
-	void HardResetSimulation();
-	void SuspendSimulation();
-	void ResumeSimulation();
-	bool IsSimulationSuspended() const;
-	void SetEnableSimulation(bool bEnable);
-	bool IsSimulationEnabled() const;
+	UE_API void SoftResetSimulation();
+	UE_API void HardResetSimulation();
+	UE_API void SuspendSimulation();
+	UE_API void ResumeSimulation();
+	UE_API bool IsSimulationSuspended() const;
+	UE_API void SetEnableSimulation(bool bEnable);
+	UE_API bool IsSimulationEnabled() const;
 
 	// IPreviewLODController interface
 	// LODIndex == INDEX_NONE is LOD Auto
-	virtual void SetLODLevel(int32 LODIndex) override;
-	virtual int32 GetLODCount() const override;
-	virtual int32 GetCurrentLOD() const override;
-	virtual bool IsLODSelected(int32 LODIndex) const override;
+	UE_API virtual void SetLODLevel(int32 LODIndex) override;
+	UE_API virtual int32 GetLODCount() const override;
+	UE_API virtual int32 GetCurrentLOD() const override;
+	UE_API virtual bool IsLODSelected(int32 LODIndex) const override;
 	virtual int32 GetAutoLODStartingIndex() const override { return 1; }
-	virtual void FillLODCommands(TArray<TSharedPtr<FUICommandInfo>>& Commands) override;
+	UE_API virtual void FillLODCommands(TArray<TSharedPtr<FUICommandInfo>>& Commands) override;
 	// ~IPreviewLODController
 
-	FBox PreviewBoundingBox() const;
+	UE_API FBox PreviewBoundingBox() const;
 
-	TWeakPtr<FChaosClothPreviewScene> GetClothPreviewScene();
-	TWeakPtr<const FChaosClothPreviewScene> GetClothPreviewScene() const;
-	UChaosClothComponent* GetPreviewClothComponent();
-	const UChaosClothComponent* GetPreviewClothComponent() const;
+	UE_API TWeakPtr<FChaosClothPreviewScene> GetClothPreviewScene();
+	UE_API TWeakPtr<const FChaosClothPreviewScene> GetClothPreviewScene() const;
+	UE_API UChaosClothComponent* GetPreviewClothComponent();
+	UE_API const UChaosClothComponent* GetPreviewClothComponent() const;
 	TWeakPtr<FClothEditorSimulationVisualization> GetSimulationVisualization() {
 		return ClothEditorSimulationVisualization;
 	}
@@ -90,34 +92,34 @@ public:
 private:
 
 	// FGCObject override
-	virtual void AddReferencedObjects(FReferenceCollector& Collector) override;
+	UE_API virtual void AddReferencedObjects(FReferenceCollector& Collector) override;
 
 	// FEditorViewportClient overrides
 	virtual bool CanSetWidgetMode(UE::Widget::EWidgetMode NewMode) const override { return false; }
 	virtual void SetWidgetMode(UE::Widget::EWidgetMode NewMode) override {}
 	virtual UE::Widget::EWidgetMode GetWidgetMode() const override { return UE::Widget::EWidgetMode::WM_None; }
-	virtual void ProcessClick(FSceneView& View, HHitProxy* HitProxy, FKey Key, EInputEvent Event, uint32 HitX, uint32 HitY) override;
-	virtual void Draw(const FSceneView* View, FPrimitiveDrawInterface* PDI) override;
-	virtual void DrawCanvas(FViewport& InViewport, FSceneView& View, FCanvas& Canvas) override;
+	UE_API virtual void ProcessClick(FSceneView& View, HHitProxy* HitProxy, FKey Key, EInputEvent Event, uint32 HitX, uint32 HitY) override;
+	UE_API virtual void Draw(const FSceneView* View, FPrimitiveDrawInterface* PDI) override;
+	UE_API virtual void DrawCanvas(FViewport& InViewport, FSceneView& View, FCanvas& Canvas) override;
 
 	// IClickBehaviorTarget
-	virtual FInputRayHit IsHitByClick(const FInputDeviceRay& ClickPos) override;
-	virtual void OnClicked(const FInputDeviceRay& ClickPos) override;
+	UE_API virtual FInputRayHit IsHitByClick(const FInputDeviceRay& ClickPos) override;
+	UE_API virtual void OnClicked(const FInputDeviceRay& ClickPos) override;
 
 	// IClickDragBehaviorTarget
-	virtual FInputRayHit CanBeginClickDragSequence(const FInputDeviceRay& PressPos) override;
+	UE_API virtual FInputRayHit CanBeginClickDragSequence(const FInputDeviceRay& PressPos) override;
 	virtual void OnClickPress(const FInputDeviceRay& PressPos) override {}
 	virtual void OnClickDrag(const FInputDeviceRay& DragPos) override {}
 	virtual void OnClickRelease(const FInputDeviceRay& ReleasePos) override {}
 	virtual void OnTerminateDragSequence() override {}
 
 	// IInputBehaviorSource
-	virtual const UInputBehaviorSet* GetInputBehaviors() const override;
+	UE_API virtual const UInputBehaviorSet* GetInputBehaviors() const override;
 
-	void ComponentSelectionChanged(UObject* NewSelection);
+	UE_API void ComponentSelectionChanged(UObject* NewSelection);
 
 	// Update the selected components based on hitproxy
-	void UpdateSelection(HHitProxy* HitProxy);
+	UE_API void UpdateSelection(HHitProxy* HitProxy);
 
 	TWeakPtr<FChaosClothPreviewScene> ClothPreviewScene;
 
@@ -141,3 +143,5 @@ private:
 	TObjectPtr<UInputBehaviorSet> InputBehaviorSet;
 };
 } // namespace UE::Chaos::ClothAsset
+
+#undef UE_API
diff --git a/Engine/Plugins/ChaosClothAssetEditor/Source/ChaosClothAssetEditor/Private/ChaosClothAsset/ClothEditorMode.h b/Engine/Plugins/ChaosClothAssetEditor/Source/ChaosClothAssetEditor/Private/ChaosClothAsset/ClothEditorMode.h
--- a/Engine/Plugins/ChaosClothAssetEditor/Source/ChaosClothAssetEditor/Private/ChaosClothAsset/ClothEditorMode.h
+++ b/Engine/Plugins/ChaosClothAssetEditor/Source/ChaosClothAssetEditor/Private/ChaosClothAsset/ClothEditorMode.h
@@ -9,6 +9,8 @@
 #include "ChaosClothAsset/ClothPatternVertexType.h"
 #include "ClothEditorMode.generated.h"
 
+#define UE_API CHAOSCLOTHASSETEDITOR_API
+
 PREDECLARE_GEOMETRY(class FDynamicMesh3);
 
 class FEditorViewportClient;
@@ -46,86 +48,86 @@ class UPreviewGeometry;
  * We put things in a mode instead of directly into the asset editor in case we want to someday use the mode
  * in multiple asset editors.
  */
-UCLASS(Transient)
-class CHAOSCLOTHASSETEDITOR_API UChaosClothAssetEditorMode final : public UBaseCharacterFXEditorMode
+UCLASS(MinimalAPI, Transient)
+class UChaosClothAssetEditorMode final : public UBaseCharacterFXEditorMode
 {
 	GENERATED_BODY()
 
 public:
 
-	const static FEditorModeID EM_ChaosClothAssetEditorModeId;
+	UE_API const static FEditorModeID EM_ChaosClothAssetEditorModeId;
 
-	UChaosClothAssetEditorMode();
+	UE_API UChaosClothAssetEditorMode();
 
 	// Bounding box for selected rest space mesh components
-	FBox SelectionBoundingBox() const;
+	UE_API FBox SelectionBoundingBox() const;
 
 	// Bounding box for sim space meshes
-	FBox PreviewBoundingBox() const;
+	UE_API FBox PreviewBoundingBox() const;
 
-	void SetConstructionViewMode(UE::Chaos::ClothAsset::EClothPatternVertexType InMode);
-	UE::Chaos::ClothAsset::EClothPatternVertexType GetConstructionViewMode() const;
-	bool CanChangeConstructionViewModeTo(UE::Chaos::ClothAsset::EClothPatternVertexType NewViewMode) const;
+	UE_API void SetConstructionViewMode(UE::Chaos::ClothAsset::EClothPatternVertexType InMode);
+	UE_API UE::Chaos::ClothAsset::EClothPatternVertexType GetConstructionViewMode() const;
+	UE_API bool CanChangeConstructionViewModeTo(UE::Chaos::ClothAsset::EClothPatternVertexType NewViewMode) const;
 
-	void ToggleConstructionViewWireframe();
-	bool CanSetConstructionViewWireframeActive() const;
+	UE_API void ToggleConstructionViewWireframe();
+	UE_API bool CanSetConstructionViewWireframeActive() const;
 	bool IsConstructionViewWireframeActive() const
 	{
 		return bConstructionViewWireframe;
 	}
 
-	void ToggleConstructionViewSeams();
-	bool CanSetConstructionViewSeamsActive() const;
+	UE_API void ToggleConstructionViewSeams();
+	UE_API bool CanSetConstructionViewSeamsActive() const;
 	bool IsConstructionViewSeamsActive() const
 	{
 		return bConstructionViewSeamsVisible;
 	}
 
-	void ToggleConstructionViewSeamsCollapse();
-	bool CanSetConstructionViewSeamsCollapse() const;
+	UE_API void ToggleConstructionViewSeamsCollapse();
+	UE_API bool CanSetConstructionViewSeamsCollapse() const;
 	bool IsConstructionViewSeamsCollapseActive() const
 	{
 		return bConstructionViewSeamsCollapse;
 	}
 
-	void TogglePatternColor();
-	bool CanSetPatternColor() const;
+	UE_API void TogglePatternColor();
+	UE_API bool CanSetPatternColor() const;
 	bool IsPatternColorActive() const
 	{
 		return bPatternColors;
 	}
 
-	void ToggleConstructionViewSurfaceNormals();
-	bool CanSetConstructionViewSurfaceNormalsActive() const;
+	UE_API void ToggleConstructionViewSurfaceNormals();
+	UE_API bool CanSetConstructionViewSurfaceNormalsActive() const;
 	bool IsConstructionViewSurfaceNormalsActive() const
 	{
 		return bConstructionViewNormalsVisible;
 	}
 
-	void ToggleMeshStats();
-	bool CanSetMeshStats() const;
+	UE_API void ToggleMeshStats();
+	UE_API bool CanSetMeshStats() const;
 	bool IsMeshStatsActive() const
 	{
 		return bMeshStats;
 	}
 
 	// Simulation controls
-	void SoftResetSimulation();
-	void HardResetSimulation();
-	void SuspendSimulation();
-	void ResumeSimulation();
-	bool IsSimulationSuspended() const;
-	void SetEnableSimulation(bool bEnabled);
-	bool IsSimulationEnabled() const;
+	UE_API void SoftResetSimulation();
+	UE_API void HardResetSimulation();
+	UE_API void SuspendSimulation();
+	UE_API void ResumeSimulation();
+	UE_API bool IsSimulationSuspended() const;
+	UE_API void SetEnableSimulation(bool bEnabled);
+	UE_API bool IsSimulationEnabled() const;
 
-	int32 GetConstructionViewTriangleCount() const;
-	int32 GetConstructionViewVertexCount() const;
+	UE_API int32 GetConstructionViewTriangleCount() const;
+	UE_API int32 GetConstructionViewVertexCount() const;
 
 	// LODIndex == INDEX_NONE is LOD Auto
-	void SetLODModel(int32 LODIndex);
-	bool IsLODModelSelected(int32 LODIndex) const;
-	int32 GetLODModel() const;
-	int32 GetNumLODs() const;
+	UE_API void SetLODModel(int32 LODIndex);
+	UE_API bool IsLODModelSelected(int32 LODIndex) const;
+	UE_API int32 GetLODModel() const;
+	UE_API int32 GetNumLODs() const;
 
 	TObjectPtr<UEditorInteractiveToolsContext> GetActiveToolsContext()
 	{
@@ -138,70 +140,70 @@ private:
 	friend class UE::Chaos::ClothAsset::FChaosClothAssetEditorModeToolkit;
 
 	// UEdMode
-	virtual void Enter() final;
-	virtual void Exit() override;
-	virtual void ModeTick(float DeltaTime) override;
-	virtual bool ShouldToolStartBeAllowed(const FString& ToolIdentifier) const override;
-	virtual void OnToolStarted(UInteractiveToolManager* Manager, UInteractiveTool* Tool) override;
-	virtual void OnToolEnded(UInteractiveToolManager* Manager, UInteractiveTool* Tool) override;
-	virtual void CreateToolkit() override;
-	virtual void BindCommands() override;
+	UE_API virtual void Enter() final;
+	UE_API virtual void Exit() override;
+	UE_API virtual void ModeTick(float DeltaTime) override;
+	UE_API virtual bool ShouldToolStartBeAllowed(const FString& ToolIdentifier) const override;
+	UE_API virtual void OnToolStarted(UInteractiveToolManager* Manager, UInteractiveTool* Tool) override;
+	UE_API virtual void OnToolEnded(UInteractiveToolManager* Manager, UInteractiveTool* Tool) override;
+	UE_API virtual void CreateToolkit() override;
+	UE_API virtual void BindCommands() override;
 
 	// (We don't actually override MouseEnter, etc, because things get forwarded to the input
 	// router via FEditorModeTools, and we don't have any additional input handling to do at the mode level.)
 
 	// UBaseCharacterFXEditorMode
-	virtual void AddToolTargetFactories() override;
-	virtual void RegisterTools() override;
-	virtual void CreateToolTargets(const TArray<TObjectPtr<UObject>>& AssetsIn) override;
-	virtual void InitializeTargets(const TArray<TObjectPtr<UObject>>& AssetsIn) override;
+	UE_API virtual void AddToolTargetFactories() override;
+	UE_API virtual void RegisterTools() override;
+	UE_API virtual void CreateToolTargets(const TArray<TObjectPtr<UObject>>& AssetsIn) override;
+	UE_API virtual void InitializeTargets(const TArray<TObjectPtr<UObject>>& AssetsIn) override;
 
 	// Gets the tool target requirements for the mode. The resulting targets undergo further processing
 	// to turn them into the input objects that tools get (since these need preview meshes, etc).
-	static const FToolTargetTypeRequirements& GetToolTargetRequirements();
+	static UE_API const FToolTargetTypeRequirements& GetToolTargetRequirements();
 
 	// Use this function to register tools rather than UEdMode::RegisterTool() because we need to specify the ToolsContext
-	void RegisterClothTool(TSharedPtr<FUICommandInfo> UICommand, 
+	UE_API void RegisterClothTool(TSharedPtr<FUICommandInfo> UICommand, 
 		FString ToolIdentifier, 
 		UInteractiveToolBuilder* Builder,
 		const IChaosClothAssetEditorToolBuilder* ClothToolBuilder,
 		UEditorInteractiveToolsContext* UseToolsContext, 
 		EToolsContextScope ToolScope = EToolsContextScope::Default);
 
-	void RegisterAddNodeCommand(TSharedPtr<FUICommandInfo> AddNodeCommand, const FName& NewNodeType, TSharedPtr<FUICommandInfo> StartToolCommand);
+	UE_API void RegisterAddNodeCommand(TSharedPtr<FUICommandInfo> AddNodeCommand, const FName& NewNodeType, TSharedPtr<FUICommandInfo> StartToolCommand);
 
-	void SetPreviewScene(UE::Chaos::ClothAsset::FChaosClothPreviewScene* PreviewScene);
+	UE_API void SetPreviewScene(UE::Chaos::ClothAsset::FChaosClothPreviewScene* PreviewScene);
 
 	// Bounding box for rest space meshes
-	virtual FBox SceneBoundingBox() const override;
+	UE_API virtual FBox SceneBoundingBox() const override;
 
-	void SetRestSpaceViewportClient(TWeakPtr<UE::Chaos::ClothAsset::FChaosClothEditorRestSpaceViewportClient, ESPMode::ThreadSafe> ViewportClient);
-	void RefocusRestSpaceViewportClient();
-	void FirstTimeFocusRestSpaceViewport();
+	UE_API void SetRestSpaceViewportClient(TWeakPtr<UE::Chaos::ClothAsset::FChaosClothEditorRestSpaceViewportClient, ESPMode::ThreadSafe> ViewportClient);
+	UE_API void RefocusRestSpaceViewportClient();
+	UE_API void FirstTimeFocusRestSpaceViewport();
 
 	// intended to be called by the toolkit when selected node in the Dataflow graph changes
-	void SetSelectedClothCollection(TSharedPtr<FManagedArrayCollection> Collection, TSharedPtr<FManagedArrayCollection> InputCollection = nullptr, bool bDeferDynamicMeshInitForTool = false);
+	UE_API void SetSelectedClothCollection(TSharedPtr<FManagedArrayCollection> Collection, TSharedPtr<FManagedArrayCollection> InputCollection = nullptr, bool bDeferDynamicMeshInitForTool = false);
 
 	// gets the currently selected cloth collection, as specified by the toolkit
-	TSharedPtr<FManagedArrayCollection> GetClothCollection();
-	TSharedPtr<FManagedArrayCollection> GetInputClothCollection();
+	UE_API TSharedPtr<FManagedArrayCollection> GetClothCollection();
+	UE_API TSharedPtr<FManagedArrayCollection> GetInputClothCollection();
 
-	void SetDataflowContext(TWeakPtr<UE::Dataflow::FEngineContext> InDataflowContext);
-	void SetDataflowGraphEditor(TSharedPtr<SDataflowGraphEditor> InGraphEditor);
+	UE_API void SetDataflowContext(TWeakPtr<UE::Dataflow::FEngineContext> InDataflowContext);
+	UE_API void SetDataflowGraphEditor(TSharedPtr<SDataflowGraphEditor> InGraphEditor);
 	
-	void StartToolForSelectedNode(const UObject* SelectedNode);
-	void OnDataflowNodeDeleted(const TSet<UObject*>& DeletedNodes);
+	UE_API void StartToolForSelectedNode(const UObject* SelectedNode);
+	UE_API void OnDataflowNodeDeleted(const TSet<UObject*>& DeletedNodes);
 
 	/**
 	* Return the single selected node in the Dataflow Graph Editor only if it has an output of the specified type
 	* If there is not a single node selected, or if it does not have the specified output, return null
 	*/
-	UEdGraphNode* GetSingleSelectedNodeWithOutputType(const FName& SelectedNodeOutputTypeName) const;
+	UE_API UEdGraphNode* GetSingleSelectedNodeWithOutputType(const FName& SelectedNodeOutputTypeName) const;
 
 	/**
 	 * Create a node with the specified type in the graph
 	*/
-	UEdGraphNode* CreateNewNode(const FName& NewNodeTypeName);
+	UE_API UEdGraphNode* CreateNewNode(const FName& NewNodeTypeName);
 
 	/** Create a node with the specified type, then connect it to the output of the specified UpstreamNode.
 	* If the specified output of the upstream node is already connected to another node downstream, we first break
@@ -221,17 +223,17 @@ private:
 	* @param NewNodeConnectionName The name of the input/output connection on our new node that will be connected
 	* @return The newly-created node
 	*/
-	UEdGraphNode* CreateAndConnectNewNode(const FName& NewNodeTypeName,	UEdGraphNode& UpstreamNode,	const FName& ConnectionTypeName, const FName& NewNodeConnectionName);
+	UE_API UEdGraphNode* CreateAndConnectNewNode(const FName& NewNodeTypeName,	UEdGraphNode& UpstreamNode,	const FName& ConnectionTypeName, const FName& NewNodeConnectionName);
 
-	void AddNode(FName NewNodeType);
-	bool CanAddNode(FName NewNodeType) const;
+	UE_API void AddNode(FName NewNodeType);
+	UE_API bool CanAddNode(FName NewNodeType) const;
 
 
-	void InitializeContextObject();
-	void UpdateContextObject(const TSharedPtr<FManagedArrayCollection>& Collection);
-	void DeleteContextObject();
+	UE_API void InitializeContextObject();
+	UE_API void UpdateContextObject(const TSharedPtr<FManagedArrayCollection>& Collection);
+	UE_API void DeleteContextObject();
 
-	bool IsComponentSelected(const UPrimitiveComponent* InComponent);
+	UE_API bool IsComponentSelected(const UPrimitiveComponent* InComponent);
 
 	// Rest space wireframe. They have to get ticked to be able to respond to setting changes. 
 	UPROPERTY()
@@ -272,7 +274,7 @@ private:
 	// Whether the rest space viewport should focus on the rest space mesh on the next tick
 	bool bShouldFocusRestSpaceView = true;
 
-	void RestSpaceViewportResized(FViewport* RestspaceViewport, uint32 Unused);
+	UE_API void RestSpaceViewportResized(FViewport* RestspaceViewport, uint32 Unused);
 
 	UE::Chaos::ClothAsset::EClothPatternVertexType ConstructionViewMode;
 
@@ -291,16 +293,16 @@ private:
 	bool bConstructionViewSeamsVisible = false;
 	bool bShouldRestoreConstructionViewSeams = false;
 	bool bConstructionViewSeamsCollapse = false;
-	void InitializeSeamDraw();
+	UE_API void InitializeSeamDraw();
 
 	bool bConstructionViewNormalsVisible = false;
-	void InitializeSurfaceNormalDraw();
+	UE_API void InitializeSurfaceNormalDraw();
 
 	bool bPatternColors = false;
 	bool bMeshStats = false;
 
 	// Create dynamic mesh components from the cloth component's rest space info
-	void ReinitializeDynamicMeshComponents();
+	UE_API void ReinitializeDynamicMeshComponents();
 
 	// Simulation controls
 	bool bShouldResetSimulation = false;
@@ -328,3 +330,4 @@ private:
 	FDateTime LastToolStartTimestamp;
 };
 
+#undef UE_API
diff --git a/Engine/Plugins/ChaosClothAssetEditor/Source/ChaosClothAssetEditor/Private/ChaosClothAsset/ClothEditorModeToolkit.h b/Engine/Plugins/ChaosClothAssetEditor/Source/ChaosClothAssetEditor/Private/ChaosClothAsset/ClothEditorModeToolkit.h
--- a/Engine/Plugins/ChaosClothAssetEditor/Source/ChaosClothAssetEditor/Private/ChaosClothAsset/ClothEditorModeToolkit.h
+++ b/Engine/Plugins/ChaosClothAssetEditor/Source/ChaosClothAssetEditor/Private/ChaosClothAsset/ClothEditorModeToolkit.h
@@ -4,6 +4,8 @@
 
 #include "BaseCharacterFXEditorModeToolkit.h"
 
+#define UE_API CHAOSCLOTHASSETEDITOR_API
+
 class UEditorInteractiveToolsContext;
 class SChaosClothAssetEditorRestSpaceViewport;
 class SChaosClothAssetEditor3DViewport;
@@ -17,27 +19,29 @@ class SBaseCharacterFXEditorViewport;
  */
 namespace UE::Chaos::ClothAsset
 {
-class CHAOSCLOTHASSETEDITOR_API FChaosClothAssetEditorModeToolkit : public FBaseCharacterFXEditorModeToolkit
+class FChaosClothAssetEditorModeToolkit : public FBaseCharacterFXEditorModeToolkit
 {
 public:
 
-	void Init(const TSharedPtr<IToolkitHost>& InitToolkitHost, TWeakObjectPtr<UEdMode> InOwningMode);
+	UE_API void Init(const TSharedPtr<IToolkitHost>& InitToolkitHost, TWeakObjectPtr<UEdMode> InOwningMode);
 
-	virtual const FSlateBrush* GetActiveToolIcon(const FString& Identifier) const override;
+	UE_API virtual const FSlateBrush* GetActiveToolIcon(const FString& Identifier) const override;
 
-	virtual void OnToolStarted(UInteractiveToolManager* Manager, UInteractiveTool* Tool) override;
-	virtual void OnToolEnded(UInteractiveToolManager* Manager, UInteractiveTool* Tool) override;
+	UE_API virtual void OnToolStarted(UInteractiveToolManager* Manager, UInteractiveTool* Tool) override;
+	UE_API virtual void OnToolEnded(UInteractiveToolManager* Manager, UInteractiveTool* Tool) override;
 
 	// IToolkit
-	virtual FName GetToolkitFName() const override;
-	virtual FText GetBaseToolkitName() const override;
+	UE_API virtual FName GetToolkitFName() const override;
+	UE_API virtual FText GetBaseToolkitName() const override;
 
-	void SetRestSpaceViewportWidget(TWeakPtr<SChaosClothAssetEditorRestSpaceViewport>);
+	UE_API void SetRestSpaceViewportWidget(TWeakPtr<SChaosClothAssetEditorRestSpaceViewport>);
 
 private:
 
-	UEditorInteractiveToolsContext* GetCurrentToolsContext();
+	UE_API UEditorInteractiveToolsContext* GetCurrentToolsContext();
 
 	TWeakPtr<SChaosClothAssetEditorRestSpaceViewport> RestSpaceViewportWidget;
 };
 } // namespace UE::Chaos::ClothAsset
+
+#undef UE_API
diff --git a/Engine/Plugins/ChaosClothAssetEditor/Source/ChaosClothAssetEditor/Private/ChaosClothAsset/ClothEditorModeUILayer.h b/Engine/Plugins/ChaosClothAssetEditor/Source/ChaosClothAssetEditor/Private/ChaosClothAsset/ClothEditorModeUILayer.h
--- a/Engine/Plugins/ChaosClothAssetEditor/Source/ChaosClothAssetEditor/Private/ChaosClothAsset/ClothEditorModeUILayer.h
+++ b/Engine/Plugins/ChaosClothAssetEditor/Source/ChaosClothAssetEditor/Private/ChaosClothAsset/ClothEditorModeUILayer.h
@@ -5,19 +5,21 @@
 #include "BaseCharacterFXEditorModeUILayer.h"
 #include "ClothEditorModeUILayer.generated.h"
 
+#define UE_API CHAOSCLOTHASSETEDITOR_API
+
 /** Interchange layer to manage built in tab locations within the editor's layout. **/
-UCLASS()
-class CHAOSCLOTHASSETEDITOR_API UChaosClothAssetEditorUISubsystem : public UAssetEditorUISubsystem
+UCLASS(MinimalAPI)
+class UChaosClothAssetEditorUISubsystem : public UAssetEditorUISubsystem
 {
 	GENERATED_BODY()
 public:
 
-	virtual void Initialize(FSubsystemCollectionBase& Collection) override;
-	virtual void Deinitialize() override;
-	virtual void RegisterLayoutExtensions(FLayoutExtender& Extender) override;
+	UE_API virtual void Initialize(FSubsystemCollectionBase& Collection) override;
+	UE_API virtual void Deinitialize() override;
+	UE_API virtual void RegisterLayoutExtensions(FLayoutExtender& Extender) override;
 
 	// Identifier for the Layout extension in ChaosClothAssetEditorToolkit's StandaloneDefaultLayout
-	static const FName EditorSidePanelAreaName;
+	static UE_API const FName EditorSidePanelAreaName;
 
 };
 
@@ -28,3 +30,4 @@ public:
 	FChaosClothAssetEditorModeUILayer(const IToolkitHost* InToolkitHost) : FBaseCharacterFXEditorModeUILayer(InToolkitHost) {}
 };
 
+#undef UE_API
diff --git a/Engine/Plugins/ChaosClothAssetEditor/Source/ChaosClothAssetEditor/Private/ChaosClothAsset/ClothEditorPreviewScene.h b/Engine/Plugins/ChaosClothAssetEditor/Source/ChaosClothAssetEditor/Private/ChaosClothAsset/ClothEditorPreviewScene.h
--- a/Engine/Plugins/ChaosClothAssetEditor/Source/ChaosClothAssetEditor/Private/ChaosClothAsset/ClothEditorPreviewScene.h
+++ b/Engine/Plugins/ChaosClothAssetEditor/Source/ChaosClothAssetEditor/Private/ChaosClothAsset/ClothEditorPreviewScene.h
@@ -6,6 +6,8 @@
 #include "Math/Transform.h"
 #include "ClothEditorPreviewScene.generated.h"
 
+#define UE_API CHAOSCLOTHASSETEDITOR_API
+
 class UChaosClothAsset;
 class USkeletalMesh;
 class UChaosClothComponent;
@@ -24,8 +26,8 @@ DECLARE_EVENT(UChaosClothPreviewSceneDescription, FClothPreviewSceneDescriptionC
 ///
 /// The UChaosClothPreviewSceneDescription is a description of the Preview scene contents, intended to be editable in an FAdvancedPreviewSettingsWidget
 /// 
-UCLASS()
-class CHAOSCLOTHASSETEDITOR_API UChaosClothPreviewSceneDescription : public UObject
+UCLASS(MinimalAPI)
+class UChaosClothPreviewSceneDescription : public UObject
 {
 public:
 	GENERATED_BODY()
@@ -37,7 +39,7 @@ public:
 		SetFlags(RF_Transactional);
 	}
 
-	void SetPreviewScene(UE::Chaos::ClothAsset::FChaosClothPreviewScene* PreviewScene);
+	UE_API void SetPreviewScene(UE::Chaos::ClothAsset::FChaosClothPreviewScene* PreviewScene);
 
 	/* Whether the preview viewport should pause animation and simulation while Play In Editor (PIE) or Simulate In Editor is active */
 	UPROPERTY(EditAnywhere, Transient, Category = "Viewport")
@@ -83,11 +85,11 @@ public:
 	UPROPERTY(Transient)
 	bool bValidSelectionForTransform = false;
 
-	virtual void PostEditChangeProperty(struct FPropertyChangedEvent& PropertyChangedEvent) override;
+	UE_API virtual void PostEditChangeProperty(struct FPropertyChangedEvent& PropertyChangedEvent) override;
 
 private:
 
-	virtual void PostTransacted(const FTransactionObjectEvent& TransactionEvent) override;
+	UE_API virtual void PostTransacted(const FTransactionObjectEvent& TransactionEvent) override;
 
 	UE::Chaos::ClothAsset::FChaosClothPreviewScene* PreviewScene;
 };
@@ -98,52 +100,52 @@ namespace UE::Chaos::ClothAsset
 ///
 /// FChaosClothPreviewScene is the actual Preview scene, with contents specified by the SceneDescription
 /// 
-class CHAOSCLOTHASSETEDITOR_API FChaosClothPreviewScene : public FAdvancedPreviewScene
+class FChaosClothPreviewScene : public FAdvancedPreviewScene
 {
 public:
 
-	FChaosClothPreviewScene(FPreviewScene::ConstructionValues ConstructionValues);
-	virtual ~FChaosClothPreviewScene();
+	UE_API FChaosClothPreviewScene(FPreviewScene::ConstructionValues ConstructionValues);
+	UE_API virtual ~FChaosClothPreviewScene();
 
-	virtual void AddReferencedObjects(FReferenceCollector& Collector) override;
+	UE_API virtual void AddReferencedObjects(FReferenceCollector& Collector) override;
 
 	const UChaosClothPreviewSceneDescription* GetPreviewSceneDescription() const { return PreviewSceneDescription; }
 	UChaosClothPreviewSceneDescription* GetPreviewSceneDescription() { return PreviewSceneDescription; }
 
-	void SetClothAsset(UChaosClothAsset* Asset);
+	UE_API void SetClothAsset(UChaosClothAsset* Asset);
 
 	// Update Scene in response to the SceneDescription changing
-	void SceneDescriptionPropertyChanged(const FName& PropertyName);
+	UE_API void SceneDescriptionPropertyChanged(const FName& PropertyName);
 
-	UAnimSingleNodeInstance* GetPreviewAnimInstance();
-	const UAnimSingleNodeInstance* const GetPreviewAnimInstance() const;
+	UE_API UAnimSingleNodeInstance* GetPreviewAnimInstance();
+	UE_API const UAnimSingleNodeInstance* const GetPreviewAnimInstance() const;
 
-	UChaosClothComponent* GetClothComponent();
-	const UChaosClothComponent* GetClothComponent() const;
+	UE_API UChaosClothComponent* GetClothComponent();
+	UE_API const UChaosClothComponent* GetClothComponent() const;
 	
-	const USkeletalMeshComponent* GetSkeletalMeshComponent() const;
+	UE_API const USkeletalMeshComponent* GetSkeletalMeshComponent() const;
 
 	/** Set the scene's ModeManager, which is mainly used to track selected components */
-	void SetModeManager(const TSharedPtr<FAssetEditorModeManager>& InClothPreviewEditorModeManager);
+	UE_API void SetModeManager(const TSharedPtr<FAssetEditorModeManager>& InClothPreviewEditorModeManager);
 
-	void SetGizmoDataBinder(TSharedPtr<FTransformGizmoDataBinder> InDataBinder);
+	UE_API void SetGizmoDataBinder(TSharedPtr<FTransformGizmoDataBinder> InDataBinder);
 
 private:
 
-	virtual void Tick(float DeltaT) override;
+	UE_API virtual void Tick(float DeltaT) override;
 
 	// Create the PreviewAnimationInstance if the AnimationAsset and SkeletalMesh both exist, and set the animation to run on the SkeletalMeshComponent
-	void UpdateSkeletalMeshAnimation();
+	UE_API void UpdateSkeletalMeshAnimation();
 
 	// Attach the cloth component to the skeletal mesh component, if it exists
-	void UpdateClothComponentAttachment();
+	UE_API void UpdateClothComponentAttachment();
 
-	bool IsComponentSelected(const UPrimitiveComponent* InComponent);
+	UE_API bool IsComponentSelected(const UPrimitiveComponent* InComponent);
 
-	void SaveAnimationState();
-	void RestoreSavedAnimationState();
-	void HandlePackageReloaded(const EPackageReloadPhase InPackageReloadPhase, FPackageReloadedEvent* InPackageReloadedEvent);
-	void HandleReimportManagerPostReimport(UObject* ReimportedObject, bool bWasSuccessful);
+	UE_API void SaveAnimationState();
+	UE_API void RestoreSavedAnimationState();
+	UE_API void HandlePackageReloaded(const EPackageReloadPhase InPackageReloadPhase, FPackageReloadedEvent* InPackageReloadedEvent);
+	UE_API void HandleReimportManagerPostReimport(UObject* ReimportedObject, bool bWasSuccessful);
 
 	TObjectPtr<UChaosClothPreviewSceneDescription> PreviewSceneDescription;
 
@@ -171,3 +173,4 @@ private:
 };
 } // namespace UE::Chaos::ClothAsset
 
+#undef UE_API
diff --git a/Engine/Plugins/ChaosClothAssetEditor/Source/ChaosClothAssetEditor/Private/ChaosClothAsset/ClothEditorRestSpaceViewportClient.h b/Engine/Plugins/ChaosClothAssetEditor/Source/ChaosClothAssetEditor/Private/ChaosClothAsset/ClothEditorRestSpaceViewportClient.h
--- a/Engine/Plugins/ChaosClothAssetEditor/Source/ChaosClothAssetEditor/Private/ChaosClothAsset/ClothEditorRestSpaceViewportClient.h
+++ b/Engine/Plugins/ChaosClothAssetEditor/Source/ChaosClothAssetEditor/Private/ChaosClothAsset/ClothEditorRestSpaceViewportClient.h
@@ -9,44 +9,46 @@
 #include "ChaosClothAsset/ClothEditorMode.h"
 #include "ChaosClothAsset/ClothPatternVertexType.h"
 
+#define UE_API CHAOSCLOTHASSETEDITOR_API
+
 class UInputBehaviorSet;
 class UPointLightComponent;
 
 namespace UE::Chaos::ClothAsset
 {
 
-class CHAOSCLOTHASSETEDITOR_API FChaosClothEditorRestSpaceViewportClient : public FEditorViewportClient, public IInputBehaviorSource
+class FChaosClothEditorRestSpaceViewportClient : public FEditorViewportClient, public IInputBehaviorSource
 {
 public:
 
-	FChaosClothEditorRestSpaceViewportClient(FEditorModeTools* InModeTools, FPreviewScene* InPreviewScene = nullptr, const TWeakPtr<SEditorViewport>& InEditorViewportWidget = nullptr);
+	UE_API FChaosClothEditorRestSpaceViewportClient(FEditorModeTools* InModeTools, FPreviewScene* InPreviewScene = nullptr, const TWeakPtr<SEditorViewport>& InEditorViewportWidget = nullptr);
 
 	virtual ~FChaosClothEditorRestSpaceViewportClient() = default;
 
 	// IInputBehaviorSource
-	virtual const UInputBehaviorSet* GetInputBehaviors() const override;
+	UE_API virtual const UInputBehaviorSet* GetInputBehaviors() const override;
 
 	// FGCObject
-	virtual void AddReferencedObjects(FReferenceCollector& Collector) override;
+	UE_API virtual void AddReferencedObjects(FReferenceCollector& Collector) override;
 
-	virtual bool InputKey(const FInputKeyEventArgs& EventArgs) override;
+	UE_API virtual bool InputKey(const FInputKeyEventArgs& EventArgs) override;
 
-	virtual bool ShouldOrbitCamera() const override;
+	UE_API virtual bool ShouldOrbitCamera() const override;
 
-	void SetConstructionViewMode(EClothPatternVertexType InViewMode);
-	EClothPatternVertexType GetConstructionViewMode() const;
+	UE_API void SetConstructionViewMode(EClothPatternVertexType InViewMode);
+	UE_API EClothPatternVertexType GetConstructionViewMode() const;
 
-	void SetEditorViewportWidget(TWeakPtr<SEditorViewport> InEditorViewportWidget);
-	void SetToolCommandList(TWeakPtr<FUICommandList> ToolCommandList);
+	UE_API void SetEditorViewportWidget(TWeakPtr<SEditorViewport> InEditorViewportWidget);
+	UE_API void SetToolCommandList(TWeakPtr<FUICommandList> ToolCommandList);
 
-	float GetCameraPointLightIntensity() const;
-	void SetCameraPointLightIntensity(float Intensity);
+	UE_API float GetCameraPointLightIntensity() const;
+	UE_API void SetCameraPointLightIntensity(float Intensity);
 
 private:
 
-	virtual void Tick(float DeltaSeconds) override;
+	UE_API virtual void Tick(float DeltaSeconds) override;
 
-	void UpdateBehaviorsForCurrentViewMode();
+	UE_API void UpdateBehaviorsForCurrentViewMode();
 
 	TObjectPtr<UPointLightComponent> CameraPointLight;
 
@@ -66,3 +68,5 @@ private:
 	FViewportCameraTransform SavedInactiveViewTransform;
 };
 } // namespace UE::Chaos::ClothAsset
+
+#undef UE_API
diff --git a/Engine/Plugins/ChaosClothAssetEditor/Source/ChaosClothAssetEditor/Private/ChaosClothAsset/ClothEditorStyle.h b/Engine/Plugins/ChaosClothAssetEditor/Source/ChaosClothAssetEditor/Private/ChaosClothAsset/ClothEditorStyle.h
--- a/Engine/Plugins/ChaosClothAssetEditor/Source/ChaosClothAssetEditor/Private/ChaosClothAsset/ClothEditorStyle.h
+++ b/Engine/Plugins/ChaosClothAssetEditor/Source/ChaosClothAssetEditor/Private/ChaosClothAsset/ClothEditorStyle.h
@@ -5,24 +5,28 @@
 #include "Styling/SlateStyle.h"
 #include "Styling/SlateStyleRegistry.h"
 
+#define UE_API CHAOSCLOTHASSETEDITOR_API
+
 namespace UE::Chaos::ClothAsset
 {
 /**
  * Slate style set for Cloth Editor
  */
-class CHAOSCLOTHASSETEDITOR_API FChaosClothAssetEditorStyle
+class FChaosClothAssetEditorStyle
 	: public FSlateStyleSet
 {
 public:
-	const static FName StyleName;
+	UE_API const static FName StyleName;
 
 	/** Access the singleton instance for this style set */
-	static FChaosClothAssetEditorStyle& Get();
+	static UE_API FChaosClothAssetEditorStyle& Get();
 
 private:
-	static FString InContent(const FString& RelativePath, const ANSICHAR* Extension);
+	static UE_API FString InContent(const FString& RelativePath, const ANSICHAR* Extension);
 
-	FChaosClothAssetEditorStyle();
-	~FChaosClothAssetEditorStyle();
+	UE_API FChaosClothAssetEditorStyle();
+	UE_API ~FChaosClothAssetEditorStyle();
 };
 } // namespace UE::Chaos::ClothAsset
+
+#undef UE_API
