Updated headers using UnrealCodeFixup to make sure dllstorage is on methods/staticvars instead of types. Part 1/n

diff --git a/Engine/Plugins/ChaosClothAssetEditor/Source/ChaosClothAssetEditor/Private/ChaosClothAsset/ClothEditorToolkit.h b/Engine/Plugins/ChaosClothAssetEditor/Source/ChaosClothAssetEditor/Private/ChaosClothAsset/ClothEditorToolkit.h
--- a/Engine/Plugins/ChaosClothAssetEditor/Source/ChaosClothAssetEditor/Private/ChaosClothAsset/ClothEditorToolkit.h
+++ b/Engine/Plugins/ChaosClothAssetEditor/Source/ChaosClothAssetEditor/Private/ChaosClothAsset/ClothEditorToolkit.h
@@ -6,6 +6,8 @@
 #include "Dataflow/DataflowObjectInterface.h"
 #include "ChaosClothAsset/ClothEditorPreviewScene.h"
 
+#define UE_API CHAOSCLOTHASSETEDITOR_API
+
 template<typename T> class SComboBox;
 class SClothCollectionOutliner;
 class UDataflow;
@@ -27,7 +29,7 @@ struct FClothSimulationNodeDetailExtender;
 
 namespace UE::Dataflow
 {
-	class CHAOSCLOTHASSETEDITOR_API FClothAssetDataflowContext final : public TEngineContext<FContextSingle>
+	class FClothAssetDataflowContext final : public TEngineContext<FContextSingle>
 	{
 	public:
 		DATAFLOW_CONTEXT_INTERNAL(TEngineContext<FContextSingle>, FClothAssetDataflowContext);
@@ -48,103 +50,103 @@ namespace UE::Chaos::ClothAsset
  * initializing the Cloth mode.
  * Thus, the FChaosClothAssetEditorToolkit ends up being the central place for the Cloth Asset Editor setup.
  */
-class CHAOSCLOTHASSETEDITOR_API FChaosClothAssetEditorToolkit final : public FBaseCharacterFXEditorToolkit, public FTickableEditorObject, public FNotifyHook
+class FChaosClothAssetEditorToolkit final : public FBaseCharacterFXEditorToolkit, public FTickableEditorObject, public FNotifyHook
 {
 public:
 
-	explicit FChaosClothAssetEditorToolkit(UAssetEditor* InOwningAssetEditor);
-	virtual ~FChaosClothAssetEditorToolkit();
+	UE_API explicit FChaosClothAssetEditorToolkit(UAssetEditor* InOwningAssetEditor);
+	UE_API virtual ~FChaosClothAssetEditorToolkit();
 
-	TSharedPtr<UE::Dataflow::FEngineContext> GetDataflowContext() const;
-	const UDataflow* GetDataflow() const;
+	UE_API TSharedPtr<UE::Dataflow::FEngineContext> GetDataflowContext() const;
+	UE_API const UDataflow* GetDataflow() const;
 
 private:
 
-	static const FName ClothPreviewTabID;
-	static const FName OutlinerTabID;
-	static const FName PreviewSceneDetailsTabID;
-	static const FName SimulationVisualizationTabID;
+	static UE_API const FName ClothPreviewTabID;
+	static UE_API const FName OutlinerTabID;
+	static UE_API const FName PreviewSceneDetailsTabID;
+	static UE_API const FName SimulationVisualizationTabID;
 
 	// FTickableEditorObject
-	virtual void Tick(float DeltaTime) override;
+	UE_API virtual void Tick(float DeltaTime) override;
 	virtual bool IsTickable() const override { return true; }
-	virtual TStatId GetStatId() const override;
+	UE_API virtual TStatId GetStatId() const override;
 
 	// FBaseCharacterFXEditorToolkit
-	virtual FEditorModeID GetEditorModeId() const override;
-	virtual void InitializeEdMode(UBaseCharacterFXEditorMode* EdMode) override;
-	virtual void CreateEditorModeUILayer() override;
+	UE_API virtual FEditorModeID GetEditorModeId() const override;
+	UE_API virtual void InitializeEdMode(UBaseCharacterFXEditorMode* EdMode) override;
+	UE_API virtual void CreateEditorModeUILayer() override;
 
 	// FBaseAssetToolkit
-	virtual void CreateWidgets() override;
-	virtual AssetEditorViewportFactoryFunction GetViewportDelegate() override;
-	virtual TSharedPtr<FEditorViewportClient> CreateEditorViewportClient() const override;
+	UE_API virtual void CreateWidgets() override;
+	UE_API virtual AssetEditorViewportFactoryFunction GetViewportDelegate() override;
+	UE_API virtual TSharedPtr<FEditorViewportClient> CreateEditorViewportClient() const override;
 
 	// FAssetEditorToolkit
-	virtual void AddViewportOverlayWidget(TSharedRef<SWidget> InViewportOverlayWidget, int32 ZOrder = INDEX_NONE) override;
-	virtual void RemoveViewportOverlayWidget(TSharedRef<SWidget> InViewportOverlayWidget) override;
-	virtual bool OnRequestClose(EAssetEditorCloseReason InCloseReason) override;
-	virtual void OnClose() override;
-
-	virtual void PostInitAssetEditor() override;
-	virtual void InitToolMenuContext(FToolMenuContext& MenuContext) override;
-	virtual void GetSaveableObjects(TArray<UObject*>& OutObjects) const override;
-	virtual bool ShouldReopenEditorForSavedAsset(const UObject* Asset) const override;
-	virtual void OnAssetsSaved(const TArray<UObject*>& SavedObjects) override;
-	virtual void OnAssetsSavedAs(const TArray<UObject*>& SavedObjects) override;
+	UE_API virtual void AddViewportOverlayWidget(TSharedRef<SWidget> InViewportOverlayWidget, int32 ZOrder = INDEX_NONE) override;
+	UE_API virtual void RemoveViewportOverlayWidget(TSharedRef<SWidget> InViewportOverlayWidget) override;
+	UE_API virtual bool OnRequestClose(EAssetEditorCloseReason InCloseReason) override;
+	UE_API virtual void OnClose() override;
+
+	UE_API virtual void PostInitAssetEditor() override;
+	UE_API virtual void InitToolMenuContext(FToolMenuContext& MenuContext) override;
+	UE_API virtual void GetSaveableObjects(TArray<UObject*>& OutObjects) const override;
+	UE_API virtual bool ShouldReopenEditorForSavedAsset(const UObject* Asset) const override;
+	UE_API virtual void OnAssetsSaved(const TArray<UObject*>& SavedObjects) override;
+	UE_API virtual void OnAssetsSavedAs(const TArray<UObject*>& SavedObjects) override;
 
 	// IAssetEditorInstance
 	virtual bool IsPrimaryEditor() const override { return true; };
 
 	// IToolkit
-	virtual FText GetToolkitName() const override;
-	virtual FName GetToolkitFName() const override;
-	virtual FText GetBaseToolkitName() const override;
-	virtual FText GetToolkitToolTipText() const override;
-	virtual void RegisterTabSpawners(const TSharedRef<FTabManager>& InTabManager) override;
-	virtual void UnregisterTabSpawners(const TSharedRef<FTabManager>& TabManager) override;
+	UE_API virtual FText GetToolkitName() const override;
+	UE_API virtual FName GetToolkitFName() const override;
+	UE_API virtual FText GetBaseToolkitName() const override;
+	UE_API virtual FText GetToolkitToolTipText() const override;
+	UE_API virtual void RegisterTabSpawners(const TSharedRef<FTabManager>& InTabManager) override;
+	UE_API virtual void UnregisterTabSpawners(const TSharedRef<FTabManager>& TabManager) override;
 
 	// FNotifyHook
-	virtual void NotifyPreChange(class FEditPropertyChain* PropertyAboutToChange) override;
+	UE_API virtual void NotifyPreChange(class FEditPropertyChain* PropertyAboutToChange) override;
 
 	// Return the cloth asset held by the Cloth Editor
-	UChaosClothAsset* GetAsset() const;
-	UDataflow* GetDataflow();
+	UE_API UChaosClothAsset* GetAsset() const;
+	UE_API UDataflow* GetDataflow();
 	
-	TSharedRef<SDockTab> SpawnTab_ClothPreview(const FSpawnTabArgs& Args);
-	TSharedRef<SDockTab> SpawnTab_Outliner(const FSpawnTabArgs& Args);
-	TSharedRef<SDockTab> SpawnTab_PreviewSceneDetails(const FSpawnTabArgs& Args);
-	TSharedRef<SDockTab> SpawnTab_GraphCanvas(const FSpawnTabArgs& Args);
-	TSharedRef<SDockTab> SpawnTab_NodeDetails(const FSpawnTabArgs& Args);
-	TSharedRef<SDockTab> SpawnTab_SimulationVisualization(const FSpawnTabArgs& Args);
+	UE_API TSharedRef<SDockTab> SpawnTab_ClothPreview(const FSpawnTabArgs& Args);
+	UE_API TSharedRef<SDockTab> SpawnTab_Outliner(const FSpawnTabArgs& Args);
+	UE_API TSharedRef<SDockTab> SpawnTab_PreviewSceneDetails(const FSpawnTabArgs& Args);
+	UE_API TSharedRef<SDockTab> SpawnTab_GraphCanvas(const FSpawnTabArgs& Args);
+	UE_API TSharedRef<SDockTab> SpawnTab_NodeDetails(const FSpawnTabArgs& Args);
+	UE_API TSharedRef<SDockTab> SpawnTab_SimulationVisualization(const FSpawnTabArgs& Args);
 
-	void InitDetailsViewPanel();
-	void OnFinishedChangingAssetProperties(const FPropertyChangedEvent&);
+	UE_API void InitDetailsViewPanel();
+	UE_API void OnFinishedChangingAssetProperties(const FPropertyChangedEvent&);
 
-	void OnClothAssetChanged();
-	void InvalidateViews();
+	UE_API void OnClothAssetChanged();
+	UE_API void InvalidateViews();
 
 	// Dataflow
-	void EvaluateNode(const FDataflowNode* Node, const FDataflowOutput* Output);
-	TSharedRef<SDataflowGraphEditor> CreateGraphEditorWidget();
-	void ReinitializeGraphEditorWidget();
-	TSharedPtr<IStructureDetailsView> CreateNodeDetailsEditorWidget(UObject* ObjectToEdit);
+	UE_API void EvaluateNode(const FDataflowNode* Node, const FDataflowOutput* Output);
+	UE_API TSharedRef<SDataflowGraphEditor> CreateGraphEditorWidget();
+	UE_API void ReinitializeGraphEditorWidget();
+	UE_API TSharedPtr<IStructureDetailsView> CreateNodeDetailsEditorWidget(UObject* ObjectToEdit);
 
-	TSharedPtr<FManagedArrayCollection> GetClothCollectionIfPossible(const TSharedPtr<FDataflowNode> InDataflowNode, const TSharedPtr<UE::Dataflow::FEngineContext> Context);
-	TSharedPtr<FManagedArrayCollection> GetInputClothCollectionIfPossible(const TSharedPtr<FDataflowNode> InDataflowNode, const TSharedPtr<UE::Dataflow::FEngineContext> Context);
-	TSharedPtr<FDataflowNode> GetSelectedDataflowNode();
-	TSharedPtr<const FDataflowNode> GetSelectedDataflowNode() const;
+	UE_API TSharedPtr<FManagedArrayCollection> GetClothCollectionIfPossible(const TSharedPtr<FDataflowNode> InDataflowNode, const TSharedPtr<UE::Dataflow::FEngineContext> Context);
+	UE_API TSharedPtr<FManagedArrayCollection> GetInputClothCollectionIfPossible(const TSharedPtr<FDataflowNode> InDataflowNode, const TSharedPtr<UE::Dataflow::FEngineContext> Context);
+	UE_API TSharedPtr<FDataflowNode> GetSelectedDataflowNode();
+	UE_API TSharedPtr<const FDataflowNode> GetSelectedDataflowNode() const;
 
 	// DataflowEditorActions
-	void OnPropertyValueChanged(const FPropertyChangedEvent& PropertyChangedEvent);
-	bool OnNodeVerifyTitleCommit(const FText& NewText, UEdGraphNode* GraphNode, FText& OutErrorMessage) const;
-	void OnNodeTitleCommitted(const FText& InNewText, ETextCommit::Type InCommitType, UEdGraphNode* GraphNode) const;
-	void OnNodeSelectionChanged(const TSet<UObject*>& NewSelection);
-	void OnNodeDeleted(const TSet<UObject*>& DeletedNodes);
-	void OnNodeSingleClicked(UObject* ClickedNode) const;
+	UE_API void OnPropertyValueChanged(const FPropertyChangedEvent& PropertyChangedEvent);
+	UE_API bool OnNodeVerifyTitleCommit(const FText& NewText, UEdGraphNode* GraphNode, FText& OutErrorMessage) const;
+	UE_API void OnNodeTitleCommitted(const FText& InNewText, ETextCommit::Type InCommitType, UEdGraphNode* GraphNode) const;
+	UE_API void OnNodeSelectionChanged(const TSet<UObject*>& NewSelection);
+	UE_API void OnNodeDeleted(const TSet<UObject*>& DeletedNodes);
+	UE_API void OnNodeSingleClicked(UObject* ClickedNode) const;
 
 	// Delegates
-	void HandlePackageReloaded(const EPackageReloadPhase InPackageReloadPhase, FPackageReloadedEvent* InPackageReloadedEvent);
+	UE_API void HandlePackageReloaded(const EPackageReloadPhase InPackageReloadPhase, FPackageReloadedEvent* InPackageReloadedEvent);
 
 	/** Scene in which the 3D sim space preview meshes live. Ownership shared with AdvancedPreviewSettingsWidget*/
 	TSharedPtr<FChaosClothPreviewScene> ClothPreviewScene;
@@ -171,11 +173,11 @@ private:
 	FDelegateHandle OnNodeInvalidatedDelegateHandle;
 	FGuid SelectedDataflowNodeGuid;
 
-	static const FName GraphCanvasTabId;
+	static UE_API const FName GraphCanvasTabId;
 	TSharedPtr<SDockTab> GraphEditorTab;
 	TSharedPtr<SDataflowGraphEditor> GraphEditor;
 
-	static const FName NodeDetailsTabId;
+	static UE_API const FName NodeDetailsTabId;
 	TSharedPtr<SDockTab> NodeDetailsTab;
 	TSharedPtr<IStructureDetailsView> NodeDetailsEditor;
 	TSharedPtr<FClothSimulationNodeDetailExtender> NodeDetailsExtender;
@@ -186,3 +188,5 @@ private:
 	FTickCommands TickCommands;
 };
 } // namespace UE::Chaos::ClothAsset
+
+#undef UE_API
diff --git a/Engine/Plugins/ChaosClothAssetEditor/Source/ChaosClothAssetEditor/Private/ChaosClothAsset/SClothEditor3DViewport.h b/Engine/Plugins/ChaosClothAssetEditor/Source/ChaosClothAssetEditor/Private/ChaosClothAsset/SClothEditor3DViewport.h
--- a/Engine/Plugins/ChaosClothAssetEditor/Source/ChaosClothAssetEditor/Private/ChaosClothAsset/SClothEditor3DViewport.h
+++ b/Engine/Plugins/ChaosClothAssetEditor/Source/ChaosClothAssetEditor/Private/ChaosClothAsset/SClothEditor3DViewport.h
@@ -7,6 +7,8 @@
 #include "SCommonEditorViewportToolbarBase.h"
 #include "SBaseCharacterFXEditorViewport.h"
 
+#define UE_API CHAOSCLOTHASSETEDITOR_API
+
 namespace UE::Chaos::ClothAsset
 {
 class FChaosClothPreviewScene;
@@ -16,7 +18,7 @@ class SClothAnimationScrubPanel;
 /**
  * Viewport used for 3D preview in cloth editor. Has a custom toolbar overlay at the top.
  */
-class CHAOSCLOTHASSETEDITOR_API SChaosClothAssetEditor3DViewport : public SBaseCharacterFXEditorViewport, public ICommonEditorViewportToolbarInfoProvider
+class SChaosClothAssetEditor3DViewport : public SBaseCharacterFXEditorViewport, public ICommonEditorViewportToolbarInfoProvider
 {
 
 public:
@@ -27,19 +29,20 @@ public:
 		SLATE_ARGUMENT(TSharedPtr<FUICommandList>, ToolkitCommandList)
 	SLATE_END_ARGS()
 
-	void Construct(const FArguments& InArgs, const FAssetEditorViewportConstructionArgs& InViewportConstructionArgs);
+	UE_API void Construct(const FArguments& InArgs, const FAssetEditorViewportConstructionArgs& InViewportConstructionArgs);
 
 	// SAssetEditorViewport
-	virtual void BindCommands() override;
-	virtual TSharedPtr<SWidget> BuildViewportToolbar() override;
-	virtual void PopulateViewportOverlays(TSharedRef<SOverlay> Overlay) override;
-	virtual bool IsVisible() const override;
+	UE_API virtual void BindCommands() override;
+	UE_API virtual TSharedPtr<SWidget> MakeViewportToolbar() override;
+	UE_API virtual TSharedPtr<SWidget> BuildViewportToolbar() override;
+	UE_API virtual void PopulateViewportOverlays(TSharedRef<SOverlay> Overlay) override;
+	UE_API virtual bool IsVisible() const override;
 
-	virtual void OnFocusViewportToSelection() override;
+	UE_API virtual void OnFocusViewportToSelection() override;
 
 	// ICommonEditorViewportToolbarInfoProvider
-	virtual TSharedRef<class SEditorViewport> GetViewportWidget() override;
-	virtual TSharedPtr<FExtender> GetExtenders() const override;
+	UE_API virtual TSharedRef<class SEditorViewport> GetViewportWidget() override;
+	UE_API virtual TSharedPtr<FExtender> GetExtenders() const override;
 	virtual void OnFloatingButtonClicked() override {}
 
 private:
@@ -48,16 +51,18 @@ private:
 	// Use SEditorViewport::CommandList if we want command hotkeys to only be active when the mouse is in this viewport.
 	TSharedPtr<FUICommandList> ToolkitCommandList;
 
-	TWeakPtr<UE::Chaos::ClothAsset::FChaosClothPreviewScene> GetPreviewScene();
-	TWeakPtr<const UE::Chaos::ClothAsset::FChaosClothPreviewScene> GetPreviewScene() const;
+	UE_API TWeakPtr<UE::Chaos::ClothAsset::FChaosClothPreviewScene> GetPreviewScene();
+	UE_API TWeakPtr<const UE::Chaos::ClothAsset::FChaosClothPreviewScene> GetPreviewScene() const;
 
-	virtual TSharedPtr<IPreviewProfileController> CreatePreviewProfileController() override;
+	UE_API virtual TSharedPtr<IPreviewProfileController> CreatePreviewProfileController() override;
 	
-	FText GetViewportDisplayString() const;
+	UE_API FText GetViewportDisplayString() const;
 
-	float GetViewMinInput() const;
-	float GetViewMaxInput() const;
-	EVisibility GetAnimControlVisibility() const;
+	UE_API float GetViewMinInput() const;
+	UE_API float GetViewMaxInput() const;
+	UE_API EVisibility GetAnimControlVisibility() const;
 		 
 
 };
+
+#undef UE_API
diff --git a/Engine/Plugins/ChaosClothAssetEditor/Source/ChaosClothAssetEditor/Private/ChaosClothAsset/SClothEditor3DViewportToolBar.h b/Engine/Plugins/ChaosClothAssetEditor/Source/ChaosClothAssetEditor/Private/ChaosClothAsset/SClothEditor3DViewportToolBar.h
new file mode 100644
--- /dev/null
+++ b/Engine/Plugins/ChaosClothAssetEditor/Source/ChaosClothAssetEditor/Private/ChaosClothAsset/SClothEditor3DViewportToolBar.h
@@ -0,0 +1,40 @@
+// Copyright Epic Games, Inc. All Rights Reserved.
+
+#pragma once
+
+#include "SClothEditorViewportToolBarBase.h"
+
+#define UE_API CHAOSCLOTHASSETEDITOR_API
+
+class FExtender;
+class FUICommandList;
+class SChaosClothAssetEditor3DViewport;
+
+/**
+ * Toolbar that shows up at the top of the 3d viewport
+ */
+class SChaosClothAssetEditor3DViewportToolBar : public SChaosClothAssetEditorViewportToolBarBase
+{
+public:
+	SLATE_BEGIN_ARGS(SChaosClothAssetEditor3DViewportToolBar) {}
+		SLATE_ARGUMENT(TSharedPtr<FUICommandList>, CommandList)
+		SLATE_ARGUMENT(TSharedPtr<FExtender>, Extenders)
+	SLATE_END_ARGS()
+
+	UE_API void Construct(const FArguments& InArgs, TSharedPtr<class SChaosClothAssetEditor3DViewport> InChaosClothAssetEditor3DViewport);
+
+private:
+	UE_API TSharedRef<SWidget> MakeOptionsMenu();
+	UE_API TSharedRef<SWidget> MakeDisplayToolBar(const TSharedPtr<FExtender> InExtenders);
+	UE_API TSharedRef<SWidget> MakeToolBar(const TSharedPtr<FExtender> InExtenders);
+	UE_API FText GetDisplayString() const;
+	UE_API FText GetLODMenuLabel() const;
+	UE_API TSharedRef<SWidget> MakeLODMenu() const;
+
+	/** The viewport that we are in */
+	TWeakPtr<class SChaosClothAssetEditor3DViewport> ChaosClothAssetEditor3DViewportPtr;
+
+	TSharedPtr<FUICommandList> CommandList;
+};
+
+#undef UE_API
diff --git a/Engine/Plugins/ChaosClothAssetEditor/Source/ChaosClothAssetEditor/Private/ChaosClothAsset/SClothEditorAdvancedPreviewDetailsTab.h b/Engine/Plugins/ChaosClothAssetEditor/Source/ChaosClothAssetEditor/Private/ChaosClothAsset/SClothEditorAdvancedPreviewDetailsTab.h
--- a/Engine/Plugins/ChaosClothAssetEditor/Source/ChaosClothAssetEditor/Private/ChaosClothAsset/SClothEditorAdvancedPreviewDetailsTab.h
+++ b/Engine/Plugins/ChaosClothAssetEditor/Source/ChaosClothAssetEditor/Private/ChaosClothAsset/SClothEditorAdvancedPreviewDetailsTab.h
@@ -4,15 +4,19 @@
 
 #include "SAdvancedPreviewDetailsTab.h"
 
+#define UE_API CHAOSCLOTHASSETEDITOR_API
+
 // Custom subclass of SAdvancedPreviewDetailsTab that allows invalidating the cached state of the SettingsView when the preview scene changes
-class CHAOSCLOTHASSETEDITOR_API SChaosClothEditorAdvancedPreviewDetailsTab : public SAdvancedPreviewDetailsTab
+class SChaosClothEditorAdvancedPreviewDetailsTab : public SAdvancedPreviewDetailsTab
 {
 public:
 
-	SChaosClothEditorAdvancedPreviewDetailsTab();
-	virtual ~SChaosClothEditorAdvancedPreviewDetailsTab() override;
+	UE_API SChaosClothEditorAdvancedPreviewDetailsTab();
+	UE_API virtual ~SChaosClothEditorAdvancedPreviewDetailsTab() override;
 
 private:
 
 	FDelegateHandle PropertyChangedDelegateHandle;
 };
+
+#undef UE_API
diff --git a/Engine/Plugins/ChaosClothAssetEditor/Source/ChaosClothAssetEditor/Private/ChaosClothAsset/SClothEditorRestSpaceViewport.h b/Engine/Plugins/ChaosClothAssetEditor/Source/ChaosClothAssetEditor/Private/ChaosClothAsset/SClothEditorRestSpaceViewport.h
--- a/Engine/Plugins/ChaosClothAssetEditor/Source/ChaosClothAssetEditor/Private/ChaosClothAsset/SClothEditorRestSpaceViewport.h
+++ b/Engine/Plugins/ChaosClothAssetEditor/Source/ChaosClothAssetEditor/Private/ChaosClothAsset/SClothEditorRestSpaceViewport.h
@@ -5,11 +5,14 @@
 #include "CoreMinimal.h"
 #include "SBaseCharacterFXEditorViewport.h"
 #include "SCommonEditorViewportToolbarBase.h"
+#include "SClothEditorRestSpaceViewportToolBar.h"
+
+#define UE_API CHAOSCLOTHASSETEDITOR_API
 
 class UChaosClothAssetEditorMode;
 namespace UE::Chaos::ClothAsset { class FChaosClothEditorRestSpaceViewportClient; }
 
-class CHAOSCLOTHASSETEDITOR_API SChaosClothAssetEditorRestSpaceViewport : public SBaseCharacterFXEditorViewport, public ICommonEditorViewportToolbarInfoProvider
+class SChaosClothAssetEditorRestSpaceViewport : public SBaseCharacterFXEditorViewport, public ICommonEditorViewportToolbarInfoProvider
 {
 public:
 
@@ -18,26 +21,29 @@ public:
 		SLATE_ARGUMENT(TSharedPtr<UE::Chaos::ClothAsset::FChaosClothEditorRestSpaceViewportClient>, RestSpaceViewportClient)
 	SLATE_END_ARGS()
 
-	void Construct(const FArguments& InArgs, const FAssetEditorViewportConstructionArgs& InViewportConstructionArgs);
+	UE_API void Construct(const FArguments& InArgs, const FAssetEditorViewportConstructionArgs& InViewportConstructionArgs);
 
 	// SEditorViewport
-	virtual void BindCommands() override;
-	virtual TSharedPtr<SWidget> BuildViewportToolbar() override;
-	virtual void PopulateViewportOverlays(TSharedRef<SOverlay> Overlay) override;
-	virtual void OnFocusViewportToSelection() override;
-	virtual bool IsVisible() const override;
+	UE_API virtual void BindCommands() override;
+	UE_API virtual TSharedPtr<SWidget> MakeViewportToolbar() override;
+	UE_API virtual TSharedPtr<SWidget> BuildViewportToolbar() override;
+	UE_API virtual void PopulateViewportOverlays(TSharedRef<SOverlay> Overlay) override;
+	UE_API virtual void OnFocusViewportToSelection() override;
+	UE_API virtual bool IsVisible() const override;
 
 	// ICommonEditorViewportToolbarInfoProvider
-	virtual TSharedRef<class SEditorViewport> GetViewportWidget() override;
-	virtual TSharedPtr<FExtender> GetExtenders() const override;
-	virtual void OnFloatingButtonClicked() override;
+	UE_API virtual TSharedRef<class SEditorViewport> GetViewportWidget() override;
+	UE_API virtual TSharedPtr<FExtender> GetExtenders() const override;
+	UE_API virtual void OnFloatingButtonClicked() override;
 
 private:
 
-	UChaosClothAssetEditorMode* GetEdMode() const;
+	UE_API UChaosClothAssetEditorMode* GetEdMode() const;
 	
-	FText GetDisplayString() const; 
+	UE_API FText GetDisplayString() const; 
 
 	TSharedPtr<UE::Chaos::ClothAsset::FChaosClothEditorRestSpaceViewportClient> RestSpaceViewportClient;
 
 };
+
+#undef UE_API
diff --git a/Engine/Plugins/ChaosClothAssetEditor/Source/ChaosClothAssetEditor/Private/ChaosClothAsset/SClothEditorRestSpaceViewportToolBar.h b/Engine/Plugins/ChaosClothAssetEditor/Source/ChaosClothAssetEditor/Private/ChaosClothAsset/SClothEditorRestSpaceViewportToolBar.h
new file mode 100644
--- /dev/null
+++ b/Engine/Plugins/ChaosClothAssetEditor/Source/ChaosClothAssetEditor/Private/ChaosClothAsset/SClothEditorRestSpaceViewportToolBar.h
@@ -0,0 +1,56 @@
+// Copyright Epic Games, Inc. All Rights Reserved.
+
+#pragma once
+
+#include "SClothEditorViewportToolBarBase.h"
+
+#define UE_API CHAOSCLOTHASSETEDITOR_API
+
+class FExtender;
+class FUICommandList;
+class SChaosClothAssetEditorRestSpaceViewport;
+namespace UE::Chaos::ClothAsset { class FChaosClothEditorRestSpaceViewportClient; }
+
+/**
+ * Toolbar that shows up at the top of the rest space viewport
+ */
+class SChaosClothAssetEditorRestSpaceViewportToolBar : public SChaosClothAssetEditorViewportToolBarBase
+{
+public:
+	SLATE_BEGIN_ARGS(SChaosClothAssetEditorRestSpaceViewportToolBar) {}
+		SLATE_ARGUMENT(TSharedPtr<FUICommandList>, CommandList)
+		SLATE_ARGUMENT(TSharedPtr<FExtender>, Extenders)
+		SLATE_ARGUMENT(TSharedPtr<UE::Chaos::ClothAsset::FChaosClothEditorRestSpaceViewportClient>, RestSpaceViewportClient)
+	SLATE_END_ARGS()
+
+	UE_API void Construct(const FArguments& InArgs, TSharedPtr<class SChaosClothAssetEditorRestSpaceViewport> InChaosClothAssetEditorViewport);
+
+private:
+
+	UE_API float GetCameraPointLightIntensity() const;
+	UE_API void CameraPointLightIntensityChanged(float NewValue);
+	UE_API bool IsRenderModeEnabled() const;
+
+	UE_API TSharedRef<SWidget> GenerateLightMenu();
+	UE_API TSharedRef<SWidget> GenerateClothRestSpaceViewportOptionsMenu();
+
+	UE_API TSharedRef<SWidget> MakeOptionsMenu();
+	UE_API TSharedRef<SWidget> MakeDisplayToolBar(const TSharedPtr<FExtender> InExtenders);
+	UE_API TSharedRef<SWidget> MakeToolBar(const TSharedPtr<FExtender> InExtenders);
+
+	UE_API FText GetViewModeMenuLabel() const;
+	UE_API const FSlateBrush* GetViewModeMenuLabelIcon() const;
+	UE_API TSharedRef<SWidget> GenerateViewModeMenuContent() const;
+
+	UE_API void RegisterViewModeMenuContent();
+
+	/** The viewport that we are in */
+	TWeakPtr<class SChaosClothAssetEditorRestSpaceViewport> ChaosClothAssetEditorRestSpaceViewportPtr;
+
+	/** Client associated with the viewport we are in */
+	TSharedPtr<UE::Chaos::ClothAsset::FChaosClothEditorRestSpaceViewportClient> RestSpaceViewportClient;
+
+	TSharedPtr<FUICommandList> CommandList;
+};
+
+#undef UE_API
diff --git a/Engine/Plugins/ChaosClothAssetEditor/Source/ChaosClothAssetEditor/Private/ChaosClothAsset/SClothEditorViewportToolBarBase.h b/Engine/Plugins/ChaosClothAssetEditor/Source/ChaosClothAssetEditor/Private/ChaosClothAsset/SClothEditorViewportToolBarBase.h
new file mode 100644
--- /dev/null
+++ b/Engine/Plugins/ChaosClothAssetEditor/Source/ChaosClothAssetEditor/Private/ChaosClothAsset/SClothEditorViewportToolBarBase.h
@@ -0,0 +1,24 @@
+// Copyright Epic Games, Inc. All Rights Reserved.
+
+#pragma once
+
+#include "SCommonEditorViewportToolbarBase.h"
+
+#define UE_API CHAOSCLOTHASSETEDITOR_API
+
+/// Common code for toolbars in the Cloth Editor
+class SChaosClothAssetEditorViewportToolBarBase : public SCommonEditorViewportToolbarBase
+{
+protected:
+	UE_API TSharedRef<SWidget> GenerateClothViewportOptionsMenu() const;
+	UE_API TSharedRef<SWidget> GenerateCameraSpeedSettingsMenu() const;
+	UE_API float OnGetFOVValue() const;
+	UE_API FText GetCameraSpeedLabel() const;
+	UE_API float GetCamSpeedSliderPosition() const;
+	UE_API void OnSetCamSpeed(float NewValue) const;
+	UE_API FText GetCameraSpeedScalarLabel() const;
+	UE_API float GetCamSpeedScalarBoxValue() const;
+	UE_API void OnSetCamSpeedScalarBoxValue(float NewValue) const;
+};
+
+#undef UE_API
diff --git a/Engine/Plugins/ChaosClothAssetEditor/Source/ChaosClothAssetEditorTools/Private/ChaosClothAsset/ClothMeshSelectionTool.h b/Engine/Plugins/ChaosClothAssetEditor/Source/ChaosClothAssetEditorTools/Private/ChaosClothAsset/ClothMeshSelectionTool.h
--- a/Engine/Plugins/ChaosClothAssetEditor/Source/ChaosClothAssetEditorTools/Private/ChaosClothAsset/ClothMeshSelectionTool.h
+++ b/Engine/Plugins/ChaosClothAssetEditor/Source/ChaosClothAssetEditorTools/Private/ChaosClothAsset/ClothMeshSelectionTool.h
@@ -6,6 +6,8 @@
 #include "Selection/PolygonSelectionMechanic.h"
 #include "ClothMeshSelectionTool.generated.h"
 
+#define UE_API CHAOSCLOTHASSETEDITORTOOLS_API
+
 class UPolygonSelectionMechanic;
 class UDataflowContextObject;
 class UPreviewMesh;
@@ -41,8 +43,8 @@ private:
 };
 
 
-UCLASS()
-class CHAOSCLOTHASSETEDITORTOOLS_API UClothMeshSelectionToolActions :  public UInteractiveToolPropertySet
+UCLASS(MinimalAPI)
+class UClothMeshSelectionToolActions :  public UInteractiveToolPropertySet
 {
 	GENERATED_BODY()
 
@@ -52,7 +54,7 @@ public:
 
 	void Initialize(UClothMeshSelectionTool* ParentToolIn) { ParentTool = ParentToolIn; }
 
-	void PostAction(EClothMeshSelectionToolActions Action);
+	UE_API void PostAction(EClothMeshSelectionToolActions Action);
 
 	UFUNCTION(CallInEditor, Category = Selection)
 	void GrowSelection()
@@ -80,8 +82,8 @@ public:
 
 };
 
-UCLASS()
-class CHAOSCLOTHASSETEDITORTOOLS_API UClothMeshSelectionToolProperties : public UInteractiveToolPropertySet
+UCLASS(MinimalAPI)
+class UClothMeshSelectionToolProperties : public UInteractiveToolPropertySet
 {
 	GENERATED_BODY()
 
@@ -101,12 +103,12 @@ public:
 
 private:
 
-	virtual void PostEditChangeProperty(FPropertyChangedEvent& PropertyChangedEvent) override;
+	UE_API virtual void PostEditChangeProperty(FPropertyChangedEvent& PropertyChangedEvent) override;
 
 };
 
-UCLASS()
-class CHAOSCLOTHASSETEDITORTOOLS_API UClothMeshSelectionTool : public USingleSelectionMeshEditingTool
+UCLASS(MinimalAPI)
+class UClothMeshSelectionTool : public USingleSelectionMeshEditingTool
 {
 	GENERATED_BODY()
 
@@ -114,24 +116,24 @@ private:
 
 	friend class UClothMeshSelectionToolBuilder;
 
-	virtual void Setup() override;
-	virtual void OnShutdown(EToolShutdownType ShutdownType) override;
-	virtual void Render(IToolsContextRenderAPI* RenderAPI) override;
-	virtual void DrawHUD(FCanvas* Canvas, IToolsContextRenderAPI* RenderAPI) override;
+	UE_API virtual void Setup() override;
+	UE_API virtual void OnShutdown(EToolShutdownType ShutdownType) override;
+	UE_API virtual void Render(IToolsContextRenderAPI* RenderAPI) override;
+	UE_API virtual void DrawHUD(FCanvas* Canvas, IToolsContextRenderAPI* RenderAPI) override;
 
-	virtual void OnTick(float DeltaTime) override;
+	UE_API virtual void OnTick(float DeltaTime) override;
 
 	virtual bool HasCancel() const override { return true; }
 	virtual bool HasAccept() const override { return true; }
-	virtual bool CanAccept() const override;
+	UE_API virtual bool CanAccept() const override;
 
 	// IInteractiveToolCameraFocusAPI implementation
-	virtual FBox GetWorldSpaceFocusBox() override;
+	UE_API virtual FBox GetWorldSpaceFocusBox() override;
 
-	void SetDataflowContextObject(TObjectPtr<UDataflowContextObject> InDataflowContextObject);
+	UE_API void SetDataflowContextObject(TObjectPtr<UDataflowContextObject> InDataflowContextObject);
 
-	bool GetSelectedNodeInfo(FString& OutMapName, UE::Geometry::FGroupTopologySelection& OutSelection, EChaosClothAssetSelectionOverrideType& OutOverrideType);
-	void UpdateSelectedNode();
+	UE_API bool GetSelectedNodeInfo(FString& OutMapName, UE::Geometry::FGroupTopologySelection& OutSelection, EChaosClothAssetSelectionOverrideType& OutOverrideType);
+	UE_API void UpdateSelectedNode();
 
 	UPROPERTY()
 	TObjectPtr<UClothMeshSelectionToolProperties> ToolProperties;
@@ -158,9 +160,9 @@ private:
 
 public:
 
-	virtual void RequestAction(EClothMeshSelectionToolActions ActionType);
+	UE_API virtual void RequestAction(EClothMeshSelectionToolActions ActionType);
 
-	void NotifyTargetChanged();
+	UE_API void NotifyTargetChanged();
 
 	UPROPERTY()
 	TObjectPtr<UClothMeshSelectionToolActions> ActionsProps;
@@ -168,8 +170,9 @@ public:
 private:
 	bool bHavePendingAction = false;
 	EClothMeshSelectionToolActions PendingAction;
-	virtual void ApplyAction(EClothMeshSelectionToolActions ActionType);
+	UE_API virtual void ApplyAction(EClothMeshSelectionToolActions ActionType);
 
-	void InitializeSculptMeshFromTarget();
+	UE_API void InitializeSculptMeshFromTarget();
 };
 
+#undef UE_API
diff --git a/Engine/Plugins/ChaosClothAssetEditor/Source/ChaosClothAssetEditorTools/Private/ChaosClothAsset/ClothTransferSkinWeightsTool.h b/Engine/Plugins/ChaosClothAssetEditor/Source/ChaosClothAssetEditorTools/Private/ChaosClothAsset/ClothTransferSkinWeightsTool.h
--- a/Engine/Plugins/ChaosClothAssetEditor/Source/ChaosClothAssetEditorTools/Private/ChaosClothAsset/ClothTransferSkinWeightsTool.h
+++ b/Engine/Plugins/ChaosClothAssetEditor/Source/ChaosClothAssetEditorTools/Private/ChaosClothAsset/ClothTransferSkinWeightsTool.h
@@ -7,6 +7,8 @@
 #include "Transforms/TransformGizmoDataBinder.h"
 #include "ClothTransferSkinWeightsTool.generated.h"
 
+#define UE_API CHAOSCLOTHASSETEDITORTOOLS_API
+
 
 class UClothTransferSkinWeightsTool;
 class USkeletalMesh;
@@ -19,8 +21,8 @@ class USkeletalMeshComponent;
 class FTransformGizmoDataBinder;
 struct FChaosClothAssetTransferSkinWeightsNode;
 
-UCLASS()
-class CHAOSCLOTHASSETEDITORTOOLS_API UClothTransferSkinWeightsToolProperties : public UInteractiveToolPropertySet
+UCLASS(MinimalAPI)
+class UClothTransferSkinWeightsToolProperties : public UInteractiveToolPropertySet
 {
 	GENERATED_BODY()
 public:
@@ -41,8 +43,8 @@ public:
 	bool bHideSourceMesh = false;
 };
 
-UCLASS()
-class CHAOSCLOTHASSETEDITORTOOLS_API UClothTransferSkinWeightsTool : public USingleSelectionMeshEditingTool, public UE::Geometry::IDynamicMeshOperatorFactory
+UCLASS(MinimalAPI)
+class UClothTransferSkinWeightsTool : public USingleSelectionMeshEditingTool, public UE::Geometry::IDynamicMeshOperatorFactory
 {
 	GENERATED_BODY()
 
@@ -51,26 +53,26 @@ private:
 	friend class UClothTransferSkinWeightsToolBuilder;
 
 	// UInteractiveTool
-	virtual void Setup() override;
-	virtual void Shutdown(EToolShutdownType ShutdownType) override;
+	UE_API virtual void Setup() override;
+	UE_API virtual void Shutdown(EToolShutdownType ShutdownType) override;
 	virtual bool HasAccept() const override { return true; }
 	virtual bool HasCancel() const override { return true; }
-	virtual bool CanAccept() const override;
-	virtual void OnTick(float DeltaTime) override;
+	UE_API virtual bool CanAccept() const override;
+	UE_API virtual void OnTick(float DeltaTime) override;
 
 	// IDynamicMeshOperatorFactory
-	virtual TUniquePtr<UE::Geometry::FDynamicMeshOperator> MakeNewOperator() override;
+	UE_API virtual TUniquePtr<UE::Geometry::FDynamicMeshOperator> MakeNewOperator() override;
 	
-	void SetDataflowEditorContextObject(TObjectPtr<UDataflowContextObject> InDataflowContextObject);
+	UE_API void SetDataflowEditorContextObject(TObjectPtr<UDataflowContextObject> InDataflowContextObject);
 
-	FTransform TransformFromProperties() const;
-	void SetSRTPropertiesFromTransform(const FTransform& Transform) const;
+	UE_API FTransform TransformFromProperties() const;
+	UE_API void SetSRTPropertiesFromTransform(const FTransform& Transform) const;
 
-	void UpdateSourceMesh(TObjectPtr<USkeletalMesh> Mesh);
+	UE_API void UpdateSourceMesh(TObjectPtr<USkeletalMesh> Mesh);
 
-	void OpFinishedCallback(const UE::Geometry::FDynamicMeshOperator* Op);
+	UE_API void OpFinishedCallback(const UE::Geometry::FDynamicMeshOperator* Op);
 
-	void PreviewMeshUpdatedCallback(UMeshOpPreviewWithBackgroundCompute* Preview);
+	UE_API void PreviewMeshUpdatedCallback(UMeshOpPreviewWithBackgroundCompute* Preview);
 
 
 	UPROPERTY(Transient)
@@ -103,3 +105,4 @@ private:
 };
 
 
+#undef UE_API
diff --git a/Engine/Plugins/ChaosClothAssetEditor/Source/ChaosClothAssetEditorTools/Private/ChaosClothAsset/ClothWeightMapPaintBrushOps.h b/Engine/Plugins/ChaosClothAssetEditor/Source/ChaosClothAssetEditorTools/Private/ChaosClothAsset/ClothWeightMapPaintBrushOps.h
--- a/Engine/Plugins/ChaosClothAssetEditor/Source/ChaosClothAssetEditorTools/Private/ChaosClothAsset/ClothWeightMapPaintBrushOps.h
+++ b/Engine/Plugins/ChaosClothAssetEditor/Source/ChaosClothAssetEditorTools/Private/ChaosClothAsset/ClothWeightMapPaintBrushOps.h
@@ -8,7 +8,7 @@
 
 
 
-class CHAOSCLOTHASSETEDITORTOOLS_API FMeshVertexWeightMapEditBrushOp : public FMeshSculptBrushOp
+class FMeshVertexWeightMapEditBrushOp : public FMeshSculptBrushOp
 {
 public:
 
@@ -40,8 +40,8 @@ public:
 // Erase Brush
 // 
 
-UCLASS()
-class CHAOSCLOTHASSETEDITORTOOLS_API UWeightMapEraseBrushOpProps : public UMeshSculptBrushOpProps
+UCLASS(MinimalAPI)
+class UWeightMapEraseBrushOpProps : public UMeshSculptBrushOpProps
 {
 	GENERATED_BODY()
 public:
@@ -74,8 +74,8 @@ public:
 // Paint Brush
 // 
 
-UCLASS()
-class CHAOSCLOTHASSETEDITORTOOLS_API UWeightMapPaintBrushOpProps : public UMeshSculptBrushOpProps
+UCLASS(MinimalAPI)
+class UWeightMapPaintBrushOpProps : public UMeshSculptBrushOpProps
 {
 	GENERATED_BODY()
 
@@ -119,8 +119,8 @@ public:
 // Smooth Brush
 // 
 
-UCLASS()
-class CHAOSCLOTHASSETEDITORTOOLS_API UWeightMapSmoothBrushOpProps : public UMeshSculptBrushOpProps
+UCLASS(MinimalAPI)
+class UWeightMapSmoothBrushOpProps : public UMeshSculptBrushOpProps
 {
 	GENERATED_BODY()
 
diff --git a/Engine/Plugins/ChaosClothAssetEditor/Source/ChaosClothAssetEditorTools/Private/ChaosClothAsset/ClothWeightMapPaintTool.h b/Engine/Plugins/ChaosClothAssetEditor/Source/ChaosClothAssetEditorTools/Private/ChaosClothAsset/ClothWeightMapPaintTool.h
--- a/Engine/Plugins/ChaosClothAssetEditor/Source/ChaosClothAssetEditorTools/Private/ChaosClothAsset/ClothWeightMapPaintTool.h
+++ b/Engine/Plugins/ChaosClothAssetEditor/Source/ChaosClothAssetEditorTools/Private/ChaosClothAsset/ClothWeightMapPaintTool.h
@@ -26,6 +26,8 @@
 
 #include "ClothWeightMapPaintTool.generated.h"
 
+#define UE_API CHAOSCLOTHASSETEDITORTOOLS_API
+
 class UMeshElementsVisualizer;
 class UWeightMapEraseBrushOpProps;
 class UWeightMapPaintBrushOpProps;
@@ -125,8 +127,8 @@ enum class EClothEditorWeightMapPaintQueryType : uint8
 
 // TODO: Look at EditConditions for all these properties. Which ones make sense for which SubToolType?
 
-UCLASS()
-class CHAOSCLOTHASSETEDITORTOOLS_API UClothEditorWeightMapPaintBrushFilterProperties : public UInteractiveToolPropertySet
+UCLASS(MinimalAPI)
+class UClothEditorWeightMapPaintBrushFilterProperties : public UInteractiveToolPropertySet
 {
 	GENERATED_BODY()
 
@@ -232,8 +234,8 @@ enum class EClothEditorWeightMapPaintToolActions
 };
 
 
-UCLASS()
-class CHAOSCLOTHASSETEDITORTOOLS_API UClothEditorMeshWeightMapPaintToolActions : public UInteractiveToolPropertySet
+UCLASS(MinimalAPI)
+class UClothEditorMeshWeightMapPaintToolActions : public UInteractiveToolPropertySet
 {
 	GENERATED_BODY()
 
@@ -243,7 +245,7 @@ public:
 
 	void Initialize(UClothEditorWeightMapPaintTool* ParentToolIn) { ParentTool = ParentToolIn; }
 
-	void PostAction(EClothEditorWeightMapPaintToolActions Action);
+	UE_API void PostAction(EClothEditorWeightMapPaintToolActions Action);
 
 	UFUNCTION(CallInEditor, Category = Operations, meta = (DisplayPriority = 10))
 	void ClearAll()
@@ -271,8 +273,8 @@ public:
 
 };
 
-UCLASS()
-class CHAOSCLOTHASSETEDITORTOOLS_API UClothEditorUpdateWeightMapProperties : public UInteractiveToolPropertySet
+UCLASS(MinimalAPI)
+class UClothEditorUpdateWeightMapProperties : public UInteractiveToolPropertySet
 {
 	GENERATED_BODY()
 
@@ -286,12 +288,12 @@ public:
 
 private:
 
-	virtual void PostEditChangeProperty(struct FPropertyChangedEvent& PropertyChangedEvent) override;
+	UE_API virtual void PostEditChangeProperty(struct FPropertyChangedEvent& PropertyChangedEvent) override;
 };
 
 /// Property panel for controlling hiding triangles
-UCLASS()
-class CHAOSCLOTHASSETEDITORTOOLS_API UClothEditorMeshWeightMapPaintToolShowHideProperties : public UInteractiveToolPropertySet
+UCLASS(MinimalAPI)
+class UClothEditorMeshWeightMapPaintToolShowHideProperties : public UInteractiveToolPropertySet
 {
 	GENERATED_BODY()
 
@@ -299,7 +301,7 @@ public:
 
 	TWeakObjectPtr<UClothEditorWeightMapPaintTool> ParentTool;
 	void Initialize(UClothEditorWeightMapPaintTool* ParentToolIn) { ParentTool = ParentToolIn; }
-	void PostAction(EClothEditorWeightMapPaintToolActions Action);
+	UE_API void PostAction(EClothEditorWeightMapPaintToolActions Action);
 
 	/** Toggles whether each pattern is shown or hidden */
 	UPROPERTY(EditAnywhere, EditFixedSize, Transient, Category = TriangleVisibility, meta = (DisplayPriority = 2));
@@ -317,46 +319,46 @@ public:
 /**
  * Mesh Element Paint Tool Class
  */
-UCLASS()
-class CHAOSCLOTHASSETEDITORTOOLS_API UClothEditorWeightMapPaintTool : public UMeshSculptToolBase
+UCLASS(MinimalAPI)
+class UClothEditorWeightMapPaintTool : public UMeshSculptToolBase
 {
 	GENERATED_BODY()
 
 public:
-	virtual void RegisterActions(FInteractiveToolActionSet& ActionSet) override;
+	UE_API virtual void RegisterActions(FInteractiveToolActionSet& ActionSet) override;
 
-	virtual void Setup() override;
-	virtual void Shutdown(EToolShutdownType ShutdownType) override;
-	virtual void DrawHUD(FCanvas* Canvas, IToolsContextRenderAPI* RenderAPI) override;
-	virtual void Render(IToolsContextRenderAPI* RenderAPI) override;
-	virtual void UpdateMaterialMode(EMeshEditingMaterialModes MaterialMode) override;
+	UE_API virtual void Setup() override;
+	UE_API virtual void Shutdown(EToolShutdownType ShutdownType) override;
+	UE_API virtual void DrawHUD(FCanvas* Canvas, IToolsContextRenderAPI* RenderAPI) override;
+	UE_API virtual void Render(IToolsContextRenderAPI* RenderAPI) override;
+	UE_API virtual void UpdateMaterialMode(EMeshEditingMaterialModes MaterialMode) override;
 
-	virtual void UpdateStampPendingState() override;
+	UE_API virtual void UpdateStampPendingState() override;
 
-	virtual void OnTick(float DeltaTime) override;
+	UE_API virtual void OnTick(float DeltaTime) override;
 
 	virtual bool HasCancel() const override { return true; }
 	virtual bool HasAccept() const override { return true; }
-	virtual bool CanAccept() const override;
+	UE_API virtual bool CanAccept() const override;
 
-	virtual bool OnUpdateHover(const FInputDeviceRay& DevicePos) override;
+	UE_API virtual bool OnUpdateHover(const FInputDeviceRay& DevicePos) override;
 
-	virtual void OnPropertyModified(UObject* PropertySet, FProperty* Property) override;
+	UE_API virtual void OnPropertyModified(UObject* PropertySet, FProperty* Property) override;
 
-	bool IsInBrushSubMode() const;
+	UE_API bool IsInBrushSubMode() const;
 
-	virtual void CommitResult(UBaseDynamicMeshComponent* Component, bool bModifiedTopology) override;
+	UE_API virtual void CommitResult(UBaseDynamicMeshComponent* Component, bool bModifiedTopology) override;
 
-	void SetDataflowContextObject(TObjectPtr<UDataflowContextObject> InDataflowContextObject);
+	UE_API void SetDataflowContextObject(TObjectPtr<UDataflowContextObject> InDataflowContextObject);
 
 private:
 
 	// Initialization support
-	void InitializeSculptMeshFromTarget();
-	void UpdateShowHideProperties();
+	UE_API void InitializeSculptMeshFromTarget();
+	UE_API void UpdateShowHideProperties();
 
 	// Make sure things are set up correctly after initialization or re-initialization
-	void PostSetupCheck() const;
+	UE_API void PostSetupCheck() const;
 
 public:
 
@@ -376,24 +378,24 @@ private:
 	UPROPERTY()
 	TObjectPtr<UWeightMapEraseBrushOpProps> EraseBrushOpProperties;
 
-	void FloodFillCurrentWeightAction();
-	void ClearAllWeightsAction();
-	void InvertWeightsAction();
-	void MultiplyWeightsAction();
-	void ClearHiddenAction();
+	UE_API void FloodFillCurrentWeightAction();
+	UE_API void ClearAllWeightsAction();
+	UE_API void InvertWeightsAction();
+	UE_API void MultiplyWeightsAction();
+	UE_API void ClearHiddenAction();
 
 public:
-	void SetVerticesToWeightMap(const TSet<int32>& Vertices, double WeightValue, bool bIsErase);
+	UE_API void SetVerticesToWeightMap(const TSet<int32>& Vertices, double WeightValue, bool bIsErase);
 
-	bool HaveVisibilityFilter() const;
-	void ApplyVisibilityFilter(const TArray<int32>& Vertices, TArray<int32>& VisibleVertices);
-	void ApplyVisibilityFilter(TSet<int32>& Triangles, TArray<int32>& ROIBuffer, TArray<int32>& OutputBuffer);
+	UE_API bool HaveVisibilityFilter() const;
+	UE_API void ApplyVisibilityFilter(const TArray<int32>& Vertices, TArray<int32>& VisibleVertices);
+	UE_API void ApplyVisibilityFilter(TSet<int32>& Triangles, TArray<int32>& ROIBuffer, TArray<int32>& OutputBuffer);
 
 	// we override these so we can update the separate BrushSize property added for this tool
-	virtual void IncreaseBrushRadiusAction();
-	virtual void DecreaseBrushRadiusAction();
-	virtual void IncreaseBrushRadiusSmallStepAction();
-	virtual void DecreaseBrushRadiusSmallStepAction();
+	UE_API virtual void IncreaseBrushRadiusAction();
+	UE_API virtual void DecreaseBrushRadiusAction();
+	UE_API virtual void IncreaseBrushRadiusSmallStepAction();
+	UE_API virtual void DecreaseBrushRadiusSmallStepAction();
 
 protected:
 	// UMeshSculptToolBase API
@@ -401,27 +403,25 @@ protected:
 	virtual FDynamicMesh3* GetBaseMesh() { check(false); return nullptr; }
 	virtual const FDynamicMesh3* GetBaseMesh() const { check(false); return nullptr; }
 
-	virtual int32 FindHitSculptMeshTriangleConst(const FRay3d& LocalRay) const override;
-	virtual int32 FindHitTargetMeshTriangleConst(const FRay3d& LocalRay) const override;
-	virtual void UpdateHitSculptMeshTriangle(int32 TriangleID, const FRay3d& LocalRay) override;
-
+	UE_API virtual int32 FindHitSculptMeshTriangle(const FRay3d& LocalRay) override;
+	UE_API virtual int32 FindHitTargetMeshTriangle(const FRay3d& LocalRay) override;
 
-	virtual void OnBeginStroke(const FRay& WorldRay) override;
-	virtual void OnCancelStroke() override;
-	virtual void OnEndStroke() override;
+	UE_API virtual void OnBeginStroke(const FRay& WorldRay) override;
+	UE_API virtual void OnCancelStroke() override;
+	UE_API virtual void OnEndStroke() override;
 
-	virtual TUniquePtr<FMeshSculptBrushOp>& GetActiveBrushOp();
+	UE_API virtual TUniquePtr<FMeshSculptBrushOp>& GetActiveBrushOp();
 
 	virtual bool SharesBrushPropertiesChanges() const override { return false; }
 
-	virtual void InitializeBrushSizeRange(const UE::Geometry::FAxisAlignedBox3d& TargetBounds) override;
+	UE_API virtual void InitializeBrushSizeRange(const UE::Geometry::FAxisAlignedBox3d& TargetBounds) override;
 
-	virtual void NextBrushModeAction() override;
-	virtual void PreviousBrushModeAction() override;
+	UE_API virtual void NextBrushModeAction() override;
+	UE_API virtual void PreviousBrushModeAction() override;
 
 	// Note: these will actually modify the brush's Attribute Value since we don't use Speed in our brush
-	virtual void IncreaseBrushSpeedAction() override;
-	virtual void DecreaseBrushSpeedAction() override;
+	UE_API virtual void IncreaseBrushSpeedAction() override;
+	UE_API virtual void DecreaseBrushSpeedAction() override;
 
 
 	// end UMeshSculptToolBase API
@@ -433,7 +433,7 @@ protected:
 	//
 
 public:
-	virtual void RequestAction(EClothEditorWeightMapPaintToolActions ActionType);
+	UE_API virtual void RequestAction(EClothEditorWeightMapPaintToolActions ActionType);
 	
 	UPROPERTY()
 	TObjectPtr<UClothEditorMeshWeightMapPaintToolActions> ActionsProps;
@@ -444,7 +444,7 @@ public:
 protected:
 	bool bHavePendingAction = false;
 	EClothEditorWeightMapPaintToolActions PendingAction;
-	virtual void ApplyAction(EClothEditorWeightMapPaintToolActions ActionType);
+	UE_API virtual void ApplyAction(EClothEditorWeightMapPaintToolActions ActionType);
 
 
 
@@ -456,7 +456,7 @@ public:
 	TObjectPtr<UPolyLassoMarqueeMechanic> PolyLassoMechanic;
 
 protected:
-	void OnPolyLassoFinished(const FCameraPolyLasso& Lasso, bool bCanceled);
+	UE_API void OnPolyLassoFinished(const FCameraPolyLasso& Lasso, bool bCanceled);
 
 
 	// 
@@ -474,8 +474,8 @@ protected:
 	UE::Geometry::FGroupTopologySelection LowValueGradientVertexSelection;
 	UE::Geometry::FGroupTopologySelection HighValueGradientVertexSelection;
 
-	void ComputeGradient();
-	void OnSelectionModified();
+	UE_API void ComputeGradient();
+	UE_API void OnSelectionModified();
 
 
 	//
@@ -518,44 +518,44 @@ protected:
 	TObjectPtr<UDataflowContextObject> DataflowContextObject = nullptr;
 
 	// realtime visualization
-	void OnDynamicMeshComponentChanged(UDynamicMeshComponent* Component, const FMeshVertexChange* Change, bool bRevert);
+	UE_API void OnDynamicMeshComponentChanged(UDynamicMeshComponent* Component, const FMeshVertexChange* Change, bool bRevert);
 	FDelegateHandle OnDynamicMeshComponentChangedHandle;
 
 	TConstArrayView<float> InputWeightMap;
 
 	UE::Geometry::FDynamicMeshWeightAttribute* ActiveWeightMap;
-	double GetCurrentWeightValue(int32 VertexId) const;
-	double GetCurrentWeightValueUnderBrush() const;
+	UE_API double GetCurrentWeightValue(int32 VertexId) const;
+	UE_API double GetCurrentWeightValueUnderBrush() const;
 	FVector3d CurrentBaryCentricCoords;
-	int32 GetBrushNearestVertex() const;
-	int32 GetBrushNearestVertexAccurate() const;
+	UE_API int32 GetBrushNearestVertex() const;
+	UE_API int32 GetBrushNearestVertexAccurate() const;
 
-	void GetCurrentWeightMap(TArray<float>& OutWeights) const;
+	UE_API void GetCurrentWeightMap(TArray<float>& OutWeights) const;
 
-	void UpdateSubToolType(EClothEditorWeightMapPaintInteractionType NewType);
+	UE_API void UpdateSubToolType(EClothEditorWeightMapPaintInteractionType NewType);
 
-	void UpdateBrushType(EClothEditorWeightMapPaintBrushType BrushType);
+	UE_API void UpdateBrushType(EClothEditorWeightMapPaintBrushType BrushType);
 
 	TSet<int32> AccumulatedTriangleROI;
 	bool bUndoUpdatePending = false;
 	TArray<int> NormalsBuffer;
-	void WaitForPendingUndoRedoUpdate();
+	UE_API void WaitForPendingUndoRedo();
 
 	TArray<int> TempROIBuffer;
 	TArray<int> VertexROI;
 	TArray<bool> VisibilityFilterBuffer;
 	TSet<int> VertexSetBuffer;
 	TSet<int> TriangleROI;
-	void UpdateROI(const FSculptBrushStamp& CurrentStamp);
+	UE_API void UpdateROI(const FSculptBrushStamp& CurrentStamp);
 
 	EClothEditorWeightMapPaintBrushType PendingStampType = EClothEditorWeightMapPaintBrushType::Paint;
 
-	bool UpdateStampPosition(const FRay& WorldRay);
-	bool ApplyStamp();
+	UE_API bool UpdateStampPosition(const FRay& WorldRay);
+	UE_API bool ApplyStamp();
 
 	UE::Geometry::FDynamicMeshOctree3 Octree;
 
-	bool UpdateBrushPosition(const FRay& WorldRay);
+	UE_API bool UpdateBrushPosition(const FRay& WorldRay);
 
 	bool GetInEraseStroke()
 	{
@@ -568,8 +568,8 @@ protected:
 
 	TArray<int32> ROITriangleBuffer;
 	TArray<double> ROIWeightValueBuffer;
-	bool SyncMeshWithWeightBuffer(FDynamicMesh3* Mesh);
-	bool SyncWeightBufferWithMesh(const FDynamicMesh3* Mesh);
+	UE_API bool SyncMeshWithWeightBuffer(FDynamicMesh3* Mesh);
+	UE_API bool SyncWeightBufferWithMesh(const FDynamicMesh3* Mesh);
 
 
 	// Undo/Redo change support
@@ -583,24 +583,24 @@ protected:
 		}
 	};
 	TUniquePtr<TIndexedValuesChangeBuilder<float, FNodeBufferWeightChange>> ActiveChangeBuilder;
-	void BeginChange();
-	void EndChange();
+	UE_API void BeginChange();
+	UE_API void EndChange();
 
 	// The corresponding FChaosClothAssetWeightMapNode has a buffer of scalar weights. Depending on bHaveDynamicMeshToWeightConversion, the 
 	// indexing of this buffer might be different than the mesh vertex indexing. This function returns the corresponding index in the node buffer for
 	// a given mesh vertex index.
-	int32 MeshIndexToNodeIndex(int32 MeshVertexIndex) const;
+	UE_API int32 MeshIndexToNodeIndex(int32 MeshVertexIndex) const;
 
 	// Update the active weight map attribute from the map of indices/values. The inputs are given in "node buffer format" (see function above).
-	void UpdateMapValuesFromNodeValues(const TArray<int32>& Indices, const TArray<float>& Values);
+	UE_API void UpdateMapValuesFromNodeValues(const TArray<int32>& Indices, const TArray<float>& Values);
 
 
-	FColor GetColorForWeightValue(double WeightValue);
+	UE_API FColor GetColorForWeightValue(double WeightValue);
 
 	TArray<FVector3d> TriNormals;
 	TArray<int32> UVSeamEdges;
 	TArray<int32> NormalSeamEdges;
-	void PrecomputeFilterData();
+	UE_API void PrecomputeFilterData();
 
 	// DynamicMesh might be unwelded mesh, but weights are on the welded mesh.
 	bool bHaveDynamicMeshToWeightConversion = false;
@@ -611,7 +611,7 @@ protected:
 	virtual bool ShowWorkPlane() const override { return false; }
 
 	friend class UClothEditorWeightMapPaintToolBuilder;
-	void NotifyTargetChanged();
+	UE_API void NotifyTargetChanged();
 
 	bool bAnyChangeMade = false;
 
@@ -619,11 +619,12 @@ protected:
 
 	FChaosClothAssetWeightMapNode* WeightMapNodeToUpdate = nullptr;
 
-	void UpdateSelectedNode();
+	UE_API void UpdateSelectedNode();
 
-	void UpdateVertexColorOverlay(const TSet<int>* TrianglesToUpdate = nullptr);
+	UE_API void UpdateVertexColorOverlay(const TSet<int>* TrianglesToUpdate = nullptr);
 
 };
 
 
 
+#undef UE_API
diff --git a/Engine/Plugins/ChaosClothAssetEditor/Source/ChaosClothAssetEditorTools/Public/ChaosClothAsset/ClothEditorToolBuilders.h b/Engine/Plugins/ChaosClothAssetEditor/Source/ChaosClothAssetEditorTools/Public/ChaosClothAsset/ClothEditorToolBuilders.h
--- a/Engine/Plugins/ChaosClothAssetEditor/Source/ChaosClothAssetEditorTools/Public/ChaosClothAsset/ClothEditorToolBuilders.h
+++ b/Engine/Plugins/ChaosClothAssetEditor/Source/ChaosClothAssetEditorTools/Public/ChaosClothAsset/ClothEditorToolBuilders.h
@@ -10,6 +10,8 @@
 
 #include "ClothEditorToolBuilders.generated.h"
 
+#define UE_API CHAOSCLOTHASSETEDITORTOOLS_API
+
 class UDataflowContextObject;
 #if UE_ENABLE_INCLUDE_ORDER_DEPRECATED_IN_5_5
 class UClothEditorContextObject;
@@ -50,29 +52,29 @@ public:
 };
 
 
-UCLASS()
-class CHAOSCLOTHASSETEDITORTOOLS_API UClothEditorWeightMapPaintToolBuilder : public UMeshSurfacePointMeshEditingToolBuilder, public IChaosClothAssetEditorToolBuilder, public IDataflowEditorToolBuilder
+UCLASS(MinimalAPI)
+class UClothEditorWeightMapPaintToolBuilder : public UMeshSurfacePointMeshEditingToolBuilder, public IChaosClothAssetEditorToolBuilder, public IDataflowEditorToolBuilder
 {
 	GENERATED_BODY()
 
 private:
 
 	// IDataflowEditorToolBuilder
-	virtual void GetSupportedConstructionViewModes(const UDataflowContextObject& ContextObject, TArray<const UE::Dataflow::IDataflowConstructionViewMode*>& Modes) const override;
-	virtual bool CanSceneStateChange(const UInteractiveTool* ActiveTool, const FToolBuilderState& SceneState) const override;
-	virtual void SceneStateChanged(UInteractiveTool* ActiveTool, const FToolBuilderState& SceneState) override;
+	UE_API virtual void GetSupportedConstructionViewModes(const UDataflowContextObject& ContextObject, TArray<const UE::Dataflow::IDataflowConstructionViewMode*>& Modes) const override;
+	UE_API virtual bool CanSceneStateChange(const UInteractiveTool* ActiveTool, const FToolBuilderState& SceneState) const override;
+	UE_API virtual void SceneStateChanged(UInteractiveTool* ActiveTool, const FToolBuilderState& SceneState) override;
 
 	// IChaosClothAssetEditorToolBuilder
-	virtual void GetSupportedViewModes(const UDataflowContextObject& ContextObject, TArray<UE::Chaos::ClothAsset::EClothPatternVertexType>& Modes) const override;
+	UE_API virtual void GetSupportedViewModes(const UDataflowContextObject& ContextObject, TArray<UE::Chaos::ClothAsset::EClothPatternVertexType>& Modes) const override;
 	virtual bool CanSetConstructionViewWireframeActive() const { return false; }
 
 	// UMeshSurfacePointMeshEditingToolBuilder
-	virtual UMeshSurfacePointTool* CreateNewTool(const FToolBuilderState& SceneState) const override;
+	UE_API virtual UMeshSurfacePointTool* CreateNewTool(const FToolBuilderState& SceneState) const override;
 };
 
 
-UCLASS()
-class CHAOSCLOTHASSETEDITORTOOLS_API UClothMeshSelectionToolBuilder : public UInteractiveToolWithToolTargetsBuilder, public IChaosClothAssetEditorToolBuilder, public IDataflowEditorToolBuilder
+UCLASS(MinimalAPI)
+class UClothMeshSelectionToolBuilder : public UInteractiveToolWithToolTargetsBuilder, public IChaosClothAssetEditorToolBuilder, public IDataflowEditorToolBuilder
 {
 	GENERATED_BODY()
 
@@ -81,37 +83,37 @@ class CHAOSCLOTHASSETEDITORTOOLS_API UClothMeshSelectionToolBuilder : public UIn
 private:
 
 	// IDataflowEditorToolBuilder
-	virtual void GetSupportedConstructionViewModes(const UDataflowContextObject& ContextObject, TArray<const UE::Dataflow::IDataflowConstructionViewMode*>& Modes) const override;
-	virtual bool CanSceneStateChange(const UInteractiveTool* ActiveTool, const FToolBuilderState& SceneState) const override;
-	virtual void SceneStateChanged(UInteractiveTool* ActiveTool, const FToolBuilderState& SceneState) override;
+	UE_API virtual void GetSupportedConstructionViewModes(const UDataflowContextObject& ContextObject, TArray<const UE::Dataflow::IDataflowConstructionViewMode*>& Modes) const override;
+	UE_API virtual bool CanSceneStateChange(const UInteractiveTool* ActiveTool, const FToolBuilderState& SceneState) const override;
+	UE_API virtual void SceneStateChanged(UInteractiveTool* ActiveTool, const FToolBuilderState& SceneState) override;
 
 	// IChaosClothAssetEditorToolBuilder
-	virtual void GetSupportedViewModes(const UDataflowContextObject& ContextObject, TArray<UE::Chaos::ClothAsset::EClothPatternVertexType>& Modes) const override;
+	UE_API virtual void GetSupportedViewModes(const UDataflowContextObject& ContextObject, TArray<UE::Chaos::ClothAsset::EClothPatternVertexType>& Modes) const override;
 	virtual bool CanSetConstructionViewWireframeActive() const { return false; }
 
 	// UInteractiveToolWithToolTargetsBuilder
-	virtual bool CanBuildTool(const FToolBuilderState& SceneState) const override;
-	virtual UInteractiveTool* BuildTool(const FToolBuilderState& SceneState) const override;
-	virtual const FToolTargetTypeRequirements& GetTargetRequirements() const override;
+	UE_API virtual bool CanBuildTool(const FToolBuilderState& SceneState) const override;
+	UE_API virtual UInteractiveTool* BuildTool(const FToolBuilderState& SceneState) const override;
+	UE_API virtual const FToolTargetTypeRequirements& GetTargetRequirements() const override;
 };
 
-UCLASS()
-class CHAOSCLOTHASSETEDITORTOOLS_API UClothTransferSkinWeightsToolBuilder : public USingleSelectionMeshEditingToolBuilder, public IChaosClothAssetEditorToolBuilder, public IDataflowEditorToolBuilder
+UCLASS(MinimalAPI)
+class UClothTransferSkinWeightsToolBuilder : public USingleSelectionMeshEditingToolBuilder, public IChaosClothAssetEditorToolBuilder, public IDataflowEditorToolBuilder
 {
 	GENERATED_BODY()
 
 private:
 
 	// IDataflowEditorToolBuilder
-	virtual void GetSupportedConstructionViewModes(const UDataflowContextObject& ContextObject, TArray<const UE::Dataflow::IDataflowConstructionViewMode*>& Modes) const override;
-	virtual bool CanSceneStateChange(const UInteractiveTool* ActiveTool, const FToolBuilderState& SceneState) const override;
-	virtual void SceneStateChanged(UInteractiveTool* ActiveTool, const FToolBuilderState& SceneState) override;
+	UE_API virtual void GetSupportedConstructionViewModes(const UDataflowContextObject& ContextObject, TArray<const UE::Dataflow::IDataflowConstructionViewMode*>& Modes) const override;
+	UE_API virtual bool CanSceneStateChange(const UInteractiveTool* ActiveTool, const FToolBuilderState& SceneState) const override;
+	UE_API virtual void SceneStateChanged(UInteractiveTool* ActiveTool, const FToolBuilderState& SceneState) override;
 
 	// IChaosClothAssetEditorToolBuilder
-	virtual void GetSupportedViewModes(const UDataflowContextObject& ContextObject, TArray<UE::Chaos::ClothAsset::EClothPatternVertexType>& Modes) const override;
+	UE_API virtual void GetSupportedViewModes(const UDataflowContextObject& ContextObject, TArray<UE::Chaos::ClothAsset::EClothPatternVertexType>& Modes) const override;
 
 	// USingleSelectionMeshEditingToolBuilder
-	virtual USingleSelectionMeshEditingTool* CreateNewTool(const FToolBuilderState& SceneState) const override;
+	UE_API virtual USingleSelectionMeshEditingTool* CreateNewTool(const FToolBuilderState& SceneState) const override;
 };
 
 
@@ -127,3 +129,4 @@ namespace UE::Chaos::ClothAsset
 	FName CHAOSCLOTHASSETEDITORTOOLS_API ClothViewModeToDataflowViewModeName(EClothPatternVertexType ClothViewMode);
 }
 
+#undef UE_API
diff --git a/Engine/Plugins/ChaosClothAssetEditor/Source/ChaosClothAssetTools/Public/ChaosClothAsset/ClothAssetEditorStyle.h b/Engine/Plugins/ChaosClothAssetEditor/Source/ChaosClothAssetTools/Public/ChaosClothAsset/ClothAssetEditorStyle.h
--- a/Engine/Plugins/ChaosClothAssetEditor/Source/ChaosClothAssetTools/Public/ChaosClothAsset/ClothAssetEditorStyle.h
+++ b/Engine/Plugins/ChaosClothAssetEditor/Source/ChaosClothAssetTools/Public/ChaosClothAsset/ClothAssetEditorStyle.h
@@ -4,16 +4,18 @@
 
 #include "Styling/SlateStyle.h"
 
+#define UE_API CHAOSCLOTHASSETTOOLS_API
+
 namespace UE::Chaos::ClothAsset
 {
 	/**
 	 * Editor style setting up the cloth asset icons in editor.
 	 */
-	class CHAOSCLOTHASSETTOOLS_API FClothAssetEditorStyle final : public FSlateStyleSet
+	class FClothAssetEditorStyle final : public FSlateStyleSet
 	{
 	public:
-		FClothAssetEditorStyle();
-		virtual ~FClothAssetEditorStyle() override;
+		UE_API FClothAssetEditorStyle();
+		UE_API virtual ~FClothAssetEditorStyle() override;
 
 	public:
 		static FClothAssetEditorStyle& Get()
@@ -31,6 +33,8 @@ namespace UE::Chaos::ClothAsset
 		}
 
 	private:
-		static TOptional<FClothAssetEditorStyle> Singleton;
+		static UE_API TOptional<FClothAssetEditorStyle> Singleton;
 	};
 }  // End namespace UE::Chaos::ClothAsset
+
+#undef UE_API
diff --git a/Engine/Plugins/ChaosClothAssetEditor/Source/ChaosClothAssetTools/Public/ChaosClothAsset/ClothAssetFactory.h b/Engine/Plugins/ChaosClothAssetEditor/Source/ChaosClothAssetTools/Public/ChaosClothAsset/ClothAssetFactory.h
--- a/Engine/Plugins/ChaosClothAssetEditor/Source/ChaosClothAssetTools/Public/ChaosClothAsset/ClothAssetFactory.h
+++ b/Engine/Plugins/ChaosClothAssetEditor/Source/ChaosClothAssetTools/Public/ChaosClothAsset/ClothAssetFactory.h
@@ -8,21 +8,25 @@
 
 #include "ClothAssetFactory.generated.h"
 
+#define UE_API CHAOSCLOTHASSETTOOLS_API
+
 /**
  * Having a cloth factory allows the cloth asset to be created from the Editor's menus.
  */
-UCLASS(Experimental)
-class CHAOSCLOTHASSETTOOLS_API UChaosClothAssetFactory : public UFactory
+UCLASS(MinimalAPI, Experimental)
+class UChaosClothAssetFactory : public UFactory
 {
 	GENERATED_BODY()
 public:
-	UChaosClothAssetFactory(const FObjectInitializer& ObjectInitializer);
+	UE_API UChaosClothAssetFactory(const FObjectInitializer& ObjectInitializer);
 
 	/** UFactory Interface */
 	virtual bool CanCreateNew() const override { return true; }
 	virtual bool FactoryCanImport(const FString& Filename) override { return false; }
 	virtual bool ShouldShowInNewMenu() const override { return true; }
-	virtual UObject* FactoryCreateNew(UClass* InClass, UObject* InParent, FName InName, EObjectFlags Flags, UObject* Context, FFeedbackContext* Warn) override;
-	virtual FString GetDefaultNewAssetName() const override;
+	UE_API virtual UObject* FactoryCreateNew(UClass* InClass, UObject* InParent, FName InName, EObjectFlags Flags, UObject* Context, FFeedbackContext* Warn) override;
+	UE_API virtual FString GetDefaultNewAssetName() const override;
 	/** End UFactory Interface */
 };
+
+#undef UE_API
diff --git a/Engine/Plugins/ChaosClothAssetEditor/Source/ChaosClothAssetTools/Public/ChaosClothAsset/ClothPatternToDynamicMesh.h b/Engine/Plugins/ChaosClothAssetEditor/Source/ChaosClothAssetTools/Public/ChaosClothAsset/ClothPatternToDynamicMesh.h
--- a/Engine/Plugins/ChaosClothAssetEditor/Source/ChaosClothAssetTools/Public/ChaosClothAsset/ClothPatternToDynamicMesh.h
+++ b/Engine/Plugins/ChaosClothAssetEditor/Source/ChaosClothAssetTools/Public/ChaosClothAsset/ClothPatternToDynamicMesh.h
@@ -5,6 +5,8 @@
 #include "ClothPatternVertexType.h"
 #include "Templates/SharedPointer.h"
 
+#define UE_API CHAOSCLOTHASSETTOOLS_API
+
 
 namespace UE::Geometry { class FDynamicMesh3; }
 class UChaosClothAsset;
@@ -21,11 +23,11 @@ namespace UE::Chaos::ClothAsset
 * Use MaterialOffset to specify an offset into the material group ID, and run a Fix up on any triangles without valid material IDs.
 * When MaterialOffset is INDEX_NONE, no offset is applied and nbo fix up is run. 
 */
-class CHAOSCLOTHASSETTOOLS_API FClothPatternToDynamicMesh
+class FClothPatternToDynamicMesh
 {
 public:
 
-	void Convert(
+	UE_API void Convert(
 		const TSharedRef<const FManagedArrayCollection> ClothCollection,
 		int32 PatternIndex,
 		EClothPatternVertexType VertexDataType,
@@ -33,7 +35,7 @@ public:
 		bool bDisableAttributes = false,
 		int32 MaterialOffset = INDEX_NONE);
 
-	void Convert(
+	UE_API void Convert(
 		const UChaosClothAsset* ClothAssetMeshIn,
 		int32 LODIndex,
 		int32 PatternIndex,
@@ -44,3 +46,4 @@ public:
 
 }	// namespace UE::Chaos::ClothAsset
 
+#undef UE_API
diff --git a/Engine/Plugins/ChaosInsights/Source/ChaosInsightsAnalysis/Public/ChaosInsightsAnalysis/Model/LockRegions.h b/Engine/Plugins/ChaosInsights/Source/ChaosInsightsAnalysis/Public/ChaosInsightsAnalysis/Model/LockRegions.h
--- a/Engine/Plugins/ChaosInsights/Source/ChaosInsightsAnalysis/Public/ChaosInsightsAnalysis/Model/LockRegions.h
+++ b/Engine/Plugins/ChaosInsights/Source/ChaosInsightsAnalysis/Public/ChaosInsightsAnalysis/Model/LockRegions.h
@@ -8,10 +8,12 @@
 #include "TraceServices/Model/AnalysisSession.h"
 #include "UObject/NameTypes.h"
 
+#define UE_API CHAOSINSIGHTSANALYSIS_API
+
 namespace ChaosInsightsAnalysis
 {
 
-	struct CHAOSINSIGHTSANALYSIS_API FLockRegion
+	struct FLockRegion
 	{
 		// Time that the caller attempted to take the lock
 		double BeginTime = std::numeric_limits<double>::infinity();
@@ -41,25 +43,25 @@ namespace ChaosInsightsAnalysis
 		bool bIsWrite = false;
 	};
 
-	class CHAOSINSIGHTSANALYSIS_API FLockRegionLane
+	class FLockRegionLane
 	{
 		friend class FLockRegionProvider;
 
 	public:
-		FLockRegionLane(TraceServices::ILinearAllocator& InAllocator);
+		UE_API FLockRegionLane(TraceServices::ILinearAllocator& InAllocator);
 
-		int32 Num() const;
+		UE_API int32 Num() const;
 
-		bool ForEachRegionInRange(double IntervalStart, double IntervalEnd, TFunctionRef<bool(const FLockRegion&)> Callback) const;
+		UE_API bool ForEachRegionInRange(double IntervalStart, double IntervalEnd, TFunctionRef<bool(const FLockRegion&)> Callback) const;
 
 	private:
 		TraceServices::TPagedArray<FLockRegion> Regions;
 	};
 
-	class CHAOSINSIGHTSANALYSIS_API ILockRegionProvider : public TraceServices::IProvider
+	class ILockRegionProvider : public TraceServices::IProvider
 	{
 	public:
-		virtual ~ILockRegionProvider() override;
+		UE_API virtual ~ILockRegionProvider() override;
 
 		virtual uint64 GetRegionCount() const = 0;
 		virtual int32 GetLaneCount() const = 0;
@@ -88,3 +90,5 @@ namespace ChaosInsightsAnalysis
 	CHAOSINSIGHTSANALYSIS_API IEditableLockRegionProvider& EditRegionProvider(TraceServices::IAnalysisSession& Session);
 
 } // namespace TimingRegionsAnalysis
+
+#undef UE_API
diff --git a/Engine/Plugins/ChaosVD/Source/ChaosVD/Private/Components/ChaosVDSolverCharacterGroundConstraintDataComponent.h b/Engine/Plugins/ChaosVD/Source/ChaosVD/Private/Components/ChaosVDSolverCharacterGroundConstraintDataComponent.h
--- a/Engine/Plugins/ChaosVD/Source/ChaosVD/Private/Components/ChaosVDSolverCharacterGroundConstraintDataComponent.h
+++ b/Engine/Plugins/ChaosVD/Source/ChaosVD/Private/Components/ChaosVDSolverCharacterGroundConstraintDataComponent.h
@@ -5,16 +5,20 @@
 #include "ChaosVDConstraintDataComponent.h"
 #include "ChaosVDSolverCharacterGroundConstraintDataComponent.generated.h"
 
-UCLASS()
-class CHAOSVD_API UChaosVDSolverCharacterGroundConstraintDataComponent : public UChaosVDConstraintDataComponent
+#define UE_API CHAOSVD_API
+
+UCLASS(MinimalAPI)
+class UChaosVDSolverCharacterGroundConstraintDataComponent : public UChaosVDConstraintDataComponent
 {
 	GENERATED_BODY()
 public:
-	virtual void SetScene(const TWeakPtr<FChaosVDScene>& InSceneWeakPtr) override;
+	UE_API virtual void SetScene(const TWeakPtr<FChaosVDScene>& InSceneWeakPtr) override;
 
-	virtual void UpdateFromSolverFrameData(const FChaosVDSolverFrameData& InSolverFrameData) override;
+	UE_API virtual void UpdateFromSolverFrameData(const FChaosVDSolverFrameData& InSolverFrameData) override;
 
-	void HandleSceneUpdated();
+	UE_API void HandleSceneUpdated();
 
-	virtual void BeginDestroy() override;
+	UE_API virtual void BeginDestroy() override;
 };
+
+#undef UE_API
diff --git a/Engine/Plugins/ChaosVD/Source/ChaosVD/Private/Components/ChaosVDSolverJointConstraintDataComponent.h b/Engine/Plugins/ChaosVD/Source/ChaosVD/Private/Components/ChaosVDSolverJointConstraintDataComponent.h
--- a/Engine/Plugins/ChaosVD/Source/ChaosVD/Private/Components/ChaosVDSolverJointConstraintDataComponent.h
+++ b/Engine/Plugins/ChaosVD/Source/ChaosVD/Private/Components/ChaosVDSolverJointConstraintDataComponent.h
@@ -5,11 +5,15 @@
 #include "ChaosVDConstraintDataComponent.h"
 #include "ChaosVDSolverJointConstraintDataComponent.generated.h"
 
+#define UE_API CHAOSVD_API
 
-UCLASS()
-class CHAOSVD_API UChaosVDSolverJointConstraintDataComponent : public UChaosVDConstraintDataComponent
+
+UCLASS(MinimalAPI)
+class UChaosVDSolverJointConstraintDataComponent : public UChaosVDConstraintDataComponent
 {
 	GENERATED_BODY()
 
-	virtual void UpdateFromNewSolverStageData(const FChaosVDSolverFrameData& InSolverFrameData, const FChaosVDFrameStageData& InSolverFrameStageData) override;
+	UE_API virtual void UpdateFromNewSolverStageData(const FChaosVDSolverFrameData& InSolverFrameData, const FChaosVDFrameStageData& InSolverFrameStageData) override;
 };
+
+#undef UE_API
diff --git a/Engine/Plugins/ChaosVD/Source/ChaosVD/Public/ChaosVDObjectDetailsTab.h b/Engine/Plugins/ChaosVD/Source/ChaosVD/Public/ChaosVDObjectDetailsTab.h
--- a/Engine/Plugins/ChaosVD/Source/ChaosVD/Public/ChaosVDObjectDetailsTab.h
+++ b/Engine/Plugins/ChaosVD/Source/ChaosVD/Public/ChaosVDObjectDetailsTab.h
@@ -23,7 +23,7 @@ class FChaosVDStandAloneObjectDetailsTab : public FChaosVDTabSpawnerBase, public
 {
 public:
 
-	CHAOSVD_API FChaosVDStandAloneObjectDetailsTab(const FName& InTabID, TSharedPtr<FTabManager> InTabManager, TWeakPtr<SChaosVDMainTab> InOwningTabWidget) : FChaosVDTabSpawnerBase(InTabID, InTabManager, InOwningTabWidget)
+	FChaosVDStandAloneObjectDetailsTab(const FName& InTabID, TSharedPtr<FTabManager> InTabManager, TWeakPtr<SChaosVDMainTab> InOwningTabWidget) : FChaosVDTabSpawnerBase(InTabID, InTabManager, InOwningTabWidget)
 	{
 	}
 	
@@ -61,7 +61,7 @@ class FChaosVDObjectDetailsTab : public FChaosVDStandAloneObjectDetailsTab, publ
 {
 public:
 
-	CHAOSVD_API FChaosVDObjectDetailsTab(const FName& InTabID, TSharedPtr<FTabManager> InTabManager, TWeakPtr<SChaosVDMainTab> InOwningTabWidget) : FChaosVDStandAloneObjectDetailsTab(InTabID, InTabManager, InOwningTabWidget)
+	FChaosVDObjectDetailsTab(const FName& InTabID, TSharedPtr<FTabManager> InTabManager, TWeakPtr<SChaosVDMainTab> InOwningTabWidget) : FChaosVDStandAloneObjectDetailsTab(InTabID, InTabManager, InOwningTabWidget)
 	{
 	}
 
@@ -78,4 +78,4 @@ protected:
 
 	FDelegateHandle SelectionDelegateHandle;
 	TWeakObjectPtr<UObject> CurrentSelectedObject = nullptr;
-};
\ No newline at end of file
+};
diff --git a/Engine/Plugins/ChaosVD/Source/ChaosVD/Public/ChaosVDTabSpawnerBase.h b/Engine/Plugins/ChaosVD/Source/ChaosVD/Public/ChaosVDTabSpawnerBase.h
--- a/Engine/Plugins/ChaosVD/Source/ChaosVD/Public/ChaosVDTabSpawnerBase.h
+++ b/Engine/Plugins/ChaosVD/Source/ChaosVD/Public/ChaosVDTabSpawnerBase.h
@@ -21,7 +21,7 @@ DECLARE_MULTICAST_DELEGATE_OneParam(FChaosVDTabDestroyed, TSharedRef<SDockTab>)
 class FChaosVDTabSpawnerBase
 {
 public:
-	CHAOSVD_API virtual ~FChaosVDTabSpawnerBase() = default;
+	virtual ~FChaosVDTabSpawnerBase() = default;
 
 	CHAOSVD_API FChaosVDTabSpawnerBase(const FName& InTabID, TSharedPtr<FTabManager> InTabManager, TWeakPtr<SChaosVDMainTab> InOwningTabWidget);
 
@@ -29,7 +29,7 @@ public:
 	 * Handles a spawn request for this tab. It controls how and what contents this tab will have
 	 * @param Args Arguments about the spawn request
 	 */
-	CHAOSVD_API virtual TSharedRef<SDockTab> HandleTabSpawnRequest(const FSpawnTabArgs& Args) = 0;
+	virtual TSharedRef<SDockTab> HandleTabSpawnRequest(const FSpawnTabArgs& Args) = 0;
 
 	/**
 	 * Handles a tab being spawned by this spawner being closed
@@ -44,13 +44,13 @@ public:
 	CHAOSVD_API virtual void HandleTabSpawned(TSharedRef<SDockTab> InTabSpawned);
 
 	/** Event called when the tab is spawned */
-	CHAOSVD_API FChaosVDTabSpawned& OnTabSpawned()
+	FChaosVDTabSpawned& OnTabSpawned()
 	{
 		return TabSpawnedDelegate;
 	}
 	
 	/** Event called when the tab this spawner created is destroyed */
-	CHAOSVD_API FChaosVDTabDestroyed& OnTabDestroyed()
+	FChaosVDTabDestroyed& OnTabDestroyed()
 	{
 		return TabDestroyedDelegate;
 	}
diff --git a/Engine/Plugins/ChaosVD/Source/ChaosVD/Public/ChaosVDTabsIDs.h b/Engine/Plugins/ChaosVD/Source/ChaosVD/Public/ChaosVDTabsIDs.h
--- a/Engine/Plugins/ChaosVD/Source/ChaosVD/Public/ChaosVDTabsIDs.h
+++ b/Engine/Plugins/ChaosVD/Source/ChaosVD/Public/ChaosVDTabsIDs.h
@@ -4,25 +4,29 @@
 
 #include "UObject/NameTypes.h"
 
+#define UE_API CHAOSVD_API
+
 /** Built-in CVD tab IDs */
-class CHAOSVD_API FChaosVDTabID
+class FChaosVDTabID
 {
 public:
 
-	static const FName ChaosVisualDebuggerTab;
-	static const FName PlaybackViewport;
-	static const FName WorldOutliner;
-	static const FName DetailsPanel;
-	static const FName IndependentDetailsPanel1;
-	static const FName IndependentDetailsPanel2;
-	static const FName IndependentDetailsPanel3;
-	static const FName IndependentDetailsPanel4;
-	static const FName OutputLog;
-	static const FName SolversTrack;
-	static const FName StatusBar;
-	static const FName CollisionDataDetails;
-	static const FName SceneQueryDataDetails;
-	static const FName ConstraintsInspector;
-	static const FName SceneQueryBrowser;
-	static const FName RecordedOutputLog;
+	static UE_API const FName ChaosVisualDebuggerTab;
+	static UE_API const FName PlaybackViewport;
+	static UE_API const FName WorldOutliner;
+	static UE_API const FName DetailsPanel;
+	static UE_API const FName IndependentDetailsPanel1;
+	static UE_API const FName IndependentDetailsPanel2;
+	static UE_API const FName IndependentDetailsPanel3;
+	static UE_API const FName IndependentDetailsPanel4;
+	static UE_API const FName OutputLog;
+	static UE_API const FName SolversTrack;
+	static UE_API const FName StatusBar;
+	static UE_API const FName CollisionDataDetails;
+	static UE_API const FName SceneQueryDataDetails;
+	static UE_API const FName ConstraintsInspector;
+	static UE_API const FName SceneQueryBrowser;
+	static UE_API const FName RecordedOutputLog;
 };
+
+#undef UE_API
diff --git a/Engine/Plugins/ChaosVD/Source/ChaosVD/Public/Components/ChaosVDSolverDataComponent.h b/Engine/Plugins/ChaosVD/Source/ChaosVD/Public/Components/ChaosVDSolverDataComponent.h
--- a/Engine/Plugins/ChaosVD/Source/ChaosVD/Public/Components/ChaosVDSolverDataComponent.h
+++ b/Engine/Plugins/ChaosVD/Source/ChaosVD/Public/Components/ChaosVDSolverDataComponent.h
@@ -17,7 +17,7 @@ class UChaosVDSolverDataComponent : public UActorComponent
 {
 	GENERATED_BODY()
 public:
-	CHAOSVD_API virtual void ClearData() PURE_VIRTUAL(UChaosVDSolverDataComponent::ClearData);
+	virtual void ClearData() PURE_VIRTUAL(UChaosVDSolverDataComponent::ClearData);
 	CHAOSVD_API virtual void SetScene(const TWeakPtr<FChaosVDScene>& InSceneWeakPtr);
 
 	virtual void UpdateFromNewGameFrameData(const FChaosVDGameFrameData& InGameFrameData) {}
diff --git a/Engine/Plugins/ChaosVD/Source/ChaosVD/Public/ExtensionsSystem/ChaosVDExtension.h b/Engine/Plugins/ChaosVD/Source/ChaosVD/Public/ExtensionsSystem/ChaosVDExtension.h
--- a/Engine/Plugins/ChaosVD/Source/ChaosVD/Public/ExtensionsSystem/ChaosVDExtension.h
+++ b/Engine/Plugins/ChaosVD/Source/ChaosVD/Public/ExtensionsSystem/ChaosVDExtension.h
@@ -23,9 +23,9 @@ class FChaosVDExtension : public IModularFeature
 {
 public:
 
-	CHAOSVD_API FChaosVDExtension() = default;
+	FChaosVDExtension() = default;
 
-	CHAOSVD_API virtual ~FChaosVDExtension() = default;
+	virtual ~FChaosVDExtension() = default;
 
 	/** FName used as Type within CVD Extension system */
 	FName GetExtensionType() const
diff --git a/Engine/Plugins/ChaosVD/Source/ChaosVD/Public/Interfaces/ChaosVDPooledObject.h b/Engine/Plugins/ChaosVD/Source/ChaosVD/Public/Interfaces/ChaosVDPooledObject.h
--- a/Engine/Plugins/ChaosVD/Source/ChaosVD/Public/Interfaces/ChaosVDPooledObject.h
+++ b/Engine/Plugins/ChaosVD/Source/ChaosVD/Public/Interfaces/ChaosVDPooledObject.h
@@ -19,6 +19,6 @@ class IChaosVDPooledObject
 	GENERATED_BODY()
 
 public:
-	CHAOSVD_API virtual void OnAcquired() PURE_VIRTUAL(IChaosVDPooledObject::OnAcquired)
-	CHAOSVD_API virtual void OnDisposed() PURE_VIRTUAL(IChaosVDPooledObject::OnDisposed)
+	virtual void OnAcquired() PURE_VIRTUAL(IChaosVDPooledObject::OnAcquired)
+	virtual void OnDisposed() PURE_VIRTUAL(IChaosVDPooledObject::OnDisposed)
 };
diff --git a/Engine/Plugins/ChaosVD/Source/ChaosVD/Public/Interfaces/ChaosVDSelectableObject.h b/Engine/Plugins/ChaosVD/Source/ChaosVD/Public/Interfaces/ChaosVDSelectableObject.h
--- a/Engine/Plugins/ChaosVD/Source/ChaosVD/Public/Interfaces/ChaosVDSelectableObject.h
+++ b/Engine/Plugins/ChaosVD/Source/ChaosVD/Public/Interfaces/ChaosVDSelectableObject.h
@@ -19,6 +19,6 @@ class IChaosVDSelectableObject
 	GENERATED_BODY()
 
 public:
-	CHAOSVD_API virtual void HandleSelected() PURE_VIRTUAL(IChaosVDSelectionAware::HandleSelected);
-	CHAOSVD_API virtual void HandleDeSelected() PURE_VIRTUAL(IChaosVDSelectionAware::HandleUnSelected);
+	virtual void HandleSelected() PURE_VIRTUAL(IChaosVDSelectionAware::HandleSelected);
+	virtual void HandleDeSelected() PURE_VIRTUAL(IChaosVDSelectionAware::HandleUnSelected);
 };
diff --git a/Engine/Plugins/ChaosVD/Source/ChaosVD/Public/Settings/ChaosVDCoreSettings.h b/Engine/Plugins/ChaosVD/Source/ChaosVD/Public/Settings/ChaosVDCoreSettings.h
--- a/Engine/Plugins/ChaosVD/Source/ChaosVD/Public/Settings/ChaosVDCoreSettings.h
+++ b/Engine/Plugins/ChaosVD/Source/ChaosVD/Public/Settings/ChaosVDCoreSettings.h
@@ -37,13 +37,13 @@ public:
 
 	CHAOSVD_API virtual void PostEditChangeProperty(FPropertyChangedEvent& PropertyChangedEvent) override;
 
-	CHAOSVD_API FChaosVDSettingChanged& OnSettingsChanged() { return SettingsChangedDelegate; }
+	FChaosVDSettingChanged& OnSettingsChanged() { return SettingsChangedDelegate; }
 
 	CHAOSVD_API virtual void PostEditUndo() override;
 
 	CHAOSVD_API virtual void OverridePerObjectConfigSection(FString& SectionName) override;
 
-	CHAOSVD_API FStringView GetConfigSectionName()
+	FStringView GetConfigSectionName()
 	{
 		return OverrideConfigSectionName;
 	}
diff --git a/Engine/Plugins/ChaosVD/Source/ChaosVD/Public/TEDS/ChaosVDStructTypedElementData.h b/Engine/Plugins/ChaosVD/Source/ChaosVD/Public/TEDS/ChaosVDStructTypedElementData.h
--- a/Engine/Plugins/ChaosVD/Source/ChaosVD/Public/TEDS/ChaosVDStructTypedElementData.h
+++ b/Engine/Plugins/ChaosVD/Source/ChaosVD/Public/TEDS/ChaosVDStructTypedElementData.h
@@ -7,15 +7,17 @@
 #include "Elements/Framework/TypedElementOwnerStore.h"
 #include "Elements/Framework/TypedElementRegistry.h"
 
+#define UE_API CHAOSVD_API
+
 class FStructOnScope;
 
 namespace Chaos::VD::TypedElementDataUtil
 {
 	static const FName NAME_CVD_StructDataElement = FName("CVD_StructDataElement");
 
-	struct CHAOSVD_API FStructTypedElementData
+	struct FStructTypedElementData
 	{
-		UE_DECLARE_TYPED_ELEMENT_DATA_RTTI(FStructTypedElementData);
+		UE_DECLARE_TYPED_ELEMENT_DATA_RTTI(FStructTypedElementData, UE_API);
 
 		UStruct* TypeInfo = nullptr;
 		void* RawData = nullptr;
@@ -23,7 +25,7 @@ namespace Chaos::VD::TypedElementDataUtil
 		template<typename StructDataType>
 		StructDataType* GetData() const;
 
-		TSharedPtr<FStructOnScope> GetDataAsStructScope() const;
+		UE_API TSharedPtr<FStructOnScope> GetDataAsStructScope() const;
 	};
 
 	extern TTypedElementOwnerStore<FStructTypedElementData, void*> GCVDTypedStructDataElementOwnerStore;
@@ -102,3 +104,4 @@ inline FString GetTypedElementDebugId< Chaos::VD::TypedElementDataUtil::FStructT
 }
 
 
+#undef UE_API
diff --git a/Engine/Plugins/ChaosVD/Source/ChaosVD/Public/Trace/DataProcessors/ChaosVDDataProcessorBase.h b/Engine/Plugins/ChaosVD/Source/ChaosVD/Public/Trace/DataProcessors/ChaosVDDataProcessorBase.h
--- a/Engine/Plugins/ChaosVD/Source/ChaosVD/Public/Trace/DataProcessors/ChaosVDDataProcessorBase.h
+++ b/Engine/Plugins/ChaosVD/Source/ChaosVD/Public/Trace/DataProcessors/ChaosVDDataProcessorBase.h
@@ -1,4 +1,4 @@
-﻿// Copyright Epic Games, Inc. All Rights Reserved.
+// Copyright Epic Games, Inc. All Rights Reserved.
 
 #pragma once
 
@@ -65,7 +65,7 @@ class FChaosVDDataProcessorBase
 {
 public:
 
-	CHAOSVD_API explicit FChaosVDDataProcessorBase(FStringView InCompatibleType) : TraceProvider(nullptr), CompatibleType(InCompatibleType), ProcessedBytes(0)
+	explicit FChaosVDDataProcessorBase(FStringView InCompatibleType) : TraceProvider(nullptr), CompatibleType(InCompatibleType), ProcessedBytes(0)
 	{
 	}
 
diff --git a/Engine/Plugins/ChaosVD/Source/ChaosVD/Public/Visualizers/ChaosVDComponentVisualizerBase.h b/Engine/Plugins/ChaosVD/Source/ChaosVD/Public/Visualizers/ChaosVDComponentVisualizerBase.h
--- a/Engine/Plugins/ChaosVD/Source/ChaosVD/Public/Visualizers/ChaosVDComponentVisualizerBase.h
+++ b/Engine/Plugins/ChaosVD/Source/ChaosVD/Public/Visualizers/ChaosVDComponentVisualizerBase.h
@@ -13,6 +13,8 @@
 #include "Widgets/SChaosVDEnumFlagsMenu.h"
 #include "Widgets/SWidget.h"
 
+#define UE_API CHAOSVD_API
+
 class SChaosVDMainTab;
 class FChaosVDSolverDataSelection;
 struct FChaosVDSolverDataSelectionHandle;
@@ -31,9 +33,9 @@ struct FChaosVDVisualizationContext
 };
 
 /** Custom Hit Proxy for debug drawn particle data */
-struct CHAOSVD_API HChaosVDComponentVisProxy : public HComponentVisProxy
+struct HChaosVDComponentVisProxy : public HComponentVisProxy
 {
-	DECLARE_HIT_PROXY()
+	DECLARE_HIT_PROXY( UE_API )
 
 	HChaosVDComponentVisProxy(const UActorComponent* InComponent, const TSharedPtr<FChaosVDSolverDataSelectionHandle>& InDataSelectionHandle)
 		: HComponentVisProxy(InComponent), DataSelectionHandle(InDataSelectionHandle)
@@ -67,7 +69,7 @@ protected:
 	/** Called when the menu for this visualizer (if any) can be registered. And does so if implemented.
 	 * This needs to be manually called in the constructor of ech visualizer
 	 */
-	CHAOSVD_API virtual void RegisterVisualizerMenus() = 0;
+	virtual void RegisterVisualizerMenus() = 0;
 
 	/** Returns true if this visualizer can handle a click in the viewport
 	 */
@@ -111,3 +113,5 @@ void FChaosVDComponentVisualizerBase::CreateGenericVisualizerMenu(FName MenuToEx
 		CreateVisualizationOptionsMenuSections<ObjectSettingsType, VisualizationFlagsType>(Menu, SectionName, InSectionLabel, InFlagsMenuLabel, InFlagsMenuTooltip, FlagsMenuIcon, InSettingsMenuLabel, InSettingsMenuTooltip);
 	}
 }
+
+#undef UE_API
diff --git a/Engine/Plugins/ChaosVD/Source/ChaosVD/Public/Widgets/SChaosVDMainTab.h b/Engine/Plugins/ChaosVD/Source/ChaosVD/Public/Widgets/SChaosVDMainTab.h
--- a/Engine/Plugins/ChaosVD/Source/ChaosVD/Public/Widgets/SChaosVDMainTab.h
+++ b/Engine/Plugins/ChaosVD/Source/ChaosVD/Public/Widgets/SChaosVDMainTab.h
@@ -36,8 +36,8 @@ struct FStructureDetailsViewArgs;
 
 enum class EChaosVDLoadRecordedDataMode : uint8;
 
-UCLASS()
-class CHAOSVD_API UChaosVDMainToolbarMenuContext : public UObject
+UCLASS(MinimalAPI)
+class UChaosVDMainToolbarMenuContext : public UObject
 {
 	GENERATED_BODY()
 
diff --git a/Engine/Plugins/Chooser/Source/ChooserEditor/Public/ObjectChooserWidgetFactories.h b/Engine/Plugins/Chooser/Source/ChooserEditor/Public/ObjectChooserWidgetFactories.h
--- a/Engine/Plugins/Chooser/Source/ChooserEditor/Public/ObjectChooserWidgetFactories.h
+++ b/Engine/Plugins/Chooser/Source/ChooserEditor/Public/ObjectChooserWidgetFactories.h
@@ -8,6 +8,8 @@
 #include "Widgets/Layout/SBorder.h"
 #include "StructViewerModule.h"
 
+#define UE_API CHOOSEREDITOR_API
+
 class UChooserTable;
 struct FChooserColumnBase;
 
@@ -20,27 +22,29 @@ namespace UE::ChooserEditor
 	typedef TFunction<TSharedRef<SWidget>(bool bReadOnly, UObject* TransactionObject, void* Value, UClass* ResultBaseClass, FChooserWidgetValueChanged)> FChooserWidgetCreator;
 	typedef TFunction<TSharedRef<SWidget>(UChooserTable* Chooser, FChooserColumnBase* Column, int Row)> FColumnWidgetCreator;
 
-	class CHOOSEREDITOR_API FObjectChooserWidgetFactories
+	class FObjectChooserWidgetFactories
 	{
 	public:
 
-		static TSharedPtr<SWidget> CreateWidget(bool bReadOnly, UObject* TransactionObject, UScriptStruct* DataBaseType, FInstancedStruct* Data, UClass* ResultBaseClass,
+		static UE_API TSharedPtr<SWidget> CreateWidget(bool bReadOnly, UObject* TransactionObject, UScriptStruct* DataBaseType, FInstancedStruct* Data, UClass* ResultBaseClass,
 												FChooserWidgetValueChanged ValueChanged = FChooserWidgetValueChanged(), FText NullValueDisplayText = FText());
 		
-		static TSharedPtr<SWidget> CreateWidget(bool ReadOnly, UObject* TransactionObject, void* Value, const UStruct* ValueType, UClass* ResultBaseClass, FChooserWidgetValueChanged ValueChanged = FChooserWidgetValueChanged());
-		static TSharedPtr<SWidget> CreateWidget(bool ReadOnly, UObject* TransactionObject, const UScriptStruct* BaseType, void* Value, const UStruct* ValueType, UClass* ResultBaseClass,
+		static UE_API TSharedPtr<SWidget> CreateWidget(bool ReadOnly, UObject* TransactionObject, void* Value, const UStruct* ValueType, UClass* ResultBaseClass, FChooserWidgetValueChanged ValueChanged = FChooserWidgetValueChanged());
+		static UE_API TSharedPtr<SWidget> CreateWidget(bool ReadOnly, UObject* TransactionObject, const UScriptStruct* BaseType, void* Value, const UStruct* ValueType, UClass* ResultBaseClass,
 										const FOnStructPicked& CreateClassCallback, TSharedPtr<SBorder>* InnerWidget = nullptr, FChooserWidgetValueChanged ValueChanged = FChooserWidgetValueChanged(),
 										FText NullValueDisplayText = FText());
 		
-		static TSharedPtr<SWidget> CreateColumnWidget(FChooserColumnBase* Column, const UStruct* ColumnType, UChooserTable* Chooser, int Row); // todo: chooser should be a UObject
+		static UE_API TSharedPtr<SWidget> CreateColumnWidget(FChooserColumnBase* Column, const UStruct* ColumnType, UChooserTable* Chooser, int Row); // todo: chooser should be a UObject
 		
-		static void RegisterWidgets();
+		static UE_API void RegisterWidgets();
 
-		static void RegisterWidgetCreator(const UStruct* Type, FChooserWidgetCreator Creator);
-		static void RegisterColumnWidgetCreator(const UStruct* ColumnType, FColumnWidgetCreator Creator);
+		static UE_API void RegisterWidgetCreator(const UStruct* Type, FChooserWidgetCreator Creator);
+		static UE_API void RegisterColumnWidgetCreator(const UStruct* ColumnType, FColumnWidgetCreator Creator);
 		
 	private:
-		static TMap<const UStruct*, FColumnWidgetCreator>  ColumnWidgetCreators;
-		static TMap<const UStruct*, FChooserWidgetCreator> ChooserWidgetCreators;
+		static UE_API TMap<const UStruct*, FColumnWidgetCreator>  ColumnWidgetCreators;
+		static UE_API TMap<const UStruct*, FChooserWidgetCreator> ChooserWidgetCreators;
 	};
 }
+
+#undef UE_API
diff --git a/Engine/Plugins/Chooser/Source/ChooserEditor/Public/SPropertyAccessChainWidget.h b/Engine/Plugins/Chooser/Source/ChooserEditor/Public/SPropertyAccessChainWidget.h
--- a/Engine/Plugins/Chooser/Source/ChooserEditor/Public/SPropertyAccessChainWidget.h
+++ b/Engine/Plugins/Chooser/Source/ChooserEditor/Public/SPropertyAccessChainWidget.h
@@ -7,13 +7,15 @@
 #include "IHasContext.h"
 #include "ChooserPropertyAccess.h"
 
+#define UE_API CHOOSEREDITOR_API
+
 namespace UE::ChooserEditor
 {
 DECLARE_DELEGATE(FPropertyAccessChainChanged)
 
 		
 // Wrapper widget for Property access widget, which can update when the target Class changes
-class CHOOSEREDITOR_API SPropertyAccessChainWidget : public SCompoundWidget
+class SPropertyAccessChainWidget : public SCompoundWidget
 {
 public:
 	SLATE_BEGIN_ARGS(SPropertyAccessChainWidget)
@@ -29,21 +31,21 @@ public:
             
 	SLATE_END_ARGS()
 
-	void Construct( const FArguments& InArgs);
-	virtual ~SPropertyAccessChainWidget();
+	UE_API void Construct( const FArguments& InArgs);
+	UE_API virtual ~SPropertyAccessChainWidget();
 	
 	// Assign result from Property Access Widget to an FChooserPropertyBinding (Helper for OnAddBinding implementations)
-	static void SetPropertyBinding(IHasContextClass* HasContext, const TArray<FBindingChainElement>& InBindingChain, FChooserPropertyBinding& OutPropertyBinding);
+	static UE_API void SetPropertyBinding(IHasContextClass* HasContext, const TArray<FBindingChainElement>& InBindingChain, FChooserPropertyBinding& OutPropertyBinding);
 
 private:
-	TSharedRef<SWidget> CreatePropertyAccessWidget();
-	void UpdateWidget();
-	void ContextClassChanged();
+	UE_API TSharedRef<SWidget> CreatePropertyAccessWidget();
+	UE_API void UpdateWidget();
+	UE_API void ContextClassChanged();
 		
 	FString TypeFilter;
 	FString AlternateTypeFilter;
 	FString BindingColor;
-	FDelegateHandle ContextClassChangedHandle();
+	UE_API FDelegateHandle ContextClassChangedHandle();
 	IHasContextClass* ContextClassOwner = nullptr;
 	FOnAddBinding OnAddBinding;
 	FPropertyAccessChainChanged OnValueChanged;
@@ -52,4 +54,6 @@ private:
 };
 	
 
-}
\ No newline at end of file
+}
+
+#undef UE_API
diff --git a/Engine/Plugins/Compositing/OpenColorIO/Source/OpenColorIO/Public/OpenColorIOColorSpace.h b/Engine/Plugins/Compositing/OpenColorIO/Source/OpenColorIO/Public/OpenColorIOColorSpace.h
--- a/Engine/Plugins/Compositing/OpenColorIO/Source/OpenColorIO/Public/OpenColorIOColorSpace.h
+++ b/Engine/Plugins/Compositing/OpenColorIO/Source/OpenColorIO/Public/OpenColorIOColorSpace.h
@@ -16,7 +16,7 @@ class UOpenColorIOConfiguration;
  * Members are populated by data coming from a config file.
  */
 USTRUCT(BlueprintType, meta = (DisplayName = "OpenColorIO Color Space"))
-struct OPENCOLORIO_API FOpenColorIOColorSpace
+struct FOpenColorIOColorSpace
 {
 	GENERATED_BODY()
 
@@ -30,7 +30,7 @@ public:
 	FOpenColorIOColorSpace(FOpenColorIOColorSpace&&) = default;
 	FOpenColorIOColorSpace& operator=(const FOpenColorIOColorSpace&) = default;
 	FOpenColorIOColorSpace& operator=(FOpenColorIOColorSpace&&) = default;
-	PRAGMA_ENABLE_DEPRECATION_WARNINGS
+	OPENCOLORIO_API PRAGMA_ENABLE_DEPRECATION_WARNINGS
 
 	/**
 	 * Create and initialize a new instance.
@@ -60,7 +60,7 @@ public:
 	FString Description;
 
 	/** Delimiter used in the OpenColorIO library to make family hierarchies */
-	static const TCHAR* FamilyDelimiter;
+	static OPENCOLORIO_API const TCHAR* FamilyDelimiter;
 
 public:
 	bool operator==(const FOpenColorIOColorSpace& Other) const { return Other.ColorSpaceName == ColorSpaceName; }
@@ -70,20 +70,20 @@ public:
 	 * Get the string representation of this color space.
 	 * @return ColorSpace name. 
 	 */
-	FString ToString() const;
+	OPENCOLORIO_API FString ToString() const;
 
 	/** Return true if the index and name have been set properly */
-	bool IsValid() const;
+	OPENCOLORIO_API bool IsValid() const;
 
 	/** Reset members to default/empty values. */
-	void Reset();
+	OPENCOLORIO_API void Reset();
 
 	/** 
 	 * Return the family name at the desired depth level 
 	 * @param InDepth Desired depth in the family string. 0 == First layer. 
 	 * @return FamilyName at the desired depth. Empty string if depth level doesn't exist.
 	 */
-	FString GetFamilyNameAtDepth(int32 InDepth) const;
+	OPENCOLORIO_API FString GetFamilyNameAtDepth(int32 InDepth) const;
 };
 
 
@@ -99,7 +99,7 @@ enum class EOpenColorIOViewTransformDirection : uint8
 
 
 USTRUCT(BlueprintType, meta = (DisplayName = "OpenColorIO Display View"))
-struct OPENCOLORIO_API FOpenColorIODisplayView
+struct FOpenColorIODisplayView
 {
 	GENERATED_BODY()
 
@@ -110,7 +110,7 @@ public:
 	/**
 	 * Create and initialize a new instance.
 	 */
-	FOpenColorIODisplayView(FStringView InDisplayName, FStringView InViewName, const FString& InDescription = FString());
+	OPENCOLORIO_API FOpenColorIODisplayView(FStringView InDisplayName, FStringView InViewName, const FString& InDescription = FString());
 
 	/** Display name. */
 	UPROPERTY(EditAnywhere, BlueprintReadWrite, Category = ColorSpace)
@@ -124,13 +124,13 @@ public:
 	UPROPERTY()
 	FString Description;
 
-	FString ToString() const;
+	OPENCOLORIO_API FString ToString() const;
 
 	/** Return true if the index and name have been set properly */
-	bool IsValid() const;
+	OPENCOLORIO_API bool IsValid() const;
 
 	/** Reset members to default/empty values. */
-	void Reset();
+	OPENCOLORIO_API void Reset();
 
 public:
 	bool operator==(const FOpenColorIODisplayView& Other) const { return Other.Display == Display && Other.View == View; }
@@ -141,7 +141,7 @@ public:
  * Identifies a OCIO ColorSpace conversion.
  */
 USTRUCT(BlueprintType, meta = (DisplayName = "OpenColorIO Color Conversion Settings"))
-struct OPENCOLORIO_API FOpenColorIOColorConversionSettings
+struct FOpenColorIOColorConversionSettings
 {
 	GENERATED_BODY()
 
@@ -149,7 +149,7 @@ public:
 	DECLARE_MULTICAST_DELEGATE(FOnConversionSettingsChange);
 
 	/** Default constructor. */
-	FOpenColorIOColorConversionSettings();
+	OPENCOLORIO_API FOpenColorIOColorConversionSettings();
 
 	/** The source color space name. */
 	UPROPERTY(EditAnywhere, BlueprintReadWrite, Category = ColorSpace)
@@ -173,35 +173,35 @@ public:
 
 public:
 
-	void PostSerialize(const FArchive& Ar);
+	OPENCOLORIO_API void PostSerialize(const FArchive& Ar);
 
 	/**
 	 * Get a string representation of this conversion.
 	 * @return String representation, i.e. "ConfigurationAssetName - SourceColorSpace to DestinationColorSpace".
 	 */
-	FString ToString() const;
+	OPENCOLORIO_API FString ToString() const;
 
 	/**
 	 * Returns true if the source and destination color spaces are found in the configuration file
 	 */
-	bool IsValid() const;
+	OPENCOLORIO_API bool IsValid() const;
 
 	/** Returns a string representing the settings' source */
-	FString GetSourceString() const;
+	OPENCOLORIO_API FString GetSourceString() const;
 
 	/** Returns a string representing the settings' destination */
-	FString GetDestinationString() const;
+	OPENCOLORIO_API FString GetDestinationString() const;
 
 	/**
 	* Reset members to default/empty values.
 	* @param InDepth Desired depth in the family string. 0 == First layer. 
 	*/
-	void Reset(bool bResetConfigurationSource = false);
+	OPENCOLORIO_API void Reset(bool bResetConfigurationSource = false);
 
 	/**
 	* Ensure that the selected source and destination color spaces are valid, resets them otherwise.
 	*/
-	void ValidateColorSpaces();
+	OPENCOLORIO_API void ValidateColorSpaces();
 
 	bool operator==(const FOpenColorIOColorConversionSettings& Other) const { return Equals(Other); }
 	bool operator!=(const FOpenColorIOColorConversionSettings& Other) const { return !Equals(Other); }
@@ -217,7 +217,7 @@ public:
 	}
 	
 	/** Whether or not these settings are of the display-view type. */
-	bool IsDisplayView() const;
+	OPENCOLORIO_API bool IsDisplayView() const;
 };
 
 template<> struct TStructOpsTypeTraits<FOpenColorIOColorConversionSettings> : public TStructOpsTypeTraitsBase2<FOpenColorIOColorConversionSettings>
@@ -232,7 +232,7 @@ template<> struct TStructOpsTypeTraits<FOpenColorIOColorConversionSettings> : pu
  * Identifies an OCIO Display look configuration 
  */
 USTRUCT(BlueprintType, meta = (DisplayName = "OpenColorIO Display Configuration"))
-struct OPENCOLORIO_API FOpenColorIODisplayConfiguration
+struct FOpenColorIODisplayConfiguration
 {
 	GENERATED_BODY()
 
@@ -248,8 +248,8 @@ public:
 	UPROPERTY(EditAnywhere, BlueprintReadWrite, Category = ColorSpace)
 	FOpenColorIOColorConversionSettings ColorConfiguration;
 
-	bool Serialize(FArchive& Ar);
-	void PostSerialize(const FArchive& Ar);
+	OPENCOLORIO_API bool Serialize(FArchive& Ar);
+	OPENCOLORIO_API void PostSerialize(const FArchive& Ar);
 
 	/** Determines if this DisplayConfiguration is the same as another.*/
 	bool Equals(const FOpenColorIODisplayConfiguration& Other) const
