Reverted change from 43132987 that was merged wrong

diff --git a/Engine/Plugins/AudioInsights/Source/AudioInsights/Public/Views/SoundDashboardViewFactory.h b/Engine/Plugins/AudioInsights/Source/AudioInsights/Public/Views/SoundDashboardViewFactory.h
--- a/Engine/Plugins/AudioInsights/Source/AudioInsights/Public/Views/SoundDashboardViewFactory.h
+++ b/Engine/Plugins/AudioInsights/Source/AudioInsights/Public/Views/SoundDashboardViewFactory.h
@@ -5,8 +5,6 @@
 
 #include "Filters/GenericFilter.h"
 
-#define UE_API AUDIOINSIGHTS_API
-
 class FUICommandList;
 
 namespace UE::Audio::Insights
@@ -70,23 +68,23 @@ namespace UE::Audio::Insights
 	*	- Contains a weak handle to the original entry (OriginalDataEntry) which is updated from the trace provider
 	*	- Copies updated params to PinnedSectionEntry for display
 	*/
-	class FPinnedSoundEntryWrapper
+	class AUDIOINSIGHTS_API FPinnedSoundEntryWrapper
 	{
 	public:
 		FPinnedSoundEntryWrapper() = delete;
-		UE_API FPinnedSoundEntryWrapper(const TSharedPtr<IDashboardDataTreeViewEntry>& OriginalEntry);
+		FPinnedSoundEntryWrapper(const TSharedPtr<IDashboardDataTreeViewEntry>& OriginalEntry);
 
 		TSharedPtr<IDashboardDataTreeViewEntry> GetPinnedSectionEntry() const { return PinnedSectionEntry; }
 		TSharedPtr<IDashboardDataTreeViewEntry> GetOriginalDataEntry() const { return OriginalDataEntry.IsValid() ? OriginalDataEntry.Pin() : nullptr; }
 
-		UE_API TSharedPtr<FPinnedSoundEntryWrapper> AddChildEntry(const TSharedPtr<IDashboardDataTreeViewEntry> Child);
+		TSharedPtr<FPinnedSoundEntryWrapper> AddChildEntry(const TSharedPtr<IDashboardDataTreeViewEntry> Child);
 
-		UE_API void UpdateParams();
+		void UpdateParams();
 
-		UE_API void CleanUp();
-		UE_API void MarkToDelete();
+		void CleanUp();
+		void MarkToDelete();
 
-		UE_API bool EntryIsValid() const;
+		bool EntryIsValid() const;
 
 		TArray<TSharedPtr<FPinnedSoundEntryWrapper>> PinnedWrapperChildren;
 
@@ -95,78 +93,96 @@ namespace UE::Audio::Insights
 		TWeakPtr<IDashboardDataTreeViewEntry> OriginalDataEntry;
 	};
 
-	class FSoundDashboardViewFactory : public FTraceTreeDashboardViewFactory
+	class AUDIOINSIGHTS_API FSoundDashboardViewFactory : public FTraceTreeDashboardViewFactory
 	{
 	public:
-		UE_API FSoundDashboardViewFactory();
-		UE_API virtual ~FSoundDashboardViewFactory();
+		FSoundDashboardViewFactory();
+		virtual ~FSoundDashboardViewFactory();
 
-		UE_API virtual FName GetName() const override;
-		UE_API virtual FText GetDisplayName() const override;
-		UE_API virtual FSlateIcon GetIcon() const override;
-		UE_API virtual EDefaultDashboardTabStack GetDefaultTabStack() const override;
-		UE_API virtual TSharedRef<SWidget> MakeWidget(TSharedRef<SDockTab> OwnerTab, const FSpawnTabArgs& SpawnTabArgs) override;
-		UE_API virtual void ProcessEntries(FTraceTreeDashboardViewFactory::EProcessReason Reason) override;
+		virtual FName GetName() const override;
+		virtual FText GetDisplayName() const override;
+		virtual FSlateIcon GetIcon() const override;
+		virtual EDefaultDashboardTabStack GetDefaultTabStack() const override;
+		virtual TSharedRef<SWidget> MakeWidget(TSharedRef<SDockTab> OwnerTab, const FSpawnTabArgs& SpawnTabArgs) override;
+		virtual void ProcessEntries(FTraceTreeDashboardViewFactory::EProcessReason Reason) override;
 
 	private:
-		UE_API void BindCommands();
+		void BindCommands();
 
-		UE_API virtual TSharedPtr<SWidget> GetFilterBarWidget() override;
-		UE_API virtual TSharedPtr<SWidget> GetFilterBarButtonWidget() override;
+		virtual TSharedPtr<SWidget> GetFilterBarWidget() override;
+		virtual TSharedPtr<SWidget> GetFilterBarButtonWidget() override;
 
-		UE_API bool IsRootItem(const TSharedRef<IDashboardDataTreeViewEntry>& InEntry) const;
-		UE_API bool EntryCanHaveChildren(const TSharedRef<IDashboardDataTreeViewEntry>& InEntry) const;
-		UE_API bool IsDescendant(const TSharedPtr<IDashboardDataTreeViewEntry>& InEntry, const TSharedPtr<IDashboardDataTreeViewEntry>& InChildCandidate) const;
+		bool IsRootItem(const TSharedRef<IDashboardDataTreeViewEntry>& InEntry) const;
 
-		UE_API TSharedRef<SWidget> GenerateWidgetForRootColumn(const TSharedRef<FTraceTreeDashboardViewFactory::SRowWidget>& InRowWidget, const TSharedRef<IDashboardDataTreeViewEntry>& InRowData, const FName& InColumn, const FText& InValueText);
-		UE_API virtual TSharedRef<SWidget> GenerateWidgetForColumn(const TSharedRef<FTraceTreeDashboardViewFactory::SRowWidget>& InRowWidget, const TSharedRef<IDashboardDataTreeViewEntry>& InRowData, const FName& InColumn) override;
+		TSharedRef<SWidget> GenerateWidgetForRootColumn(const TSharedRef<FTraceTreeDashboardViewFactory::SRowWidget>& InRowWidget, const TSharedRef<IDashboardDataTreeViewEntry>& InRowData, const FName& InColumn, const FText& InValueText);
+		virtual TSharedRef<SWidget> GenerateWidgetForColumn(const TSharedRef<FTraceTreeDashboardViewFactory::SRowWidget>& InRowWidget, const TSharedRef<IDashboardDataTreeViewEntry>& InRowData, const FName& InColumn) override;
 		
-		UE_API virtual TSharedRef<ITableRow> OnGenerateRow(TSharedPtr<IDashboardDataTreeViewEntry> Item, const TSharedRef<STableViewBase>& OwnerTable) override;
+		virtual TSharedRef<ITableRow> OnGenerateRow(TSharedPtr<IDashboardDataTreeViewEntry> Item, const TSharedRef<STableViewBase>& OwnerTable) override;
 
-		UE_API virtual TSharedPtr<SWidget> OnConstructContextMenu() override;
-		UE_API virtual FReply OnDataRowKeyInput(const FGeometry& InGeometry, const FKeyEvent& InKeyEvent) const override;
+		virtual TSharedPtr<SWidget> OnConstructContextMenu() override;
+		virtual FReply OnDataRowKeyInput(const FGeometry& InGeometry, const FKeyEvent& InKeyEvent) const override;
 
-		UE_API virtual const TMap<FName, FTraceTreeDashboardViewFactory::FHeaderRowColumnData>& GetHeaderRowColumns() const override;
-		UE_API virtual const TMap<FName, FTraceTreeDashboardViewFactory::FColumnData>& GetColumns() const override;
+		virtual const TMap<FName, FTraceTreeDashboardViewFactory::FHeaderRowColumnData>& GetHeaderRowColumns() const override;
+		virtual const TMap<FName, FTraceTreeDashboardViewFactory::FColumnData>& GetColumns() const override;
 
-		UE_API virtual void SortTable() override;
+		virtual void SortTable() override;
 		virtual TArray<TSharedPtr<IDashboardDataTreeViewEntry>>& GetTreeItemsSource() override { return FullTree; }
-		UE_API virtual bool ResetTreeData() override;
+		virtual bool ResetTreeData() override;
+
+		virtual int32 GetDisplayedTreeDepth() const override { return static_cast<int32>(TreeViewingMode); }
+
+		enum class ETreeViewingOptions : uint8
+		{
+			FullTree = 0,
+			ActiveSounds,
+			FlatList
+		};
 
-		UE_API void RecursiveSort(TArray<TSharedPtr<IDashboardDataTreeViewEntry>>& OutTree, TFunctionRef<bool(const FSoundDashboardEntry&, const FSoundDashboardEntry&)> Predicate);
-		UE_API void SortByPredicate(TFunctionRef<bool(const FSoundDashboardEntry&, const FSoundDashboardEntry&)> Predicate);
+		void ChangeTreeViewingMode(const ETreeViewingOptions SelectedMode);
+		void BuildViewOptionsMenuContent(FMenuBuilder& MenuBuilder);
 
-		UE_API TSharedRef<SWidget> MakeShowRecentlyStoppedSoundsWidget();
+		void FilterText();
+		void RecursiveSort(TArray<TSharedPtr<IDashboardDataTreeViewEntry>>& OutTree, TFunctionRef<bool(const FSoundDashboardEntry&, const FSoundDashboardEntry&)> Predicate);
+		void SortByPredicate(TFunctionRef<bool(const FSoundDashboardEntry&, const FSoundDashboardEntry&)> Predicate);
+
+		TSharedRef<SWidget> OnGetSettingsMenuContent();
+		TSharedRef<SWidget> MakeSettingsButtonWidget();
 
 #if WITH_EDITOR
-		UE_API TSharedRef<SWidget> MakeMuteSoloWidget();
+		TSharedRef<SWidget> MakeMuteSoloWidget();
 
-		UE_API TSharedRef<SWidget> CreateMuteSoloButton(const TSharedRef<IDashboardDataTreeViewEntry>& InRowData,
+		TSharedRef<SWidget> CreateMuteSoloButton(const TSharedRef<FTraceTreeDashboardViewFactory::SRowWidget>& InRowWidget, 
+			const TSharedRef<IDashboardDataTreeViewEntry>& InRowData,
 			const FName& InColumn,
 			TFunction<void(const TArray<TSharedPtr<IDashboardDataTreeViewEntry>>& /*InEntries*/)> MuteSoloToggleFunc,
 			TFunctionRef<bool(const IDashboardDataTreeViewEntry& /*InEntry*/, const bool /*bInCheckChildren*/)> IsMuteSoloFunc);
 
-		UE_API void ToggleMuteSoloEntries(const TArray<TSharedPtr<IDashboardDataTreeViewEntry>>& InEntries, const EMuteSoloMode InMuteSoloMode);
+		void ToggleMuteSoloEntries(const TArray<TSharedPtr<IDashboardDataTreeViewEntry>>& InEntries, const EMuteSoloMode InMuteSoloMode);
 
-		UE_API TArray<TObjectPtr<UObject>> GetSelectedEditableAssets() const;
+		TArray<TObjectPtr<UObject>> GetSelectedEditableAssets() const;
 #endif // WITH_EDITOR
 
-		UE_API bool SelectedItemsIncludesAnAsset() const;
-		UE_API void PinSound();
-		UE_API void UnpinSound();
-		UE_API bool SelectionIncludesUnpinnedItem() const;
+		bool SelectedItemsIncludesAnAsset() const;
+		bool SelectionIncludesUnpinnedItem() const;
+		void ClearSelection();
+
+		void PinSound();
+		void UnpinSound();
 
-		UE_API void PinSelectedItems(const TSharedPtr<IDashboardDataTreeViewEntry>& Entry, const TArray<TSharedPtr<IDashboardDataTreeViewEntry>>& SelectedItems);
-		UE_API void UnpinSelectedItems(const TSharedPtr<FPinnedSoundEntryWrapper>& PinnedWrapperEntry, const TArray<TSharedPtr<IDashboardDataTreeViewEntry>>& SelectedItems, const bool bSelectionContainsAssets);
+		void PinSelectedItems(const TSharedPtr<IDashboardDataTreeViewEntry>& Entry, const TArray<TSharedPtr<IDashboardDataTreeViewEntry>>& SelectedItems);
+		void UnpinSelectedItems(const TSharedPtr<FPinnedSoundEntryWrapper>& PinnedWrapperEntry, const TArray<TSharedPtr<IDashboardDataTreeViewEntry>>& SelectedItems, const bool bSelectionContainsAssets);
 
-		UE_API void MarkBranchAsPinned(const TSharedPtr<IDashboardDataTreeViewEntry> Entry, const bool bIsPinned);
-		UE_API void InitPinnedItemEntries();
-		UE_API void CreatePinnedEntry(TSharedPtr<IDashboardDataTreeViewEntry> Entry);
-		UE_API void UpdatePinnedSection();
+		void RecreatePinnedEntries(const TSharedPtr<IDashboardDataTreeViewEntry>& Entry);
+
+		void MarkBranchAsPinned(const TSharedPtr<IDashboardDataTreeViewEntry> Entry, const bool bIsPinned);
+		void InitPinnedItemEntries();
+		void CreatePinnedEntry(TSharedPtr<IDashboardDataTreeViewEntry> Entry);
+		void UpdatePinnedSection();
+		void RebuildPinnedSection();
 
 #if WITH_EDITOR
-		UE_API void BrowseSoundAsset() const;
-		UE_API void OpenSoundAsset() const;
+		void BrowseSoundAsset() const;
+		void OpenSoundAsset() const;
 #endif
 		// @TODO UE-250399: Hide category pending to implement
 		//void HideSound();
@@ -188,10 +204,11 @@ namespace UE::Audio::Insights
 		TSharedPtr<SWidget> SoundsFilterBar;
 		TSharedPtr<SWidget> SoundsFilterBarButton;
 
+		ETreeViewingOptions TreeViewingMode = ETreeViewingOptions::FullTree;
 		ESoundDashboardFilterFlags SelectedFilterFlags = AllFilterFlags;
 		bool bIsPinnedCategoryFilterEnabled = true;
 		bool bShowRecentlyStoppedSounds = false;
+		bool bTreeViewModeChanged = false;
+		bool bDisplayAmpPeakInDb = true;
 	};
 } // namespace UE::Audio::Insights
-
-#undef UE_API
diff --git a/Engine/Plugins/ChaosClothAssetEditor/Source/ChaosClothAssetEditor/Private/ChaosClothAsset/SClothEditor3DViewport.h b/Engine/Plugins/ChaosClothAssetEditor/Source/ChaosClothAssetEditor/Private/ChaosClothAsset/SClothEditor3DViewport.h
--- a/Engine/Plugins/ChaosClothAssetEditor/Source/ChaosClothAssetEditor/Private/ChaosClothAsset/SClothEditor3DViewport.h
+++ b/Engine/Plugins/ChaosClothAssetEditor/Source/ChaosClothAssetEditor/Private/ChaosClothAsset/SClothEditor3DViewport.h
@@ -7,8 +7,6 @@
 #include "SCommonEditorViewportToolbarBase.h"
 #include "SBaseCharacterFXEditorViewport.h"
 
-#define UE_API CHAOSCLOTHASSETEDITOR_API
-
 namespace UE::Chaos::ClothAsset
 {
 class FChaosClothPreviewScene;
@@ -18,7 +16,7 @@ class SClothAnimationScrubPanel;
 /**
  * Viewport used for 3D preview in cloth editor. Has a custom toolbar overlay at the top.
  */
-class SChaosClothAssetEditor3DViewport : public SBaseCharacterFXEditorViewport, public ICommonEditorViewportToolbarInfoProvider
+class CHAOSCLOTHASSETEDITOR_API SChaosClothAssetEditor3DViewport : public SBaseCharacterFXEditorViewport, public ICommonEditorViewportToolbarInfoProvider
 {
 
 public:
@@ -29,20 +27,19 @@ public:
 		SLATE_ARGUMENT(TSharedPtr<FUICommandList>, ToolkitCommandList)
 	SLATE_END_ARGS()
 
-	UE_API void Construct(const FArguments& InArgs, const FAssetEditorViewportConstructionArgs& InViewportConstructionArgs);
+	void Construct(const FArguments& InArgs, const FAssetEditorViewportConstructionArgs& InViewportConstructionArgs);
 
 	// SAssetEditorViewport
-	UE_API virtual void BindCommands() override;
-	UE_API virtual TSharedPtr<SWidget> MakeViewportToolbar() override;
-	UE_API virtual TSharedPtr<SWidget> BuildViewportToolbar() override;
-	UE_API virtual void PopulateViewportOverlays(TSharedRef<SOverlay> Overlay) override;
-	UE_API virtual bool IsVisible() const override;
+	virtual void BindCommands() override;
+	virtual TSharedPtr<SWidget> BuildViewportToolbar() override;
+	virtual void PopulateViewportOverlays(TSharedRef<SOverlay> Overlay) override;
+	virtual bool IsVisible() const override;
 
-	UE_API virtual void OnFocusViewportToSelection() override;
+	virtual void OnFocusViewportToSelection() override;
 
 	// ICommonEditorViewportToolbarInfoProvider
-	UE_API virtual TSharedRef<class SEditorViewport> GetViewportWidget() override;
-	UE_API virtual TSharedPtr<FExtender> GetExtenders() const override;
+	virtual TSharedRef<class SEditorViewport> GetViewportWidget() override;
+	virtual TSharedPtr<FExtender> GetExtenders() const override;
 	virtual void OnFloatingButtonClicked() override {}
 
 private:
@@ -51,18 +48,16 @@ private:
 	// Use SEditorViewport::CommandList if we want command hotkeys to only be active when the mouse is in this viewport.
 	TSharedPtr<FUICommandList> ToolkitCommandList;
 
-	UE_API TWeakPtr<UE::Chaos::ClothAsset::FChaosClothPreviewScene> GetPreviewScene();
-	UE_API TWeakPtr<const UE::Chaos::ClothAsset::FChaosClothPreviewScene> GetPreviewScene() const;
+	TWeakPtr<UE::Chaos::ClothAsset::FChaosClothPreviewScene> GetPreviewScene();
+	TWeakPtr<const UE::Chaos::ClothAsset::FChaosClothPreviewScene> GetPreviewScene() const;
 
-	UE_API virtual TSharedPtr<IPreviewProfileController> CreatePreviewProfileController() override;
+	virtual TSharedPtr<IPreviewProfileController> CreatePreviewProfileController() override;
 	
-	UE_API FText GetViewportDisplayString() const;
+	FText GetViewportDisplayString() const;
 
-	UE_API float GetViewMinInput() const;
-	UE_API float GetViewMaxInput() const;
-	UE_API EVisibility GetAnimControlVisibility() const;
+	float GetViewMinInput() const;
+	float GetViewMaxInput() const;
+	EVisibility GetAnimControlVisibility() const;
 		 
 
 };
-
-#undef UE_API
diff --git a/Engine/Plugins/ChaosClothAssetEditor/Source/ChaosClothAssetEditor/Private/ChaosClothAsset/SClothEditor3DViewportToolBar.h b/Engine/Plugins/ChaosClothAssetEditor/Source/ChaosClothAssetEditor/Private/ChaosClothAsset/SClothEditor3DViewportToolBar.h
deleted file mode 100644
diff --git a/Engine/Plugins/ChaosClothAssetEditor/Source/ChaosClothAssetEditor/Private/ChaosClothAsset/SClothEditorRestSpaceViewport.h b/Engine/Plugins/ChaosClothAssetEditor/Source/ChaosClothAssetEditor/Private/ChaosClothAsset/SClothEditorRestSpaceViewport.h
--- a/Engine/Plugins/ChaosClothAssetEditor/Source/ChaosClothAssetEditor/Private/ChaosClothAsset/SClothEditorRestSpaceViewport.h
+++ b/Engine/Plugins/ChaosClothAssetEditor/Source/ChaosClothAssetEditor/Private/ChaosClothAsset/SClothEditorRestSpaceViewport.h
@@ -5,14 +5,11 @@
 #include "CoreMinimal.h"
 #include "SBaseCharacterFXEditorViewport.h"
 #include "SCommonEditorViewportToolbarBase.h"
-#include "SClothEditorRestSpaceViewportToolBar.h"
-
-#define UE_API CHAOSCLOTHASSETEDITOR_API
 
 class UChaosClothAssetEditorMode;
 namespace UE::Chaos::ClothAsset { class FChaosClothEditorRestSpaceViewportClient; }
 
-class SChaosClothAssetEditorRestSpaceViewport : public SBaseCharacterFXEditorViewport, public ICommonEditorViewportToolbarInfoProvider
+class CHAOSCLOTHASSETEDITOR_API SChaosClothAssetEditorRestSpaceViewport : public SBaseCharacterFXEditorViewport, public ICommonEditorViewportToolbarInfoProvider
 {
 public:
 
@@ -21,29 +18,26 @@ public:
 		SLATE_ARGUMENT(TSharedPtr<UE::Chaos::ClothAsset::FChaosClothEditorRestSpaceViewportClient>, RestSpaceViewportClient)
 	SLATE_END_ARGS()
 
-	UE_API void Construct(const FArguments& InArgs, const FAssetEditorViewportConstructionArgs& InViewportConstructionArgs);
+	void Construct(const FArguments& InArgs, const FAssetEditorViewportConstructionArgs& InViewportConstructionArgs);
 
 	// SEditorViewport
-	UE_API virtual void BindCommands() override;
-	UE_API virtual TSharedPtr<SWidget> MakeViewportToolbar() override;
-	UE_API virtual TSharedPtr<SWidget> BuildViewportToolbar() override;
-	UE_API virtual void PopulateViewportOverlays(TSharedRef<SOverlay> Overlay) override;
-	UE_API virtual void OnFocusViewportToSelection() override;
-	UE_API virtual bool IsVisible() const override;
+	virtual void BindCommands() override;
+	virtual TSharedPtr<SWidget> BuildViewportToolbar() override;
+	virtual void PopulateViewportOverlays(TSharedRef<SOverlay> Overlay) override;
+	virtual void OnFocusViewportToSelection() override;
+	virtual bool IsVisible() const override;
 
 	// ICommonEditorViewportToolbarInfoProvider
-	UE_API virtual TSharedRef<class SEditorViewport> GetViewportWidget() override;
-	UE_API virtual TSharedPtr<FExtender> GetExtenders() const override;
-	UE_API virtual void OnFloatingButtonClicked() override;
+	virtual TSharedRef<class SEditorViewport> GetViewportWidget() override;
+	virtual TSharedPtr<FExtender> GetExtenders() const override;
+	virtual void OnFloatingButtonClicked() override;
 
 private:
 
-	UE_API UChaosClothAssetEditorMode* GetEdMode() const;
+	UChaosClothAssetEditorMode* GetEdMode() const;
 	
-	UE_API FText GetDisplayString() const; 
+	FText GetDisplayString() const; 
 
 	TSharedPtr<UE::Chaos::ClothAsset::FChaosClothEditorRestSpaceViewportClient> RestSpaceViewportClient;
 
 };
-
-#undef UE_API
diff --git a/Engine/Plugins/ChaosClothAssetEditor/Source/ChaosClothAssetEditor/Private/ChaosClothAsset/SClothEditorRestSpaceViewportToolBar.h b/Engine/Plugins/ChaosClothAssetEditor/Source/ChaosClothAssetEditor/Private/ChaosClothAsset/SClothEditorRestSpaceViewportToolBar.h
deleted file mode 100644
diff --git a/Engine/Plugins/ChaosClothAssetEditor/Source/ChaosClothAssetEditor/Private/ChaosClothAsset/SClothEditorViewportToolBarBase.h b/Engine/Plugins/ChaosClothAssetEditor/Source/ChaosClothAssetEditor/Private/ChaosClothAsset/SClothEditorViewportToolBarBase.h
deleted file mode 100644
diff --git a/Engine/Plugins/ChaosClothAssetEditor/Source/ChaosClothAssetEditorTools/Private/ChaosClothAsset/ClothWeightMapPaintTool.h b/Engine/Plugins/ChaosClothAssetEditor/Source/ChaosClothAssetEditorTools/Private/ChaosClothAsset/ClothWeightMapPaintTool.h
--- a/Engine/Plugins/ChaosClothAssetEditor/Source/ChaosClothAssetEditorTools/Private/ChaosClothAsset/ClothWeightMapPaintTool.h
+++ b/Engine/Plugins/ChaosClothAssetEditor/Source/ChaosClothAssetEditorTools/Private/ChaosClothAsset/ClothWeightMapPaintTool.h
@@ -26,8 +26,6 @@
 
 #include "ClothWeightMapPaintTool.generated.h"
 
-#define UE_API CHAOSCLOTHASSETEDITORTOOLS_API
-
 class UMeshElementsVisualizer;
 class UWeightMapEraseBrushOpProps;
 class UWeightMapPaintBrushOpProps;
@@ -127,8 +125,8 @@ enum class EClothEditorWeightMapPaintQueryType : uint8
 
 // TODO: Look at EditConditions for all these properties. Which ones make sense for which SubToolType?
 
-UCLASS(MinimalAPI)
-class UClothEditorWeightMapPaintBrushFilterProperties : public UInteractiveToolPropertySet
+UCLASS()
+class CHAOSCLOTHASSETEDITORTOOLS_API UClothEditorWeightMapPaintBrushFilterProperties : public UInteractiveToolPropertySet
 {
 	GENERATED_BODY()
 
@@ -234,8 +232,8 @@ enum class EClothEditorWeightMapPaintToolActions
 };
 
 
-UCLASS(MinimalAPI)
-class UClothEditorMeshWeightMapPaintToolActions : public UInteractiveToolPropertySet
+UCLASS()
+class CHAOSCLOTHASSETEDITORTOOLS_API UClothEditorMeshWeightMapPaintToolActions : public UInteractiveToolPropertySet
 {
 	GENERATED_BODY()
 
@@ -245,7 +243,7 @@ public:
 
 	void Initialize(UClothEditorWeightMapPaintTool* ParentToolIn) { ParentTool = ParentToolIn; }
 
-	UE_API void PostAction(EClothEditorWeightMapPaintToolActions Action);
+	void PostAction(EClothEditorWeightMapPaintToolActions Action);
 
 	UFUNCTION(CallInEditor, Category = Operations, meta = (DisplayPriority = 10))
 	void ClearAll()
@@ -273,8 +271,8 @@ public:
 
 };
 
-UCLASS(MinimalAPI)
-class UClothEditorUpdateWeightMapProperties : public UInteractiveToolPropertySet
+UCLASS()
+class CHAOSCLOTHASSETEDITORTOOLS_API UClothEditorUpdateWeightMapProperties : public UInteractiveToolPropertySet
 {
 	GENERATED_BODY()
 
@@ -288,12 +286,12 @@ public:
 
 private:
 
-	UE_API virtual void PostEditChangeProperty(struct FPropertyChangedEvent& PropertyChangedEvent) override;
+	virtual void PostEditChangeProperty(struct FPropertyChangedEvent& PropertyChangedEvent) override;
 };
 
 /// Property panel for controlling hiding triangles
-UCLASS(MinimalAPI)
-class UClothEditorMeshWeightMapPaintToolShowHideProperties : public UInteractiveToolPropertySet
+UCLASS()
+class CHAOSCLOTHASSETEDITORTOOLS_API UClothEditorMeshWeightMapPaintToolShowHideProperties : public UInteractiveToolPropertySet
 {
 	GENERATED_BODY()
 
@@ -301,7 +299,7 @@ public:
 
 	TWeakObjectPtr<UClothEditorWeightMapPaintTool> ParentTool;
 	void Initialize(UClothEditorWeightMapPaintTool* ParentToolIn) { ParentTool = ParentToolIn; }
-	UE_API void PostAction(EClothEditorWeightMapPaintToolActions Action);
+	void PostAction(EClothEditorWeightMapPaintToolActions Action);
 
 	/** Toggles whether each pattern is shown or hidden */
 	UPROPERTY(EditAnywhere, EditFixedSize, Transient, Category = TriangleVisibility, meta = (DisplayPriority = 2));
@@ -319,46 +317,46 @@ public:
 /**
  * Mesh Element Paint Tool Class
  */
-UCLASS(MinimalAPI)
-class UClothEditorWeightMapPaintTool : public UMeshSculptToolBase
+UCLASS()
+class CHAOSCLOTHASSETEDITORTOOLS_API UClothEditorWeightMapPaintTool : public UMeshSculptToolBase
 {
 	GENERATED_BODY()
 
 public:
-	UE_API virtual void RegisterActions(FInteractiveToolActionSet& ActionSet) override;
+	virtual void RegisterActions(FInteractiveToolActionSet& ActionSet) override;
 
-	UE_API virtual void Setup() override;
-	UE_API virtual void Shutdown(EToolShutdownType ShutdownType) override;
-	UE_API virtual void DrawHUD(FCanvas* Canvas, IToolsContextRenderAPI* RenderAPI) override;
-	UE_API virtual void Render(IToolsContextRenderAPI* RenderAPI) override;
-	UE_API virtual void UpdateMaterialMode(EMeshEditingMaterialModes MaterialMode) override;
+	virtual void Setup() override;
+	virtual void Shutdown(EToolShutdownType ShutdownType) override;
+	virtual void DrawHUD(FCanvas* Canvas, IToolsContextRenderAPI* RenderAPI) override;
+	virtual void Render(IToolsContextRenderAPI* RenderAPI) override;
+	virtual void UpdateMaterialMode(EMeshEditingMaterialModes MaterialMode) override;
 
-	UE_API virtual void UpdateStampPendingState() override;
+	virtual void UpdateStampPendingState() override;
 
-	UE_API virtual void OnTick(float DeltaTime) override;
+	virtual void OnTick(float DeltaTime) override;
 
 	virtual bool HasCancel() const override { return true; }
 	virtual bool HasAccept() const override { return true; }
-	UE_API virtual bool CanAccept() const override;
+	virtual bool CanAccept() const override;
 
-	UE_API virtual bool OnUpdateHover(const FInputDeviceRay& DevicePos) override;
+	virtual bool OnUpdateHover(const FInputDeviceRay& DevicePos) override;
 
-	UE_API virtual void OnPropertyModified(UObject* PropertySet, FProperty* Property) override;
+	virtual void OnPropertyModified(UObject* PropertySet, FProperty* Property) override;
 
-	UE_API bool IsInBrushSubMode() const;
+	bool IsInBrushSubMode() const;
 
-	UE_API virtual void CommitResult(UBaseDynamicMeshComponent* Component, bool bModifiedTopology) override;
+	virtual void CommitResult(UBaseDynamicMeshComponent* Component, bool bModifiedTopology) override;
 
-	UE_API void SetDataflowContextObject(TObjectPtr<UDataflowContextObject> InDataflowContextObject);
+	void SetDataflowContextObject(TObjectPtr<UDataflowContextObject> InDataflowContextObject);
 
 private:
 
 	// Initialization support
-	UE_API void InitializeSculptMeshFromTarget();
-	UE_API void UpdateShowHideProperties();
+	void InitializeSculptMeshFromTarget();
+	void UpdateShowHideProperties();
 
 	// Make sure things are set up correctly after initialization or re-initialization
-	UE_API void PostSetupCheck() const;
+	void PostSetupCheck() const;
 
 public:
 
@@ -378,24 +376,24 @@ private:
 	UPROPERTY()
 	TObjectPtr<UWeightMapEraseBrushOpProps> EraseBrushOpProperties;
 
-	UE_API void FloodFillCurrentWeightAction();
-	UE_API void ClearAllWeightsAction();
-	UE_API void InvertWeightsAction();
-	UE_API void MultiplyWeightsAction();
-	UE_API void ClearHiddenAction();
+	void FloodFillCurrentWeightAction();
+	void ClearAllWeightsAction();
+	void InvertWeightsAction();
+	void MultiplyWeightsAction();
+	void ClearHiddenAction();
 
 public:
-	UE_API void SetVerticesToWeightMap(const TSet<int32>& Vertices, double WeightValue, bool bIsErase);
+	void SetVerticesToWeightMap(const TSet<int32>& Vertices, double WeightValue, bool bIsErase);
 
-	UE_API bool HaveVisibilityFilter() const;
-	UE_API void ApplyVisibilityFilter(const TArray<int32>& Vertices, TArray<int32>& VisibleVertices);
-	UE_API void ApplyVisibilityFilter(TSet<int32>& Triangles, TArray<int32>& ROIBuffer, TArray<int32>& OutputBuffer);
+	bool HaveVisibilityFilter() const;
+	void ApplyVisibilityFilter(const TArray<int32>& Vertices, TArray<int32>& VisibleVertices);
+	void ApplyVisibilityFilter(TSet<int32>& Triangles, TArray<int32>& ROIBuffer, TArray<int32>& OutputBuffer);
 
 	// we override these so we can update the separate BrushSize property added for this tool
-	UE_API virtual void IncreaseBrushRadiusAction();
-	UE_API virtual void DecreaseBrushRadiusAction();
-	UE_API virtual void IncreaseBrushRadiusSmallStepAction();
-	UE_API virtual void DecreaseBrushRadiusSmallStepAction();
+	virtual void IncreaseBrushRadiusAction();
+	virtual void DecreaseBrushRadiusAction();
+	virtual void IncreaseBrushRadiusSmallStepAction();
+	virtual void DecreaseBrushRadiusSmallStepAction();
 
 protected:
 	// UMeshSculptToolBase API
@@ -403,25 +401,27 @@ protected:
 	virtual FDynamicMesh3* GetBaseMesh() { check(false); return nullptr; }
 	virtual const FDynamicMesh3* GetBaseMesh() const { check(false); return nullptr; }
 
-	UE_API virtual int32 FindHitSculptMeshTriangle(const FRay3d& LocalRay) override;
-	UE_API virtual int32 FindHitTargetMeshTriangle(const FRay3d& LocalRay) override;
+	virtual int32 FindHitSculptMeshTriangleConst(const FRay3d& LocalRay) const override;
+	virtual int32 FindHitTargetMeshTriangleConst(const FRay3d& LocalRay) const override;
+	virtual void UpdateHitSculptMeshTriangle(int32 TriangleID, const FRay3d& LocalRay) override;
+
 
-	UE_API virtual void OnBeginStroke(const FRay& WorldRay) override;
-	UE_API virtual void OnCancelStroke() override;
-	UE_API virtual void OnEndStroke() override;
+	virtual void OnBeginStroke(const FRay& WorldRay) override;
+	virtual void OnCancelStroke() override;
+	virtual void OnEndStroke() override;
 
-	UE_API virtual TUniquePtr<FMeshSculptBrushOp>& GetActiveBrushOp();
+	virtual TUniquePtr<FMeshSculptBrushOp>& GetActiveBrushOp();
 
 	virtual bool SharesBrushPropertiesChanges() const override { return false; }
 
-	UE_API virtual void InitializeBrushSizeRange(const UE::Geometry::FAxisAlignedBox3d& TargetBounds) override;
+	virtual void InitializeBrushSizeRange(const UE::Geometry::FAxisAlignedBox3d& TargetBounds) override;
 
-	UE_API virtual void NextBrushModeAction() override;
-	UE_API virtual void PreviousBrushModeAction() override;
+	virtual void NextBrushModeAction() override;
+	virtual void PreviousBrushModeAction() override;
 
 	// Note: these will actually modify the brush's Attribute Value since we don't use Speed in our brush
-	UE_API virtual void IncreaseBrushSpeedAction() override;
-	UE_API virtual void DecreaseBrushSpeedAction() override;
+	virtual void IncreaseBrushSpeedAction() override;
+	virtual void DecreaseBrushSpeedAction() override;
 
 
 	// end UMeshSculptToolBase API
@@ -433,7 +433,7 @@ protected:
 	//
 
 public:
-	UE_API virtual void RequestAction(EClothEditorWeightMapPaintToolActions ActionType);
+	virtual void RequestAction(EClothEditorWeightMapPaintToolActions ActionType);
 	
 	UPROPERTY()
 	TObjectPtr<UClothEditorMeshWeightMapPaintToolActions> ActionsProps;
@@ -444,7 +444,7 @@ public:
 protected:
 	bool bHavePendingAction = false;
 	EClothEditorWeightMapPaintToolActions PendingAction;
-	UE_API virtual void ApplyAction(EClothEditorWeightMapPaintToolActions ActionType);
+	virtual void ApplyAction(EClothEditorWeightMapPaintToolActions ActionType);
 
 
 
@@ -456,7 +456,7 @@ public:
 	TObjectPtr<UPolyLassoMarqueeMechanic> PolyLassoMechanic;
 
 protected:
-	UE_API void OnPolyLassoFinished(const FCameraPolyLasso& Lasso, bool bCanceled);
+	void OnPolyLassoFinished(const FCameraPolyLasso& Lasso, bool bCanceled);
 
 
 	// 
@@ -474,8 +474,8 @@ protected:
 	UE::Geometry::FGroupTopologySelection LowValueGradientVertexSelection;
 	UE::Geometry::FGroupTopologySelection HighValueGradientVertexSelection;
 
-	UE_API void ComputeGradient();
-	UE_API void OnSelectionModified();
+	void ComputeGradient();
+	void OnSelectionModified();
 
 
 	//
@@ -518,44 +518,44 @@ protected:
 	TObjectPtr<UDataflowContextObject> DataflowContextObject = nullptr;
 
 	// realtime visualization
-	UE_API void OnDynamicMeshComponentChanged(UDynamicMeshComponent* Component, const FMeshVertexChange* Change, bool bRevert);
+	void OnDynamicMeshComponentChanged(UDynamicMeshComponent* Component, const FMeshVertexChange* Change, bool bRevert);
 	FDelegateHandle OnDynamicMeshComponentChangedHandle;
 
 	TConstArrayView<float> InputWeightMap;
 
 	UE::Geometry::FDynamicMeshWeightAttribute* ActiveWeightMap;
-	UE_API double GetCurrentWeightValue(int32 VertexId) const;
-	UE_API double GetCurrentWeightValueUnderBrush() const;
+	double GetCurrentWeightValue(int32 VertexId) const;
+	double GetCurrentWeightValueUnderBrush() const;
 	FVector3d CurrentBaryCentricCoords;
-	UE_API int32 GetBrushNearestVertex() const;
-	UE_API int32 GetBrushNearestVertexAccurate() const;
+	int32 GetBrushNearestVertex() const;
+	int32 GetBrushNearestVertexAccurate() const;
 
-	UE_API void GetCurrentWeightMap(TArray<float>& OutWeights) const;
+	void GetCurrentWeightMap(TArray<float>& OutWeights) const;
 
-	UE_API void UpdateSubToolType(EClothEditorWeightMapPaintInteractionType NewType);
+	void UpdateSubToolType(EClothEditorWeightMapPaintInteractionType NewType);
 
-	UE_API void UpdateBrushType(EClothEditorWeightMapPaintBrushType BrushType);
+	void UpdateBrushType(EClothEditorWeightMapPaintBrushType BrushType);
 
 	TSet<int32> AccumulatedTriangleROI;
 	bool bUndoUpdatePending = false;
 	TArray<int> NormalsBuffer;
-	UE_API void WaitForPendingUndoRedo();
+	void WaitForPendingUndoRedoUpdate();
 
 	TArray<int> TempROIBuffer;
 	TArray<int> VertexROI;
 	TArray<bool> VisibilityFilterBuffer;
 	TSet<int> VertexSetBuffer;
 	TSet<int> TriangleROI;
-	UE_API void UpdateROI(const FSculptBrushStamp& CurrentStamp);
+	void UpdateROI(const FSculptBrushStamp& CurrentStamp);
 
 	EClothEditorWeightMapPaintBrushType PendingStampType = EClothEditorWeightMapPaintBrushType::Paint;
 
-	UE_API bool UpdateStampPosition(const FRay& WorldRay);
-	UE_API bool ApplyStamp();
+	bool UpdateStampPosition(const FRay& WorldRay);
+	bool ApplyStamp();
 
 	UE::Geometry::FDynamicMeshOctree3 Octree;
 
-	UE_API bool UpdateBrushPosition(const FRay& WorldRay);
+	bool UpdateBrushPosition(const FRay& WorldRay);
 
 	bool GetInEraseStroke()
 	{
@@ -568,8 +568,8 @@ protected:
 
 	TArray<int32> ROITriangleBuffer;
 	TArray<double> ROIWeightValueBuffer;
-	UE_API bool SyncMeshWithWeightBuffer(FDynamicMesh3* Mesh);
-	UE_API bool SyncWeightBufferWithMesh(const FDynamicMesh3* Mesh);
+	bool SyncMeshWithWeightBuffer(FDynamicMesh3* Mesh);
+	bool SyncWeightBufferWithMesh(const FDynamicMesh3* Mesh);
 
 
 	// Undo/Redo change support
@@ -583,24 +583,24 @@ protected:
 		}
 	};
 	TUniquePtr<TIndexedValuesChangeBuilder<float, FNodeBufferWeightChange>> ActiveChangeBuilder;
-	UE_API void BeginChange();
-	UE_API void EndChange();
+	void BeginChange();
+	void EndChange();
 
 	// The corresponding FChaosClothAssetWeightMapNode has a buffer of scalar weights. Depending on bHaveDynamicMeshToWeightConversion, the 
 	// indexing of this buffer might be different than the mesh vertex indexing. This function returns the corresponding index in the node buffer for
 	// a given mesh vertex index.
-	UE_API int32 MeshIndexToNodeIndex(int32 MeshVertexIndex) const;
+	int32 MeshIndexToNodeIndex(int32 MeshVertexIndex) const;
 
 	// Update the active weight map attribute from the map of indices/values. The inputs are given in "node buffer format" (see function above).
-	UE_API void UpdateMapValuesFromNodeValues(const TArray<int32>& Indices, const TArray<float>& Values);
+	void UpdateMapValuesFromNodeValues(const TArray<int32>& Indices, const TArray<float>& Values);
 
 
-	UE_API FColor GetColorForWeightValue(double WeightValue);
+	FColor GetColorForWeightValue(double WeightValue);
 
 	TArray<FVector3d> TriNormals;
 	TArray<int32> UVSeamEdges;
 	TArray<int32> NormalSeamEdges;
-	UE_API void PrecomputeFilterData();
+	void PrecomputeFilterData();
 
 	// DynamicMesh might be unwelded mesh, but weights are on the welded mesh.
 	bool bHaveDynamicMeshToWeightConversion = false;
@@ -611,7 +611,7 @@ protected:
 	virtual bool ShowWorkPlane() const override { return false; }
 
 	friend class UClothEditorWeightMapPaintToolBuilder;
-	UE_API void NotifyTargetChanged();
+	void NotifyTargetChanged();
 
 	bool bAnyChangeMade = false;
 
@@ -619,12 +619,11 @@ protected:
 
 	FChaosClothAssetWeightMapNode* WeightMapNodeToUpdate = nullptr;
 
-	UE_API void UpdateSelectedNode();
+	void UpdateSelectedNode();
 
-	UE_API void UpdateVertexColorOverlay(const TSet<int>* TrianglesToUpdate = nullptr);
+	void UpdateVertexColorOverlay(const TSet<int>* TrianglesToUpdate = nullptr);
 
 };
 
 
 
-#undef UE_API
