Fixing engine modules that depended directly on Chaos instead of using the module setup function.

diff --git a/Engine/Plugins/ChaosClothAsset/Source/ChaosClothAsset/ChaosClothAsset.Build.cs b/Engine/Plugins/ChaosClothAsset/Source/ChaosClothAsset/ChaosClothAsset.Build.cs
--- a/Engine/Plugins/ChaosClothAsset/Source/ChaosClothAsset/ChaosClothAsset.Build.cs
+++ b/Engine/Plugins/ChaosClothAsset/Source/ChaosClothAsset/ChaosClothAsset.Build.cs
@@ -18,11 +18,6 @@ public class ChaosClothAsset : ModuleRules
 				"RenderCore"
 			}
 		);
-		PublicDependencyModuleNames.AddRange(
-			new string[]
-			{
-				"Chaos"
-			}
-		);
+		SetupModulePhysicsSupport(Target);
 	}
 }
diff --git a/Engine/Plugins/ChaosClothAsset/Source/ChaosClothAssetEngine/ChaosClothAssetEngine.Build.cs b/Engine/Plugins/ChaosClothAsset/Source/ChaosClothAssetEngine/ChaosClothAssetEngine.Build.cs
--- a/Engine/Plugins/ChaosClothAsset/Source/ChaosClothAssetEngine/ChaosClothAssetEngine.Build.cs
+++ b/Engine/Plugins/ChaosClothAsset/Source/ChaosClothAssetEngine/ChaosClothAssetEngine.Build.cs
@@ -26,13 +26,13 @@ public class ChaosClothAssetEngine : ModuleRules
 				"CoreUObject",
 				"RenderCore",
 				"RHI",
-				"Chaos",
 				"ChaosClothAsset",
 				"ChaosCloth",
 				"ChaosCaching",
 				"ClothingSystemRuntimeInterface"
 			}
 		);
+		SetupModulePhysicsSupport(Target);
 
 		if (Target.bBuildEditor || Target.bCompileAgainstEditor)
 		{
diff --git a/Engine/Plugins/ChaosClothAssetEditor/Source/ChaosClothAssetDataflowNodes/ChaosClothAssetDataflowNodes.Build.cs b/Engine/Plugins/ChaosClothAssetEditor/Source/ChaosClothAssetDataflowNodes/ChaosClothAssetDataflowNodes.Build.cs
--- a/Engine/Plugins/ChaosClothAssetEditor/Source/ChaosClothAssetDataflowNodes/ChaosClothAssetDataflowNodes.Build.cs
+++ b/Engine/Plugins/ChaosClothAssetEditor/Source/ChaosClothAssetDataflowNodes/ChaosClothAssetDataflowNodes.Build.cs
@@ -15,11 +15,11 @@ public class ChaosClothAssetDataflowNodes : ModuleRules
 			}
 		);
 
+		SetupModulePhysicsSupport(Target);
 		PrivateDependencyModuleNames.AddRange(
 			new string[]
 			{
 				"AssetTools",
-				"Chaos",
 				"ChaosCaching",
 				"ChaosCloth",
 				"ChaosClothAsset",
diff --git a/Engine/Plugins/Experimental/ChaosFlesh/Source/ChaosFleshEditor/ChaosFleshEditor.Build.cs b/Engine/Plugins/Experimental/ChaosFlesh/Source/ChaosFleshEditor/ChaosFleshEditor.Build.cs
--- a/Engine/Plugins/Experimental/ChaosFlesh/Source/ChaosFleshEditor/ChaosFleshEditor.Build.cs
+++ b/Engine/Plugins/Experimental/ChaosFlesh/Source/ChaosFleshEditor/ChaosFleshEditor.Build.cs
@@ -20,6 +20,7 @@ namespace UnrealBuildTool.Rules
 				}
 				);
 
+			SetupModulePhysicsSupport(Target);
 			PublicDependencyModuleNames.AddRange(
 				new string[]
 				{
@@ -27,7 +28,6 @@ namespace UnrealBuildTool.Rules
 					"AssetDefinition",
 					"AssetRegistry",
 					"AssetTools",
-					"Chaos",
 					"ChaosCaching",
 					"ChaosCachingUSD",
 					"ChaosFlesh",
diff --git a/Engine/Plugins/Experimental/Fracture/Source/FractureEngine/FractureEngine.Build.cs b/Engine/Plugins/Experimental/Fracture/Source/FractureEngine/FractureEngine.Build.cs
--- a/Engine/Plugins/Experimental/Fracture/Source/FractureEngine/FractureEngine.Build.cs
+++ b/Engine/Plugins/Experimental/Fracture/Source/FractureEngine/FractureEngine.Build.cs
@@ -9,7 +9,6 @@ namespace UnrealBuildTool.Rules
 			PublicDependencyModuleNames.AddRange(
 				new string[]
 				{
-					"Chaos",
 					"Core",
 					"CoreUObject",
 					"DataflowCore",
@@ -20,6 +19,7 @@ namespace UnrealBuildTool.Rules
 					"MeshUtilitiesCommon", // for FDisjointSet
 				}
 				);
+			SetupModulePhysicsSupport(Target);
 		}
 	}
 }
diff --git a/Engine/Plugins/PCG/Source/PCG/PCG.Build.cs b/Engine/Plugins/PCG/Source/PCG/PCG.Build.cs
--- a/Engine/Plugins/PCG/Source/PCG/PCG.Build.cs
+++ b/Engine/Plugins/PCG/Source/PCG/PCG.Build.cs
@@ -11,7 +11,6 @@ namespace UnrealBuildTool.Rules
 			PublicDependencyModuleNames.AddRange(
 				new string[]
 				{
-					"Chaos",
 					"ComputeFramework",
 					"Core",
 					"CoreUObject",
@@ -24,6 +23,7 @@ namespace UnrealBuildTool.Rules
 					"RenderCore",
 					"RHI",
 				});
+			SetupModulePhysicsSupport(Target);
 
 			PrivateDependencyModuleNames.AddRange(
 				new string[]
diff --git a/Engine/Source/Editor/PhysicsAssetEditor/PhysicsAssetEditor.Build.cs b/Engine/Source/Editor/PhysicsAssetEditor/PhysicsAssetEditor.Build.cs
--- a/Engine/Source/Editor/PhysicsAssetEditor/PhysicsAssetEditor.Build.cs
+++ b/Engine/Source/Editor/PhysicsAssetEditor/PhysicsAssetEditor.Build.cs
@@ -41,9 +41,9 @@ public class PhysicsAssetEditor : ModuleRules
 				"ApplicationCore",
 				"EditorStyle",
 				"ToolWidgets",
-				"Chaos"
 			}
         );
+		SetupModulePhysicsSupport(Target);
 
 		DynamicallyLoadedModuleNames.AddRange(
 			new string[] {
