Chaos cloth:  deprecate kinematic collider from self collision  nodes.

These have been superceded by the SkinnedTriangleMesh Physics Asset type.

[RN] Chaos cloth's self collision-based kinematic collider has been deprecated in favor if SkinnedTriangleMesh physics assets.

diff --git a/Engine/Plugins/ChaosClothAssetEditor/Source/ChaosClothAssetDataflowNodes/Private/ChaosClothAsset/ChaosClothAssetDataflowNodesModule.cpp b/Engine/Plugins/ChaosClothAssetEditor/Source/ChaosClothAssetDataflowNodes/Private/ChaosClothAsset/ChaosClothAssetDataflowNodesModule.cpp
--- a/Engine/Plugins/ChaosClothAssetEditor/Source/ChaosClothAssetDataflowNodes/Private/ChaosClothAsset/ChaosClothAssetDataflowNodesModule.cpp
+++ b/Engine/Plugins/ChaosClothAssetEditor/Source/ChaosClothAssetDataflowNodes/Private/ChaosClothAsset/ChaosClothAssetDataflowNodesModule.cpp
@@ -152,7 +152,7 @@ namespace UE::Chaos::ClothAsset
 			DATAFLOW_NODE_REGISTER_CREATION_FACTORY(FChaosClothAssetSimulationPBDBendingSpringConfigNode);
 			DATAFLOW_NODE_REGISTER_CREATION_FACTORY(FChaosClothAssetSimulationPBDEdgeSpringConfigNode);
 			DATAFLOW_NODE_REGISTER_CREATION_FACTORY(FChaosClothAssetSimulationPressureConfigNode);
-			DATAFLOW_NODE_REGISTER_CREATION_FACTORY(FChaosClothAssetSimulationSelfCollisionConfigNode);
+			DATAFLOW_NODE_REGISTER_CREATION_FACTORY(FChaosClothAssetSimulationSelfCollisionConfigNode_v2);
 			DATAFLOW_NODE_REGISTER_CREATION_FACTORY(FChaosClothAssetSimulationSelfCollisionSpheresConfigNode);
 			DATAFLOW_NODE_REGISTER_CREATION_FACTORY(FChaosClothAssetSimulationSolverConfigNode);
 			DATAFLOW_NODE_REGISTER_CREATION_FACTORY(FChaosClothAssetSimulationMultiResConfigNode);
@@ -185,6 +185,7 @@ PRAGMA_DISABLE_DEPRECATION_WARNINGS
 			DATAFLOW_NODE_REGISTER_CREATION_FACTORY(FChaosClothAssetRemeshNode);
 			DATAFLOW_NODE_REGISTER_CREATION_FACTORY(FChaosClothAssetSelectionNode);
 			DATAFLOW_NODE_REGISTER_CREATION_FACTORY(FChaosClothAssetSimulationLongRangeAttachmentConfigNode);
+			DATAFLOW_NODE_REGISTER_CREATION_FACTORY(FChaosClothAssetSimulationSelfCollisionConfigNode);
 			DATAFLOW_NODE_REGISTER_CREATION_FACTORY(FChaosClothAssetSkeletalMeshImportNode);
 			DATAFLOW_NODE_REGISTER_CREATION_FACTORY(FChaosClothAssetStaticMeshImportNode);
 			DATAFLOW_NODE_REGISTER_CREATION_FACTORY(FChaosClothAssetTerminalNode);
diff --git a/Engine/Plugins/ChaosClothAssetEditor/Source/ChaosClothAssetDataflowNodes/Private/ChaosClothAsset/SimulationSelfCollisionConfigNode.cpp b/Engine/Plugins/ChaosClothAssetEditor/Source/ChaosClothAssetDataflowNodes/Private/ChaosClothAsset/SimulationSelfCollisionConfigNode.cpp
--- a/Engine/Plugins/ChaosClothAssetEditor/Source/ChaosClothAssetDataflowNodes/Private/ChaosClothAsset/SimulationSelfCollisionConfigNode.cpp
+++ b/Engine/Plugins/ChaosClothAssetEditor/Source/ChaosClothAssetDataflowNodes/Private/ChaosClothAsset/SimulationSelfCollisionConfigNode.cpp
@@ -7,6 +7,53 @@
 
 #include UE_INLINE_GENERATED_CPP_BY_NAME(SimulationSelfCollisionConfigNode)
 
+
+FChaosClothAssetSimulationSelfCollisionConfigNode_v2::FChaosClothAssetSimulationSelfCollisionConfigNode_v2(const UE::Dataflow::FNodeParameters& InParam, FGuid InGuid)
+	: FChaosClothAssetSimulationBaseConfigNode(InParam, InGuid)
+{
+	RegisterCollectionConnections();
+	RegisterInputConnection(&SelfCollisionLayers.StringValue, GET_MEMBER_NAME_CHECKED(FChaosClothAssetConnectableIStringValue, StringValue))
+		.SetCanHidePin(true)
+		.SetPinIsHidden(true);
+	RegisterInputConnection(&SelfCollisionDisabledFaces.StringValue, GET_MEMBER_NAME_CHECKED(FChaosClothAssetConnectableIStringValue, StringValue))
+		.SetCanHidePin(true)
+		.SetPinIsHidden(true);
+	RegisterInputConnection(&SelfCollisionThickness.WeightMap)
+		.SetCanHidePin(true)
+		.SetPinIsHidden(true);
+}
+
+void FChaosClothAssetSimulationSelfCollisionConfigNode_v2::AddProperties(FPropertyHelper& PropertyHelper) const
+{
+	PropertyHelper.SetPropertyBool(this, &bUseSelfCollisions);
+	PropertyHelper.SetProperty(this, &SelfCollisionStiffness);
+	PropertyHelper.SetProperty(this, &SelfCollisionDisableNeighborDistance, {}, ECollectionPropertyFlags::None); // Non animatable
+	PropertyHelper.SetPropertyString(this, &SelfCollisionDisabledFaces);
+	PropertyHelper.SetPropertyBool(this, &bUseSelfIntersections);
+	PropertyHelper.SetPropertyBool(this, &bUseGlobalIntersectionAnalysis);
+	PropertyHelper.SetPropertyBool(this, &bUseContourMinimization);
+	PropertyHelper.SetProperty(this, &NumContourMinimizationPostSteps);
+	PropertyHelper.SetPropertyBool(this, &bUseGlobalPostStepContours);
+
+	PropertyHelper.SetFabricPropertyString<int32, FChaosClothAssetConnectableIStringValue>(FName(TEXT("SelfCollisionLayers")), SelfCollisionLayers,
+		[](const UE::Chaos::ClothAsset::FCollectionClothFabricFacade& FabricFacade)-> float
+		{
+			return FabricFacade.GetLayer();
+		}, {}, ECollectionPropertyFlags::None, UE::Chaos::ClothAsset::ClothCollectionGroup::SimFaces);
+
+	PropertyHelper.SetFabricPropertyWeighted(FName(TEXT("SelfCollisionThickness")), SelfCollisionThickness,
+		[](const UE::Chaos::ClothAsset::FCollectionClothFabricFacade& FabricFacade)-> float
+		{
+			return FabricFacade.GetCollisionThickness();
+		}, {});
+
+	PropertyHelper.SetFabricProperty(FName(TEXT("SelfCollisionFriction")), SelfCollisionFriction,
+		[](UE::Chaos::ClothAsset::FCollectionClothFabricFacade& FabricFacade)-> float
+		{
+			return FabricFacade.GetFriction();
+		}, {});
+}
+
 FChaosClothAssetSimulationSelfCollisionConfigNode::FChaosClothAssetSimulationSelfCollisionConfigNode(const UE::Dataflow::FNodeParameters& InParam, FGuid InGuid)
 	: FChaosClothAssetSimulationBaseConfigNode(InParam, InGuid)
 {
@@ -30,6 +77,7 @@ FChaosClothAssetSimulationSelfCollisionConfigNode::FChaosClothAssetSimulationSel
 
 void FChaosClothAssetSimulationSelfCollisionConfigNode::AddProperties(FPropertyHelper& PropertyHelper) const
 {
+PRAGMA_DISABLE_DEPRECATION_WARNINGS  // SetProperty functions are templated and cause deprecation warnings with the now deprecated v1
 	PropertyHelper.SetPropertyBool(this, &bUseSelfCollisions);
 	PropertyHelper.SetProperty(this, &SelfCollisionStiffness);
 	PropertyHelper.SetProperty(this, &SelfCollisionDisableNeighborDistance, {}, ECollectionPropertyFlags::None); // Non animatable
@@ -65,6 +113,7 @@ void FChaosClothAssetSimulationSelfCollisionConfigNode::AddProperties(FPropertyH
 		{
 			return FabricFacade.GetFriction();
 		}, {});
+PRAGMA_ENABLE_DEPRECATION_WARNINGS
 }
 
 void FChaosClothAssetSimulationSelfCollisionConfigNode::Serialize(FArchive& Ar)
diff --git a/Engine/Plugins/ChaosClothAssetEditor/Source/ChaosClothAssetDataflowNodes/Public/ChaosClothAsset/SimulationSelfCollisionConfigNode.h b/Engine/Plugins/ChaosClothAssetEditor/Source/ChaosClothAssetDataflowNodes/Public/ChaosClothAsset/SimulationSelfCollisionConfigNode.h
--- a/Engine/Plugins/ChaosClothAssetEditor/Source/ChaosClothAssetDataflowNodes/Public/ChaosClothAsset/SimulationSelfCollisionConfigNode.h
+++ b/Engine/Plugins/ChaosClothAssetEditor/Source/ChaosClothAssetDataflowNodes/Public/ChaosClothAsset/SimulationSelfCollisionConfigNode.h
@@ -6,9 +6,80 @@
 #include "ChaosClothAsset/ConnectableValue.h"
 #include "SimulationSelfCollisionConfigNode.generated.h"
 
-/** Self-collision repulsion forces (point-face) properties configuration node. */
+
+/** Self-collision repulsion forces (point-face) properties configuration node. 
+ *  Note that the kinematic collider has been deprecated in favor of SkinnedTriangleMesh Physics Asset bodies */
 USTRUCT(Meta = (DataflowCloth))
-struct FChaosClothAssetSimulationSelfCollisionConfigNode : public FChaosClothAssetSimulationBaseConfigNode
+struct FChaosClothAssetSimulationSelfCollisionConfigNode_v2 final : public FChaosClothAssetSimulationBaseConfigNode
+{
+	GENERATED_USTRUCT_BODY()
+	DATAFLOW_NODE_DEFINE_INTERNAL(FChaosClothAssetSimulationSelfCollisionConfigNode_v2, "SimulationSelfCollisionConfig", "Cloth", "Cloth Simulation Self Collision Config")
+
+public:
+
+	FChaosClothAssetSimulationSelfCollisionConfigNode_v2(const UE::Dataflow::FNodeParameters& InParam, FGuid InGuid = FGuid::NewGuid());
+
+private:
+
+	/** Activating this node will enable self collisions.*/
+	UPROPERTY(VisibleAnywhere, Category = "Self-Collision Properties", meta = (InteractorName = "UseSelfCollisions"))
+	bool bUseSelfCollisions = true;
+
+	/** The self collision offset per side. Total thickness of cloth is 2x this value. */
+	UPROPERTY(EditAnywhere, Category = "Self-Collision Properties", meta = (UIMin = "0", UIMax = "100", ClampMin = "0", ClampMax = "1000", InteractorName = "SelfCollisionThickness"))
+	FChaosClothAssetWeightedValue SelfCollisionThickness = { true, UE::Chaos::ClothAsset::FDefaultFabric::SelfCollisionThickness,
+		UE::Chaos::ClothAsset::FDefaultFabric::SelfCollisionThickness, TEXT("SelfCollisionThickness"), true };
+
+	/** The stiffness of the springs used to control self collision. */
+	UPROPERTY(EditAnywhere, Category = "Self-Collision Properties", meta = (UIMin = "0", UIMax = "1", ClampMin = "0", ClampMax = "1", InteractorName = "SelfCollisionStiffness"))
+	float SelfCollisionStiffness = 0.5f;
+
+	/** Friction coefficient for cloth - cloth interaction. */
+	UPROPERTY(EditAnywhere, Category = "Self-Collision Properties", meta = (UIMin = "0", UIMax = "1", ClampMin = "0", ClampMax = "1", InteractorName = "SelfCollisionFriction"))
+	FChaosClothAssetImportedFloatValue SelfCollisionFriction = { UE::Chaos::ClothAsset::FDefaultFabric::SelfFriction };
+
+	/** Disabled neighbor collision ring. Collisions are disabled between vertices within this N-ring connectivity distance.*/
+	UPROPERTY(EditAnywhere, Category = "Self-Collision Properties", meta = (UIMin = "1", UIMax = "5", ClampMin = "1", ClampMax = "10"))
+	int32 SelfCollisionDisableNeighborDistance = 5;
+
+	/** Self collision layers face int map. Generate this map using the SelectionsToIntMap node with SimFace Selections.
+	* Faces labeled with -1 will collide normally without any layering behavior.
+	* Faces labeled with any other number will keep higher layer numbers outside lower layer numbers (outside = front facing normal direction).
+	*/
+	UPROPERTY(EditAnywhere, Category = "Self-Collision Properties", meta = (InteractorName = "SelfCollisionLayers"))
+	FChaosClothAssetConnectableIStringValue SelfCollisionLayers = { TEXT("SelfCollisionLayers"), true };
+
+	/** Sim face selection set of faces which should not self collide */
+	UPROPERTY(EditAnywhere, Category = "Self-Collision Properties", meta = (InteractorName = "SelfCollisionDisabledFaces"))
+	FChaosClothAssetConnectableIStringValue SelfCollisionDisabledFaces = { TEXT("SelfCollisionDisabledFaces") };
+
+	/** Enable self intersection resolution. This will try to fix any cloth intersections that are not handled by collision repulsions. Can be expensive.*/
+	UPROPERTY(EditAnywhere, Category = "Self-Collision Properties", meta = (InteractorName = "UseSelfIntersections"))
+	bool bUseSelfIntersections = true;
+
+	/** Do global intersection analysis to determine the correct normals for the collision springs */
+	UPROPERTY(EditAnywhere, Category = "Self-Collision Properties", meta = (EditCondition = "bUseSelfIntersections", InteractorName = "UseGlobalIntersectionAnalysis"))
+	bool bUseGlobalIntersectionAnalysis = true;
+
+	/** Do a step of contour minimization at the beginning of the timestep. 
+	 * Contour minimization will attempt to resolve intersections by shortening the intersection edge. Helpful with open intersections that global intersection analysis can't fix.*/
+	UPROPERTY(EditAnywhere, Category = "Self-Collision Properties", meta = (EditCondition = "bUseSelfIntersections", InteractorName = "UseContourMinimization"))
+	bool bUseContourMinimization = true;
+
+	/** Number of post timestep contour minimization steps to do. (Very Expensive!)
+	 * Contour minimization will attempt to resolve intersections by shortening the intersection edge.Helpful with open intersections that global intersection analysis can't fix.*/
+	UPROPERTY(EditAnywhere, Category = "Self-Collision Properties", meta = (ClampMin = "0", EditCondition = "bUseSelfIntersections", InteractorName = "NumContourMinimizationPostSteps"))
+	int32 NumContourMinimizationPostSteps = 0;
+
+	/** Use global contour gradients when doing post timestep contour minimization */
+	UPROPERTY(EditAnywhere, Category = "Self-Collision Properties", meta = (EditCondition = "bUseSelfIntersections && NumContourMinimizationPostSteps > 0", InteractorName = "UseGlobalPostStepContours"))
+	bool bUseGlobalPostStepContours = true;
+	virtual void AddProperties(FPropertyHelper& PropertyHelper) const override;
+};
+
+/** Self-collision repulsion forces (point-face) properties configuration node. */
+USTRUCT(Meta = (DataflowCloth, Deprecated = 5.7))
+struct  UE_DEPRECATED(5.7, "Use the newer version of this node instead.") FChaosClothAssetSimulationSelfCollisionConfigNode : public FChaosClothAssetSimulationBaseConfigNode
 {
 	GENERATED_USTRUCT_BODY()
 	DATAFLOW_NODE_DEFINE_INTERNAL(FChaosClothAssetSimulationSelfCollisionConfigNode, "SimulationSelfCollisionConfig", "Cloth", "Cloth Simulation Self Collision Config")
