Dataflow : fix ClothWeightmap paint tool where the brush was not working anymore
- remove the UpdateStampPendingState override method to make sure the stamp gets properly processed

diff --git a/Engine/Plugins/ChaosClothAssetEditor/Source/ChaosClothAssetEditorTools/Private/ChaosClothAsset/ClothWeightMapPaintTool.cpp b/Engine/Plugins/ChaosClothAssetEditor/Source/ChaosClothAssetEditorTools/Private/ChaosClothAsset/ClothWeightMapPaintTool.cpp
--- a/Engine/Plugins/ChaosClothAssetEditor/Source/ChaosClothAssetEditorTools/Private/ChaosClothAsset/ClothWeightMapPaintTool.cpp
+++ b/Engine/Plugins/ChaosClothAssetEditor/Source/ChaosClothAssetEditorTools/Private/ChaosClothAsset/ClothWeightMapPaintTool.cpp
@@ -1878,12 +1878,6 @@ void UClothEditorWeightMapPaintTool::UpdateMaterialMode(EMeshEditingMaterialMode
 	}
 }
 
-void UClothEditorWeightMapPaintTool::UpdateStampPendingState()
-{
-	if (InStroke() == false) return;
-	bIsStampPending = true;
-}
-
 void UClothEditorWeightMapPaintTool::OnTick(float DeltaTime)
 {
 	UMeshSculptToolBase::OnTick(DeltaTime);
diff --git a/Engine/Plugins/ChaosClothAssetEditor/Source/ChaosClothAssetEditorTools/Private/ChaosClothAsset/ClothWeightMapPaintTool.h b/Engine/Plugins/ChaosClothAssetEditor/Source/ChaosClothAssetEditorTools/Private/ChaosClothAsset/ClothWeightMapPaintTool.h
--- a/Engine/Plugins/ChaosClothAssetEditor/Source/ChaosClothAssetEditorTools/Private/ChaosClothAsset/ClothWeightMapPaintTool.h
+++ b/Engine/Plugins/ChaosClothAssetEditor/Source/ChaosClothAssetEditorTools/Private/ChaosClothAsset/ClothWeightMapPaintTool.h
@@ -333,8 +333,6 @@ public:
 	UE_API virtual void Render(IToolsContextRenderAPI* RenderAPI) override;
 	UE_API virtual void UpdateMaterialMode(EMeshEditingMaterialModes MaterialMode) override;
 
-	UE_API virtual void UpdateStampPendingState() override;
-
 	UE_API virtual void OnTick(float DeltaTime) override;
 
 	virtual bool HasCancel() const override { return true; }
