Chaos Cloth Asset:
- Adjusted the schemaless USD import code path message as to not confuse users into thinking it was an error.
- Fixed missing render skin weights causing ensure.

[RN] Adjusted the schemaless USD import code path message as to not confuse users into thinking it was an error.

diff --git a/Engine/Plugins/ChaosClothAssetEditor/Source/ChaosClothAssetDataflowNodes/Private/ChaosClothAsset/USDImportNode.cpp b/Engine/Plugins/ChaosClothAssetEditor/Source/ChaosClothAssetDataflowNodes/Private/ChaosClothAsset/USDImportNode.cpp
--- a/Engine/Plugins/ChaosClothAssetEditor/Source/ChaosClothAssetDataflowNodes/Private/ChaosClothAsset/USDImportNode.cpp
+++ b/Engine/Plugins/ChaosClothAssetEditor/Source/ChaosClothAssetDataflowNodes/Private/ChaosClothAsset/USDImportNode.cpp
@@ -94,7 +94,9 @@ void FChaosClothAssetUSDImportNode::SetAssetValue(TObjectPtr<UObject> Asset, UE:
 				FChaosClothAssetUSDImportNode* const MutableThis = static_cast<FChaosClothAssetUSDImportNode*>(BaseNode.Get());
 				check(MutableThis == this);
 
+PRAGMA_DISABLE_DEPRECATION_WARNINGS
 				const FString& FilePath = UsdFile.FilePath;
+PRAGMA_ENABLE_DEPRECATION_WARNINGS
 				const FMD5Hash InFileHash = FilePath.IsEmpty() ?
 					FMD5Hash() :  // Reset to an empty import
 					FPaths::FileExists(FilePath) ?
@@ -138,9 +140,11 @@ void FChaosClothAssetUSDImportNode::Evaluate(UE::Dataflow::FContext& Context, co
 		FText ErrorText;
 		if (!ImportFromCache(ClothCollection, ErrorText))
 		{
+PRAGMA_DISABLE_DEPRECATION_WARNINGS
 			FClothDataflowTools::LogAndToastWarning(*this,
 				LOCTEXT("FailedToExportUsdCacheHeadline", "Failed to import USD file from cache."),
 				FText::Format(LOCTEXT("FailedToExportUsdCacheDetails", "Error while importing USD cloth from cache '{0}':\n{1}"), FText::FromString(UsdFile.FilePath), ErrorText));
+PRAGMA_ENABLE_DEPRECATION_WARNINGS
 		}
 
 		SetValue(Context, MoveTemp(*ClothCollection), &Collection);
@@ -173,6 +177,11 @@ void FChaosClothAssetUSDImportNode::Serialize(FArchive& Ar)
 		{
 			UpdateImportedAssets();
 		}
+
+		// Update the deprecated import path display
+PRAGMA_DISABLE_DEPRECATION_WARNINGS
+		ImportedFilePath = UsdFile.FilePath;
+PRAGMA_ENABLE_DEPRECATION_WARNINGS
 	}
 }
 
@@ -973,6 +982,11 @@ bool FChaosClothAssetUSDImportNode::ImportFromCache(const TSharedRef<FManagedArr
 		}
 	}
 
+	// Bind to root bone
+	constexpr bool bBindSimMesh = true;  // Some old imported collection caches did not have the sim mesh bound
+	constexpr bool bBindRenderMesh = true;
+	FClothGeometryTools::BindMeshToRootBone(OutClothCollection, bBindSimMesh, bBindRenderMesh);
+
 	return true;
 }
 
diff --git a/Engine/Plugins/ChaosClothAssetEditor/Source/ChaosClothAssetDataflowNodes/Private/ChaosClothAsset/USDImportNode_v2.cpp b/Engine/Plugins/ChaosClothAssetEditor/Source/ChaosClothAssetDataflowNodes/Private/ChaosClothAsset/USDImportNode_v2.cpp
--- a/Engine/Plugins/ChaosClothAssetEditor/Source/ChaosClothAssetDataflowNodes/Private/ChaosClothAsset/USDImportNode_v2.cpp
+++ b/Engine/Plugins/ChaosClothAssetEditor/Source/ChaosClothAssetDataflowNodes/Private/ChaosClothAsset/USDImportNode_v2.cpp
@@ -499,11 +499,9 @@ FChaosClothAssetUSDImportNode_v2::FChaosClothAssetUSDImportNode_v2(const UE::Dat
 				FText ErrorText;
 				if (!ImportUsdFile(UsdFile.FilePath, AssetPath, ErrorText))
 				{
-					UE::Chaos::ClothAsset::FClothDataflowTools::LogAndToastWarning(*this,
-						LOCTEXT("InvalidUSDClothSchemaHeadline", "Invalid USD Cloth Schema."),
-						FText::Format(LOCTEXT("InvalidUSDClothSchemaDetails", "Error while importing USD cloth from file '{0}':\n{1}\n\nWill now fallback to the legacy schema-less USD import."), FText::FromString(UsdFile.FilePath), ErrorText));
+					UE_LOG(LogChaosClothAssetDataflowNodes, Display, TEXT("No valid USD Cloth Schema was found in file '%s': %s. Will now fallback to the legacy schema-less USD import."), *UsdFile.FilePath, *ErrorText.ToString());
 
-					if(!ImportUsdFile_Schemaless(UsdFile.FilePath, AssetPath, ErrorText))
+					if (!ImportUsdFile_Schemaless(UsdFile.FilePath, AssetPath, ErrorText))
 					{
 						UE::Chaos::ClothAsset::FClothDataflowTools::LogAndToastWarning(*this,
 							LOCTEXT("FailedToImportUsdFileHeadline", "Failed to import USD file from file."),
diff --git a/Engine/Plugins/ChaosClothAssetEditor/Source/ChaosClothAssetDataflowNodes/Public/ChaosClothAsset/USDImportNode.h b/Engine/Plugins/ChaosClothAssetEditor/Source/ChaosClothAssetDataflowNodes/Public/ChaosClothAsset/USDImportNode.h
--- a/Engine/Plugins/ChaosClothAssetEditor/Source/ChaosClothAssetDataflowNodes/Public/ChaosClothAsset/USDImportNode.h
+++ b/Engine/Plugins/ChaosClothAssetEditor/Source/ChaosClothAssetDataflowNodes/Public/ChaosClothAsset/USDImportNode.h
@@ -25,8 +25,8 @@ public:
 	UPROPERTY(meta = (DataflowOutput))
 	FManagedArrayCollection Collection;
 
-	/** Name of the USD file to import. */
-	UPROPERTY(EditAnywhere, Category = "USD Import", Meta = (DisplayName = "USD File"))
+	UE_DEPRECATED(5.7, "This node is deprecated, and this value can no longer be edited.")
+	UPROPERTY()
 	FChaosClothAssetImportFilePath UsdFile;
 
 	FChaosClothAssetUSDImportNode(const UE::Dataflow::FNodeParameters& InParam, FGuid InGuid = FGuid::NewGuid());
@@ -51,6 +51,10 @@ private:
 	bool ImportFromCache(const TSharedRef<FManagedArrayCollection>& ClothCollection, class FText& OutErrorText) const;
 	void UpdateImportedAssets();
 
+	/** Name of the imported USD file. This node is deprecated, and this value can no longer be edited. */
+	UPROPERTY(Transient, VisibleAnywhere, Category = "USD Import", Meta = (DisplayName = "USD File"))
+	FString ImportedFilePath;
+
 	/** Content folder where all the USD assets are imported. */
 	UPROPERTY(VisibleAnywhere, Category = "USD Import")
 	FString PackagePath;
