Dataflow : remove hard check for view modes and map out other modes to mitigate crashing issues
- Default to use 2DSim or Render based on the IsPerspective method of the view
- More changes will come to better address the underlying issues

diff --git a/Engine/Plugins/ChaosClothAssetEditor/Source/ChaosClothAssetEditorTools/Private/ChaosClothAsset/ClothEditorToolBuilders.cpp b/Engine/Plugins/ChaosClothAssetEditor/Source/ChaosClothAssetEditorTools/Private/ChaosClothAsset/ClothEditorToolBuilders.cpp
--- a/Engine/Plugins/ChaosClothAssetEditor/Source/ChaosClothAssetEditorTools/Private/ChaosClothAsset/ClothEditorToolBuilders.cpp
+++ b/Engine/Plugins/ChaosClothAssetEditor/Source/ChaosClothAssetEditorTools/Private/ChaosClothAsset/ClothEditorToolBuilders.cpp
@@ -126,6 +126,15 @@ void UClothEditorWeightMapPaintToolBuilder::GetSupportedViewModes(const UDataflo
 	}
 }
 
+bool UClothEditorWeightMapPaintToolBuilder::CanBuildTool(const FToolBuilderState& SceneState) const
+{
+	if (UDataflowContextObject* const DataflowContextObject = SceneState.ToolManager->GetContextObjectStore()->FindContext<UDataflowContextObject>())
+	{
+		return (DataflowContextObject->GetSelectedNode() != nullptr);
+	}
+	return false;
+}
+
 UMeshSurfacePointTool* UClothEditorWeightMapPaintToolBuilder::CreateNewTool(const FToolBuilderState& SceneState) const
 {
 	UClothEditorWeightMapPaintTool* PaintTool = NewObject<UClothEditorWeightMapPaintTool>(SceneState.ToolManager);
@@ -330,8 +339,13 @@ namespace UE::Chaos::ClothAsset
 		}
 		else
 		{
-			check(ViewModeName == FName("ClothRenderView"));
-			return EClothPatternVertexType::Render;
+			// We should nromally expect ClothRenderView but in dataflow Editor this can be another mode 
+			ensure(ViewModeName == FName("ClothRenderView"));
+			if (DataflowViewMode->IsPerspective())
+			{
+				return EClothPatternVertexType::Render; // could be Sim3D but need to make a choice :)
+			}
+			return EClothPatternVertexType::Sim2D;
 		}
 	}
 
diff --git a/Engine/Plugins/ChaosClothAssetEditor/Source/ChaosClothAssetEditorTools/Public/ChaosClothAsset/ClothEditorToolBuilders.h b/Engine/Plugins/ChaosClothAssetEditor/Source/ChaosClothAssetEditorTools/Public/ChaosClothAsset/ClothEditorToolBuilders.h
--- a/Engine/Plugins/ChaosClothAssetEditor/Source/ChaosClothAssetEditorTools/Public/ChaosClothAsset/ClothEditorToolBuilders.h
+++ b/Engine/Plugins/ChaosClothAssetEditor/Source/ChaosClothAssetEditorTools/Public/ChaosClothAsset/ClothEditorToolBuilders.h
@@ -69,6 +69,8 @@ private:
 	virtual bool CanSetConstructionViewWireframeActive() const { return false; }
 
 	// UMeshSurfacePointMeshEditingToolBuilder
+	UE_API virtual bool CanBuildTool(const FToolBuilderState& SceneState) const override;
+
 	UE_API virtual UMeshSurfacePointTool* CreateNewTool(const FToolBuilderState& SceneState) const override;
 };
 
