During clothing resizing, only generate a dynamic mesh if it will be used

diff --git a/Engine/Plugins/ChaosClothAssetEditor/Source/ChaosClothAssetDataflowNodes/Private/ChaosClothAsset/ClothDataflowTools.cpp b/Engine/Plugins/ChaosClothAssetEditor/Source/ChaosClothAssetDataflowNodes/Private/ChaosClothAsset/ClothDataflowTools.cpp
--- a/Engine/Plugins/ChaosClothAssetEditor/Source/ChaosClothAssetDataflowNodes/Private/ChaosClothAsset/ClothDataflowTools.cpp
+++ b/Engine/Plugins/ChaosClothAssetEditor/Source/ChaosClothAssetDataflowNodes/Private/ChaosClothAsset/ClothDataflowTools.cpp
@@ -1055,8 +1055,7 @@ namespace UE::Chaos::ClothAsset
 		if (ClothFacade.IsValid(EClothCollectionExtendedSchemas::Resizing))
 		{
 			UE::Geometry::FDynamicMesh3 SourceMesh;
-			FMeshDescriptionToDynamicMesh Converter;
-			Converter.Convert(&SourceMeshDescription, SourceMesh, true);
+			bool bIsSourceMeshValid = false;
 
 			TConstArrayView<FString> ClothSetNames = ClothFacade.GetCustomResizingRegionSet();
 			TConstArrayView<int32> ClothSetTypes = ClothFacade.GetCustomResizingRegionType();
@@ -1066,6 +1065,13 @@ namespace UE::Chaos::ClothAsset
 				TSet<int32> SetValues;
 				if (FClothGeometryTools::ConvertSelectionToNewGroupType(ClothCollection, FName(ClothSetNames[GroupIndex]), ClothCollectionGroup::RenderVertices, SetValues))
 				{
+					if (!bIsSourceMeshValid)
+					{
+						FMeshDescriptionToDynamicMesh Converter;
+						Converter.Convert(&SourceMeshDescription, SourceMesh, true);
+						bIsSourceMeshValid = true;
+					}
+
 					UE::MeshResizing::FCustomRegionResizing::GenerateCustomRegion(ClothFacade.GetRenderPosition(), SourceMesh, SetValues, OutResizingGroupData[GroupIndex]);
 				}
 			}
