Add the ability to specify the ReferenceBone for a cloth asset.

diff --git a/Engine/Plugins/ChaosClothAsset/Source/ChaosClothAsset/Private/ChaosClothAsset/ClothCollection.cpp b/Engine/Plugins/ChaosClothAsset/Source/ChaosClothAsset/Private/ChaosClothAsset/ClothCollection.cpp
--- a/Engine/Plugins/ChaosClothAsset/Source/ChaosClothAsset/Private/ChaosClothAsset/ClothCollection.cpp
+++ b/Engine/Plugins/ChaosClothAsset/Source/ChaosClothAsset/Private/ChaosClothAsset/ClothCollection.cpp
@@ -12,6 +12,7 @@ namespace UE::Chaos::ClothAsset::Private
 	{
 		ClothCollectionAttribute::PhysicsAssetSoftObjectPathName,
 		ClothCollectionAttribute::SkeletalMeshSoftObjectPathName,
+		ClothCollectionAttribute::ReferenceBoneName,
 	};
 
 	// Solvers Group
@@ -190,6 +191,7 @@ namespace UE::Chaos::ClothAsset
 		// LODs Group
 		PhysicsAssetSoftObjectPathName = ManagedArrayCollection->FindAttribute<FSoftObjectPath>(ClothCollectionAttribute::PhysicsAssetSoftObjectPathName, ClothCollectionGroup::Lods);
 		SkeletalMeshSoftObjectPathName = ManagedArrayCollection->FindAttribute<FSoftObjectPath>(ClothCollectionAttribute::SkeletalMeshSoftObjectPathName, ClothCollectionGroup::Lods);
+		ReferenceBoneName = ManagedArrayCollection->FindAttribute<FName>(ClothCollectionAttribute::ReferenceBoneName, ClothCollectionGroup::Lods);
 
 		// Solvers Group
 		SolverGravity = ManagedArrayCollection->FindAttribute<FVector3f>(ClothCollectionAttribute::SolverGravity, ClothCollectionGroup::Solvers);
@@ -315,6 +317,7 @@ namespace UE::Chaos::ClothAsset
 			// LODs Group
 			PhysicsAssetSoftObjectPathName &&
 			SkeletalMeshSoftObjectPathName &&
+			ReferenceBoneName &&
 
 			// Solvers Group
 			(!EnumHasAnyFlags(OptionalSchemas, EClothCollectionExtendedSchemas::Solvers) ||
@@ -465,6 +468,7 @@ namespace UE::Chaos::ClothAsset
 		// LODs Group
 		PhysicsAssetSoftObjectPathName = &GetManagedArrayCollection()->AddAttribute<FSoftObjectPath>(ClothCollectionAttribute::PhysicsAssetSoftObjectPathName, ClothCollectionGroup::Lods);
 		SkeletalMeshSoftObjectPathName = &GetManagedArrayCollection()->AddAttribute<FSoftObjectPath>(ClothCollectionAttribute::SkeletalMeshSoftObjectPathName, ClothCollectionGroup::Lods);
+		ReferenceBoneName = &GetManagedArrayCollection()->AddAttribute<FName>(ClothCollectionAttribute::ReferenceBoneName, ClothCollectionGroup::Lods);
 
 		// Solvers Group
 		if (EnumHasAnyFlags(OptionalSchemas, EClothCollectionExtendedSchemas::Solvers))
@@ -631,6 +635,12 @@ namespace UE::Chaos::ClothAsset
 				ConvertStringToSoftObjectPath(*InRenderMaterialPathName, *GetRenderMaterialSoftObjectPathName());
 				GetManagedArrayCollection()->RemoveAttribute(ClothCollectionAttribute::RenderMaterialPathName, ClothCollectionGroup::RenderPatterns);
 			}
+
+			// Add ReferenceBoneName attribute if other LODs group attributes are defined.
+			if (PhysicsAssetSoftObjectPathName && SkeletalMeshSoftObjectPathName && !ReferenceBoneName)
+			{
+				ReferenceBoneName = &GetManagedArrayCollection()->AddAttribute<FName>(ClothCollectionAttribute::ReferenceBoneName, ClothCollectionGroup::Lods);
+			}
 		}
 	}
 
diff --git a/Engine/Plugins/ChaosClothAsset/Source/ChaosClothAsset/Private/ChaosClothAsset/ClothCollection.h b/Engine/Plugins/ChaosClothAsset/Source/ChaosClothAsset/Private/ChaosClothAsset/ClothCollection.h
--- a/Engine/Plugins/ChaosClothAsset/Source/ChaosClothAsset/Private/ChaosClothAsset/ClothCollection.h
+++ b/Engine/Plugins/ChaosClothAsset/Source/ChaosClothAsset/Private/ChaosClothAsset/ClothCollection.h
@@ -136,6 +136,7 @@ namespace UE::Chaos::ClothAsset
 		const TManagedArray<FString>* GetPhysicsAssetPathName() const { return nullptr; }
 		UE_DEPRECATED(5.7, "Use GetSkeletalMeshSoftObjectPathName instead")
 		const TManagedArray<FString>* GetSkeletalMeshPathName() const { return nullptr; }
+		const TManagedArray<FName>* GetReferenceBoneName() const { return ReferenceBoneName; }
 
 		// ~ Solvers Group
 		const TManagedArray<FVector3f>* GetSolverGravity() const { return SolverGravity; }
@@ -250,6 +251,7 @@ namespace UE::Chaos::ClothAsset
 		const TManagedArray<FString>* PhysicsAssetPathName;
 		UE_DEPRECATED(5.7, "Use SkeletalMeshSoftObjectPathName instead")
 		const TManagedArray<FString>* SkeletalMeshPathName;
+		const TManagedArray<FName>* ReferenceBoneName;
 
 		//~ Solvers group
 		const TManagedArray<FVector3f>* SolverGravity;
@@ -402,6 +404,7 @@ namespace UE::Chaos::ClothAsset
 		TManagedArray<FString>* GetPhysicsAssetPathName() { return nullptr; }
 		UE_DEPRECATED(5.7, "Use GetSkeletalMeshSoftObjectPathName instead")
 		TManagedArray<FString>* GetSkeletalMeshPathName() { return nullptr; }
+		TManagedArray<FName>* GetReferenceBoneName() { return  const_cast<TManagedArray<FName>*>(ReferenceBoneName); }
 
 		//~ Solver Group
 		TManagedArray<FVector3f>* GetSolverGravity() { return const_cast<TManagedArray<FVector3f>*>(SolverGravity); }
diff --git a/Engine/Plugins/ChaosClothAsset/Source/ChaosClothAsset/Private/ChaosClothAsset/CollectionClothFacade.cpp b/Engine/Plugins/ChaosClothAsset/Source/ChaosClothAsset/Private/ChaosClothAsset/CollectionClothFacade.cpp
--- a/Engine/Plugins/ChaosClothAsset/Source/ChaosClothAsset/Private/ChaosClothAsset/CollectionClothFacade.cpp
+++ b/Engine/Plugins/ChaosClothAsset/Source/ChaosClothAsset/Private/ChaosClothAsset/CollectionClothFacade.cpp
@@ -62,6 +62,7 @@ namespace UE::Chaos::ClothAsset
 		//~ LOD (single per collection) Group
 		ResultHash = HashCombineFast(ResultHash, GetTypeHash(GetPhysicsAssetSoftObjectPathName()));
 		ResultHash = HashCombineFast(ResultHash, GetTypeHash(GetSkeletalMeshSoftObjectPathName()));
+		ResultHash = HashCombineFast(ResultHash, GetTypeHash(GetReferenceBoneName()));
 
 		//~ Solvers Group
 		ResultHash = HashCombineFast(ResultHash, GetTypeHash(GetSolverGravity()));
@@ -206,6 +207,11 @@ namespace UE::Chaos::ClothAsset
 		return ClothCollection->GetSkeletalMeshSoftObjectPathName() && ClothCollection->GetNumElements(ClothCollectionGroup::Lods) > 0 ? (*ClothCollection->GetSkeletalMeshSoftObjectPathName())[0] : EmptyString;
 	}
 
+	FName FCollectionClothConstFacade::GetReferenceBoneName() const
+	{
+		return ClothCollection->GetReferenceBoneName() && ClothCollection->GetNumElements(ClothCollectionGroup::Lods) > 0 ? (*ClothCollection->GetReferenceBoneName())[0] : NAME_None;
+	}
+
 	bool FCollectionClothConstFacade::HasSolverElement() const
 	{
 		return ClothCollection->GetNumElements(ClothCollectionGroup::Solvers) == 1;
@@ -710,6 +716,10 @@ namespace UE::Chaos::ClothAsset
 		{
 			SetSkeletalMeshSoftObjectPathName(Other.GetSkeletalMeshSoftObjectPathName());
 		}
+		if (GetReferenceBoneName() == NAME_None)
+		{
+			SetReferenceBoneName(Other.GetReferenceBoneName());
+		}
 
 		const bool bOtherHasResizingData = Other.IsValid(EClothCollectionExtendedSchemas::Resizing);
 		if (bOtherHasResizingData)
@@ -886,6 +896,14 @@ namespace UE::Chaos::ClothAsset
 			(*GetClothCollection()->GetSkeletalMeshSoftObjectPathName())[0] = PathName;
 		}
 	}
+
+	void FCollectionClothFacade::SetReferenceBoneName(const FName& BoneName)
+	{
+		if (ClothCollection->GetNumElements(ClothCollectionGroup::Lods))
+		{
+			(*GetClothCollection()->GetReferenceBoneName())[0] = BoneName;
+		}
+	}
 	
 	void FCollectionClothFacade::SetSolverGravity(const FVector3f& SolverGravity)
 	{
diff --git a/Engine/Plugins/ChaosClothAsset/Source/ChaosClothAsset/Public/ChaosClothAsset/ClothCollectionAttribute.h b/Engine/Plugins/ChaosClothAsset/Source/ChaosClothAsset/Public/ChaosClothAsset/ClothCollectionAttribute.h
--- a/Engine/Plugins/ChaosClothAsset/Source/ChaosClothAsset/Public/ChaosClothAsset/ClothCollectionAttribute.h
+++ b/Engine/Plugins/ChaosClothAsset/Source/ChaosClothAsset/Public/ChaosClothAsset/ClothCollectionAttribute.h
@@ -14,6 +14,7 @@ namespace UE::Chaos::ClothAsset
 		// Lods Group
 		inline const FName PhysicsAssetSoftObjectPathName(TEXT("PhysicsAssetSoftObjectPathName"));
 		inline const FName SkeletalMeshSoftObjectPathName(TEXT("SkeletalMeshSoftObjectPathName"));
+		inline const FName ReferenceBoneName(TEXT("ReferenceBoneName"));
 
 		// Solvers Group
 		inline const FName SolverGravity(TEXT("SolverGravity"));
diff --git a/Engine/Plugins/ChaosClothAsset/Source/ChaosClothAsset/Public/ChaosClothAsset/CollectionClothFacade.h b/Engine/Plugins/ChaosClothAsset/Source/ChaosClothAsset/Public/ChaosClothAsset/CollectionClothFacade.h
--- a/Engine/Plugins/ChaosClothAsset/Source/ChaosClothAsset/Public/ChaosClothAsset/CollectionClothFacade.h
+++ b/Engine/Plugins/ChaosClothAsset/Source/ChaosClothAsset/Public/ChaosClothAsset/CollectionClothFacade.h
@@ -84,6 +84,8 @@ namespace UE::Chaos::ClothAsset
 			static const FString EmptyString;
 			return EmptyString;
 		}
+		/** Return the reference bone name used for this collection. */
+		UE_API FName GetReferenceBoneName() const;
 
 		//~ Solver (single per collection) Group
 		/** Return true if the solver group has one element*/
@@ -290,6 +292,8 @@ namespace UE::Chaos::ClothAsset
 		void SetPhysicsAssetPathName(const FString& PathName){}
 		UE_DEPRECATED(5.7, "Use SetSkeletalMeshSoftObjectPathName instead")
 		void SetSkeletalMeshPathName(const FString& PathName){}
+		/** Set the reference bone name for this asset.*/
+		UE_API void SetReferenceBoneName(const FName& BoneName);
 
 		//~ Solver (max 1 per collection) Group
 		/** Set the solver gravity */
diff --git a/Engine/Plugins/ChaosClothAsset/Source/ChaosClothAssetEngine/Private/ChaosClothAsset/ClothEngineTools.cpp b/Engine/Plugins/ChaosClothAsset/Source/ChaosClothAssetEngine/Private/ChaosClothAsset/ClothEngineTools.cpp
--- a/Engine/Plugins/ChaosClothAsset/Source/ChaosClothAssetEngine/Private/ChaosClothAsset/ClothEngineTools.cpp
+++ b/Engine/Plugins/ChaosClothAsset/Source/ChaosClothAssetEngine/Private/ChaosClothAsset/ClothEngineTools.cpp
@@ -7,6 +7,7 @@
 #include "ChaosClothAsset/ClothCollectionGroup.h"
 #include "ChaosClothAsset/ClothGeometryTools.h"
 #include "ClothTetherData.h"
+#include "ReferenceSkeleton.h"
 
 namespace UE::Chaos::ClothAsset
 {
@@ -175,4 +176,100 @@ FPointWeightMap FClothEngineTools::GetMaxDistanceWeightMap(const FCollectionClot
 		FPointWeightMap(NumLodSimVertices, MaxDistanceOffset);
 }
 
+int32 FClothEngineTools::CalculateReferenceBoneIndex(const TArray<int32>& UsedBones, const FReferenceSkeleton& ReferenceSkeleton)
+{
+	// Starts at root
+	int32 ReferenceBoneIndex = 0;
+
+	// List of valid paths to the root bone from each weighted bone
+	TArray<TArray<int32>> PathsToRoot;
+
+	const int32 NumUsedBones = UsedBones.Num();
+	PathsToRoot.Reserve(NumUsedBones);
+
+	// Compute paths to the root bone
+	for (int32 UsedBoneIndex = 0; UsedBoneIndex < NumUsedBones; ++UsedBoneIndex)
+	{
+		PathsToRoot.AddDefaulted();
+		TArray<int32>& Path = PathsToRoot.Last();
+
+		int32 CurrentBone = UsedBones[UsedBoneIndex];
+		Path.Add(CurrentBone);
+
+		while (CurrentBone != 0 && CurrentBone != INDEX_NONE)
+		{
+			CurrentBone = ReferenceSkeleton.GetParentIndex(CurrentBone);
+			Path.Add(CurrentBone);
+		}
+	}
+
+	// Paths are from leaf->root, we want the other way
+	for (TArray<int32>& Path : PathsToRoot)
+	{
+		Algo::Reverse(Path);
+	}
+
+	// Verify the last common bone in all paths as the root of the sim space
+	const int32 NumPaths = PathsToRoot.Num();
+	if (NumPaths > 0)
+	{
+		TArray<int32>& FirstPath = PathsToRoot[0];
+
+		const int32 FirstPathSize = FirstPath.Num();
+		for (int32 PathEntryIndex = 0; PathEntryIndex < FirstPathSize; ++PathEntryIndex)
+		{
+			const int32 CurrentQueryIndex = FirstPath[PathEntryIndex];
+			bool bValidRoot = true;
+
+			for (int32 PathIndex = 1; PathIndex < NumPaths; ++PathIndex)
+			{
+				if (!PathsToRoot[PathIndex].Contains(CurrentQueryIndex))
+				{
+					bValidRoot = false;
+					break;
+				}
+			}
+
+			if (bValidRoot)
+			{
+				ReferenceBoneIndex = CurrentQueryIndex;
+			}
+			else
+			{
+				// Once we fail to find a valid root we're done.
+				break;
+			}
+		}
+	}
+	else
+	{
+		// Just use the root
+		ReferenceBoneIndex = 0;
+	}
+	return ReferenceBoneIndex;
+}
+
+
+int32 FClothEngineTools::CalculateReferenceBoneIndex(const TSharedRef<const FManagedArrayCollection>& ClothCollection, const FReferenceSkeleton& ReferenceSkeleton)
+{
+	FCollectionClothConstFacade ClothFacade(ClothCollection);
+	TConstArrayView<TArray<int32>> SimBoneIndices = ClothFacade.GetSimBoneIndices();
+	TConstArrayView<TArray<float>> SimBoneWeights = ClothFacade.GetSimBoneWeights();
+
+	TSet<int32> UsedBones;
+	UsedBones.Reserve(ReferenceSkeleton.GetRawBoneNum());
+	for (int32 VertexIndex = 0; VertexIndex < SimBoneIndices.Num(); ++VertexIndex)
+	{
+		check(SimBoneIndices[VertexIndex].Num() == SimBoneWeights[VertexIndex].Num());
+		for (int32 BoneIndex = 0; BoneIndex < SimBoneIndices[VertexIndex].Num(); ++BoneIndex)
+		{
+			if (SimBoneWeights[VertexIndex][BoneIndex] > UE_SMALL_NUMBER)
+			{
+				UsedBones.Add(SimBoneIndices[VertexIndex][BoneIndex]);
+			}
+		}
+	}
+
+	return CalculateReferenceBoneIndex(UsedBones.Array(), ReferenceSkeleton);
+}
 }  // End namespace UE::Chaos::ClothAsset
diff --git a/Engine/Plugins/ChaosClothAsset/Source/ChaosClothAssetEngine/Private/ChaosClothAsset/ClothSimulationModel.cpp b/Engine/Plugins/ChaosClothAsset/Source/ChaosClothAssetEngine/Private/ChaosClothAsset/ClothSimulationModel.cpp
--- a/Engine/Plugins/ChaosClothAsset/Source/ChaosClothAssetEngine/Private/ChaosClothAsset/ClothSimulationModel.cpp
+++ b/Engine/Plugins/ChaosClothAsset/Source/ChaosClothAssetEngine/Private/ChaosClothAsset/ClothSimulationModel.cpp
@@ -6,6 +6,7 @@
 #include "ChaosClothAsset/ClothAssetPrivate.h"
 #include "ChaosClothAsset/ClothCollectionGroup.h"
 #include "ChaosClothAsset/ClothGeometryTools.h"
+#include "ChaosClothAsset/ClothEngineTools.h"
 #include "ChaosClothAsset/ClothLodTransitionDataCache.h"
 #include "Chaos/CollectionPropertyFacade.h"
 #include "Utils/ClothingMeshUtils.h"
@@ -37,12 +38,6 @@ namespace UE::Chaos::ClothAsset::Private
 	// Find the root bone for this cloth asset (common bone for all used bones)
 	static int32 CalculateReferenceBoneIndex(const TArray<FChaosClothSimulationLodModel> ClothSimulationLodModels, const FReferenceSkeleton& ReferenceSkeleton, const TArray<int32>& UsedBoneIndices)
 	{
-		// Starts at root
-		int32 ReferenceBoneIndex = 0;
-
-		// List of valid paths to the root bone from each weighted bone
-		TArray<TArray<int32>> PathsToRoot;
-
 		// First build a list per used bone for it's path to root
 		TArray<int32> WeightedBones;  // List of actually weighted (not just used) bones
 
@@ -67,69 +62,7 @@ namespace UE::Chaos::ClothAsset::Private
 			}
 		}
 
-		const int32 NumWeightedBones = WeightedBones.Num();
-		PathsToRoot.Reserve(NumWeightedBones);
-
-		// Compute paths to the root bone
-		for (int32 WeightedBoneIndex = 0; WeightedBoneIndex < NumWeightedBones; ++WeightedBoneIndex)
-		{
-			PathsToRoot.AddDefaulted();
-			TArray<int32>& Path = PathsToRoot.Last();
-
-			int32 CurrentBone = WeightedBones[WeightedBoneIndex];
-			Path.Add(CurrentBone);
-
-			while (CurrentBone != 0 && CurrentBone != INDEX_NONE)
-			{
-				CurrentBone = ReferenceSkeleton.GetParentIndex(CurrentBone);
-				Path.Add(CurrentBone);
-			}
-		}
-
-		// Paths are from leaf->root, we want the other way
-		for (TArray<int32>& Path : PathsToRoot)
-		{
-			Algo::Reverse(Path);
-		}
-
-		// Verify the last common bone in all paths as the root of the sim space
-		const int32 NumPaths = PathsToRoot.Num();
-		if (NumPaths > 0)
-		{
-			TArray<int32>& FirstPath = PathsToRoot[0];
-
-			const int32 FirstPathSize = FirstPath.Num();
-			for (int32 PathEntryIndex = 0; PathEntryIndex < FirstPathSize; ++PathEntryIndex)
-			{
-				const int32 CurrentQueryIndex = FirstPath[PathEntryIndex];
-				bool bValidRoot = true;
-
-				for (int32 PathIndex = 1; PathIndex < NumPaths; ++PathIndex)
-				{
-					if (!PathsToRoot[PathIndex].Contains(CurrentQueryIndex))
-					{
-						bValidRoot = false;
-						break;
-					}
-				}
-
-				if (bValidRoot)
-				{
-					ReferenceBoneIndex = CurrentQueryIndex;
-				}
-				else
-				{
-					// Once we fail to find a valid root we're done.
-					break;
-				}
-			}
-		}
-		else
-		{
-			// Just use the root
-			ReferenceBoneIndex = 0;
-		}
-		return ReferenceBoneIndex;
+		return FClothEngineTools::CalculateReferenceBoneIndex(WeightedBones, ReferenceSkeleton);
 	}
 
 	FMD5Hash CalculateTransitionDataHash(const FChaosClothSimulationLodModel& CurrModel)
@@ -255,10 +188,19 @@ FChaosClothSimulationModel::FChaosClothSimulationModel(const TArray<TSharedRef<c
 	const int32 NumLods = ClothCollections.Num();
 	ClothSimulationLodModels.SetNum(NumLods);
 
+	ReferenceBoneIndex = INDEX_NONE;
 	for (int32 LodIndex = 0; LodIndex < NumLods; ++LodIndex)
 	{
 		const FCollectionClothConstFacade Cloth(ClothCollections[LodIndex]);
 
+		if (LodIndex == 0)
+		{
+			// Get reference bone
+			const FName ReferenceBoneName = Cloth.GetReferenceBoneName();
+
+			ReferenceBoneIndex = ReferenceSkeleton.FindRawBoneIndex(ReferenceBoneName);
+		}
+
 		// Retrieve weigh map names
 		const TArray<FName> WeightMapNames = Cloth.GetWeightMapNames();
 
@@ -337,6 +279,10 @@ FChaosClothSimulationModel::FChaosClothSimulationModel(const TArray<TSharedRef<c
 		MergedTetherData.SetNum(NumSimVertices3D);
 		TSet<FBoneIndexType> RequiredExtraSimBones;
 		RequiredExtraSimBones.Reserve(ReferenceSkeleton.GetRawBoneNum());
+		if (ReferenceBoneIndex != INDEX_NONE)
+		{
+			RequiredExtraSimBones.Add(ReferenceBoneIndex);
+		}
 
 		for (int32 VertexIndex = 0; VertexIndex < NumSimVertices3D; ++VertexIndex)
 		{
@@ -382,8 +328,11 @@ FChaosClothSimulationModel::FChaosClothSimulationModel(const TArray<TSharedRef<c
 		UsedBoneIndices.Add(Index);
 	}
 
-	// Initialize Reference bone index
-	ReferenceBoneIndex = Private::CalculateReferenceBoneIndex(ClothSimulationLodModels, ReferenceSkeleton, UsedBoneIndices);
+	// Initialize Reference bone index if not set in the cloth collection
+	if (ReferenceBoneIndex == INDEX_NONE)
+	{
+		ReferenceBoneIndex = Private::CalculateReferenceBoneIndex(ClothSimulationLodModels, ReferenceSkeleton, UsedBoneIndices);
+	}
 
 	CalculateLODTransitionUpDownData(InOutTransitionCache);
 }
diff --git a/Engine/Plugins/ChaosClothAsset/Source/ChaosClothAssetEngine/Public/ChaosClothAsset/ClothEngineTools.h b/Engine/Plugins/ChaosClothAsset/Source/ChaosClothAssetEngine/Public/ChaosClothAsset/ClothEngineTools.h
--- a/Engine/Plugins/ChaosClothAsset/Source/ChaosClothAssetEngine/Public/ChaosClothAsset/ClothEngineTools.h
+++ b/Engine/Plugins/ChaosClothAsset/Source/ChaosClothAssetEngine/Public/ChaosClothAsset/ClothEngineTools.h
@@ -11,6 +11,7 @@
 struct FManagedArrayCollection;
 class FName;
 class FSkeletalMeshLODModel;
+struct FReferenceSkeleton;
 namespace Chaos
 {
 	namespace Softs
@@ -38,6 +39,11 @@ namespace UE::Chaos::ClothAsset
 		static FPointWeightMap GetMaxDistanceWeightMap(const FCollectionClothConstFacade& ClothFacade, const ::Chaos::Softs::FCollectionPropertyConstFacade& PropertyFacade, const int32 NumLodSimVertices);
 		/** Retrieve the MaxDistance weight map from the cloth collection. Return an empty map if the collection doesn't match the requested NumLodSimVertices. */
 		static FPointWeightMap GetMaxDistanceWeightMap(const TSharedRef<const FManagedArrayCollection>& ClothCollection, const int32 NumLodSimVertices);
+
+		/** Calculate ReferenceBone from bone closest to the root of all used bones */
+		static UE_API int32 CalculateReferenceBoneIndex(const TArray<int32>& UsedBones, const FReferenceSkeleton& ReferenceSkeleton);
+		/** Calculate ReferenceBone from bone closest to the root of all used weighted sim bones in a collection */
+		static UE_API int32 CalculateReferenceBoneIndex(const TSharedRef<const FManagedArrayCollection>& ClothCollection, const FReferenceSkeleton& ReferenceSkeleton);
 	};
 }  // End namespace UE::Chaos::ClothAsset
 
diff --git a/Engine/Plugins/ChaosClothAssetEditor/Source/ChaosClothAssetDataflowNodes/Private/ChaosClothAsset/ChaosClothAssetDataflowNodesModule.cpp b/Engine/Plugins/ChaosClothAssetEditor/Source/ChaosClothAssetDataflowNodes/Private/ChaosClothAsset/ChaosClothAssetDataflowNodesModule.cpp
--- a/Engine/Plugins/ChaosClothAssetEditor/Source/ChaosClothAssetDataflowNodes/Private/ChaosClothAsset/ChaosClothAssetDataflowNodesModule.cpp
+++ b/Engine/Plugins/ChaosClothAssetEditor/Source/ChaosClothAssetDataflowNodes/Private/ChaosClothAsset/ChaosClothAssetDataflowNodesModule.cpp
@@ -25,9 +25,10 @@
 #include "ChaosClothAsset/MergeClothCollectionsNode.h"
 #include "ChaosClothAsset/ProceduralSelectionNode.h"
 #include "ChaosClothAsset/ProxyDeformerNode.h"
+#include "ChaosClothAsset/RecalculateNormalsNode.h"
+#include "ChaosClothAsset/ReferenceBoneNode.h"
 #include "ChaosClothAsset/RemeshNode.h"
 #include "ChaosClothAsset/ReverseNormalsNode.h"
-#include "ChaosClothAsset/RecalculateNormalsNode.h"
 #include "ChaosClothAsset/SelectionNode.h"
 #include "ChaosClothAsset/SelectionToIntMapNode.h"
 #include "ChaosClothAsset/SelectionToWeightMapNode.h"
@@ -90,8 +91,10 @@
 #include "Dataflow/DataflowNodeColorsRegistry.h"
 #include "Dataflow/DataflowNodeFactory.h"
 #include "Dataflow/DataflowRenderingFactory.h"
-#include "Modules/ModuleManager.h"
+#include "Engine/SkeletalMesh.h"
 #include "GeometryCollection/Facades/CollectionRenderingFacade.h"
+#include "Modules/ModuleManager.h"
+#include "ReferenceSkeleton.h"
 
 namespace UE::Chaos::ClothAsset
 {
@@ -123,6 +126,7 @@ namespace UE::Chaos::ClothAsset
 			DATAFLOW_NODE_REGISTER_CREATION_FACTORY(FChaosClothAssetMergeClothCollectionsNode_v2);
 			DATAFLOW_NODE_REGISTER_CREATION_FACTORY(FChaosClothAssetProceduralSelectionNode);
 			DATAFLOW_NODE_REGISTER_CREATION_FACTORY(FChaosClothAssetProxyDeformerNode_v3);
+			DATAFLOW_NODE_REGISTER_CREATION_FACTORY(FChaosClothAssetReferenceBoneNode);
 			DATAFLOW_NODE_REGISTER_CREATION_FACTORY(FChaosClothAssetRemeshNode_v2);
 			DATAFLOW_NODE_REGISTER_CREATION_FACTORY(FChaosClothAssetReverseNormalsNode);
 			DATAFLOW_NODE_REGISTER_CREATION_FACTORY(FChaosClothAssetRecalculateNormalsNode);
@@ -176,8 +180,8 @@ namespace UE::Chaos::ClothAsset
 			DATAFLOW_NODE_REGISTER_CREATION_FACTORY(FChaosClothAssetUSDImportNode_v2);
 			DATAFLOW_NODE_REGISTER_CREATION_FACTORY(FChaosClothAssetWeightMapToSelectionNode);
 			// Deprecated nodes
-PRAGMA_DISABLE_DEPRECATION_WARNINGS
-			DATAFLOW_NODE_REGISTER_CREATION_FACTORY(FChaosClothAssetAttributeNode);
+			PRAGMA_DISABLE_DEPRECATION_WARNINGS
+				DATAFLOW_NODE_REGISTER_CREATION_FACTORY(FChaosClothAssetAttributeNode);
 			DATAFLOW_NODE_REGISTER_CREATION_FACTORY(FChaosClothAssetAddWeightMapNode);
 			DATAFLOW_NODE_REGISTER_CREATION_FACTORY(FChaosClothAssetMergeClothCollectionsNode);
 			DATAFLOW_NODE_REGISTER_CREATION_FACTORY(FChaosClothAssetProxyDeformerNode);
@@ -190,9 +194,22 @@ PRAGMA_DISABLE_DEPRECATION_WARNINGS
 			DATAFLOW_NODE_REGISTER_CREATION_FACTORY(FChaosClothAssetStaticMeshImportNode);
 			DATAFLOW_NODE_REGISTER_CREATION_FACTORY(FChaosClothAssetTerminalNode);
 			DATAFLOW_NODE_REGISTER_CREATION_FACTORY(FChaosClothAssetUSDImportNode);
-PRAGMA_ENABLE_DEPRECATION_WARNINGS
+			PRAGMA_ENABLE_DEPRECATION_WARNINGS
 		}
 
+		static TArray<FName> GetReferenceBoneNamesFromCollection(const FManagedArrayCollection& Collection)
+		{
+			TSharedRef<FManagedArrayCollection> CollectionPtr = MakeShared<FManagedArrayCollection>(Collection);
+			FCollectionClothConstFacade ClothFacade(CollectionPtr);
+			const FSoftObjectPath& SkeletalMeshPathName = ClothFacade.GetSkeletalMeshSoftObjectPathName();
+			TArray<FName> BoneNames;
+			if (const USkeletalMesh* const SkeletalMesh = Cast<USkeletalMesh>(SkeletalMeshPathName.TryLoad()))
+			{
+				const FReferenceSkeleton& RefSkeleton = SkeletalMesh->GetRefSkeleton();
+				return RefSkeleton.GetRawRefBoneNames();
+			}
+			return TArray<FName>();
+		}
 
 		class FClothSurfaceRenderCallbacks : public UE::Dataflow::FRenderingFactory::ICallbackInterface
 		{
@@ -405,6 +422,10 @@ PRAGMA_ENABLE_DEPRECATION_WARNINGS
 				PropertyModule->RegisterCustomPropertyTypeLayout(FChaosClothAssetNodeSelectionGroup::StaticStruct()->GetFName(), FOnGetPropertyTypeCustomizationInstance::CreateStatic(&FSelectionGroupCustomization::MakeInstance));
 				PropertyModule->RegisterCustomPropertyTypeLayout(FChaosClothAssetImportFilePath::StaticStruct()->GetFName(), FOnGetPropertyTypeCustomizationInstance::CreateStatic(&FImportFilePathCustomization::MakeInstance));
 				PropertyModule->RegisterCustomPropertyTypeLayout(FChaosClothAssetNodeAttributeGroup::StaticStruct()->GetFName(), FOnGetPropertyTypeCustomizationInstance::CreateStatic(&FSelectionGroupCustomization::MakeInstance));
+				PropertyModule->RegisterCustomPropertyTypeLayout(FChaosClothAssetReferenceBoneSelection::StaticStruct()->GetFName(), FOnGetPropertyTypeCustomizationInstance::CreateLambda([]()
+					{
+						return MakeShareable(new UE::Dataflow::FPropertyListCustomization(UE::Dataflow::FPropertyListCustomization::FOnGetListNames::CreateStatic(&Private::GetReferenceBoneNamesFromCollection)));
+					}));
 				PropertyModule->RegisterCustomPropertyTypeLayout(FChaosClothAssetImportedVectorValue::StaticStruct()->GetFName(), FOnGetPropertyTypeCustomizationInstance::CreateStatic(&FImportedValueCustomization::MakeInstance));
 				PropertyModule->RegisterCustomPropertyTypeLayout(FChaosClothAssetImportedFloatValue::StaticStruct()->GetFName(), FOnGetPropertyTypeCustomizationInstance::CreateStatic(&FImportedValueCustomization::MakeInstance));
 				PropertyModule->RegisterCustomPropertyTypeLayout(FChaosClothAssetImportedIntValue::StaticStruct()->GetFName(), FOnGetPropertyTypeCustomizationInstance::CreateStatic(&FImportedValueCustomization::MakeInstance));
@@ -435,6 +456,7 @@ PRAGMA_ENABLE_DEPRECATION_WARNINGS
 					PropertyModule->UnregisterCustomPropertyTypeLayout(FChaosClothAssetNodeSelectionGroup::StaticStruct()->GetFName());
 					PropertyModule->UnregisterCustomPropertyTypeLayout(FChaosClothAssetImportFilePath::StaticStruct()->GetFName());
 					PropertyModule->UnregisterCustomPropertyTypeLayout(FChaosClothAssetNodeAttributeGroup::StaticStruct()->GetFName());
+					PropertyModule->UnregisterCustomPropertyTypeLayout(FChaosClothAssetReferenceBoneSelection::StaticStruct()->GetFName());
 					PropertyModule->UnregisterCustomPropertyTypeLayout(FChaosClothAssetImportedVectorValue::StaticStruct()->GetFName());
 					PropertyModule->UnregisterCustomPropertyTypeLayout(FChaosClothAssetImportedFloatValue::StaticStruct()->GetFName());
 					PropertyModule->UnregisterCustomPropertyTypeLayout(FChaosClothAssetImportedIntValue::StaticStruct()->GetFName());
diff --git a/Engine/Plugins/ChaosClothAssetEditor/Source/ChaosClothAssetDataflowNodes/Private/ChaosClothAsset/ReferenceBoneNode.cpp b/Engine/Plugins/ChaosClothAssetEditor/Source/ChaosClothAssetDataflowNodes/Private/ChaosClothAsset/ReferenceBoneNode.cpp
new file mode 100644
--- /dev/null
+++ b/Engine/Plugins/ChaosClothAssetEditor/Source/ChaosClothAssetDataflowNodes/Private/ChaosClothAsset/ReferenceBoneNode.cpp
@@ -0,0 +1,64 @@
+// Copyright Epic Games, Inc. All Rights Reserved.
+
+#include "ChaosClothAsset/ReferenceBoneNode.h"
+#include "ChaosClothAsset/ClothEngineTools.h"
+#include "ChaosClothAsset/CollectionClothFacade.h"
+#include "Engine/SkeletalMesh.h"
+#include "ReferenceSkeleton.h"
+
+#include UE_INLINE_GENERATED_CPP_BY_NAME(ReferenceBoneNode)
+
+#define LOCTEXT_NAMESPACE "ChaosClothAssetReferenceBoneNode"
+
+
+FChaosClothAssetReferenceBoneNode::FChaosClothAssetReferenceBoneNode(const UE::Dataflow::FNodeParameters& InParam, FGuid InGuid)
+	: FDataflowNode(InParam, InGuid)
+	, CalculateDefaultReferenceBone(FDataflowFunctionProperty::FDelegate::CreateRaw(this, &FChaosClothAssetReferenceBoneNode::OnCalculateDefaultReferenceBone))
+{
+	RegisterInputConnection(&Collection);
+	RegisterOutputConnection(&Collection, &Collection);
+}
+
+void FChaosClothAssetReferenceBoneNode::Evaluate(UE::Dataflow::FContext& Context, const FDataflowOutput* Out) const
+{
+	using namespace UE::Chaos::ClothAsset;
+	if (Out->IsA<FManagedArrayCollection>(&Collection))
+	{
+		using namespace UE::Chaos::ClothAsset;
+		// Evaluate in collection
+		FManagedArrayCollection InCollection = GetValue<FManagedArrayCollection>(Context, &Collection);
+		const TSharedRef<FManagedArrayCollection> ClothCollection = MakeShared<FManagedArrayCollection>(MoveTemp(InCollection));
+		// Always check for a valid cloth collection/facade/sim mesh to avoid processing non cloth collections or pure render mesh cloth assets
+		FCollectionClothFacade ClothFacade(ClothCollection);
+		if (ClothFacade.IsValid() && ClothFacade.HasValidData())
+		{
+			ClothFacade.SetReferenceBoneName(ReferenceBone.Name);
+		}
+		SetValue(Context, MoveTemp(*ClothCollection), &Collection);
+	}
+}
+
+void FChaosClothAssetReferenceBoneNode::OnCalculateDefaultReferenceBone(UE::Dataflow::FContext& Context)
+{
+	using namespace UE::Chaos::ClothAsset;
+
+	ReferenceBone.Name = NAME_None;
+
+	FManagedArrayCollection InCollection = GetValue<FManagedArrayCollection>(Context, &Collection);
+	const TSharedRef<FManagedArrayCollection> ClothCollection = MakeShared<FManagedArrayCollection>(MoveTemp(InCollection));
+
+	FCollectionClothConstFacade ClothFacade(ClothCollection);
+
+	const FSoftObjectPath& SkeletalMeshPathName = ClothFacade.GetSkeletalMeshSoftObjectPathName();
+	if (const USkeletalMesh* const SkeletalMesh = Cast<USkeletalMesh>(SkeletalMeshPathName.TryLoad()))
+	{
+		const FReferenceSkeleton& RefSkeleton = SkeletalMesh->GetRefSkeleton();
+		const int32 ReferenceBoneIndex = FClothEngineTools::CalculateReferenceBoneIndex(ClothCollection, RefSkeleton);
+		if (RefSkeleton.GetRawRefBoneInfo().IsValidIndex(ReferenceBoneIndex))
+		{
+			ReferenceBone.Name = RefSkeleton.GetRawRefBoneInfo()[ReferenceBoneIndex].Name;
+		}
+	}
+}
+
+#undef LOCTEXT_NAMESPACE
diff --git a/Engine/Plugins/ChaosClothAssetEditor/Source/ChaosClothAssetDataflowNodes/Private/ChaosClothAsset/ReferenceBoneNode.h b/Engine/Plugins/ChaosClothAssetEditor/Source/ChaosClothAssetDataflowNodes/Private/ChaosClothAsset/ReferenceBoneNode.h
new file mode 100644
--- /dev/null
+++ b/Engine/Plugins/ChaosClothAssetEditor/Source/ChaosClothAssetDataflowNodes/Private/ChaosClothAsset/ReferenceBoneNode.h
@@ -0,0 +1,53 @@
+// Copyright Epic Games, Inc. All Rights Reserved.
+
+#pragma once 
+
+#include "Dataflow/DataflowNode.h"
+#include "Dataflow/DataflowFunctionProperty.h"
+#include "GeometryCollection/ManagedArrayCollection.h"
+#include "ReferenceBoneNode.generated.h"
+
+/**
+ * The managed array collection group used in the selection.
+ * This separate structure is required to allow for customization of the UI.
+ */
+USTRUCT()
+struct FChaosClothAssetReferenceBoneSelection
+{
+	GENERATED_USTRUCT_BODY()
+
+	UPROPERTY(EditAnywhere, Category = "Bone Name")
+	FName Name = "";
+};
+
+/** Explicitly set the cloth Reference Bone (used when calculating SimulationVelocityScale). This will automatically computed as the bone closest to the root with weights if no valid Reference Bone is explicitly set.
+* Reference bones are shared by all LODs. Only the bones set for LOD0 will be used.
+*/
+USTRUCT(Meta = (DataflowCloth))
+struct FChaosClothAssetReferenceBoneNode : public FDataflowNode
+{
+	GENERATED_USTRUCT_BODY()
+	DATAFLOW_NODE_DEFINE_INTERNAL(FChaosClothAssetReferenceBoneNode, "ReferenceBone", "Cloth", "Cloth Reference Bone")
+	DATAFLOW_NODE_RENDER_TYPE("SurfaceRender", FName("FClothCollection"), "Collection")
+
+public:
+
+	FChaosClothAssetReferenceBoneNode(const UE::Dataflow::FNodeParameters& InParam, FGuid InGuid = FGuid::NewGuid());
+
+private:
+
+	virtual void Evaluate(UE::Dataflow::FContext& Context, const FDataflowOutput* Out) const override;
+	void OnCalculateDefaultReferenceBone(UE::Dataflow::FContext& Context);
+
+	UPROPERTY(Meta = (Dataflowinput, DataflowOutput, DataflowPassthrough = "Collection"))
+	FManagedArrayCollection Collection;
+
+	/** Select the cloth reference bone. This will automatically computed as the bone closest to the root with weights if no valid Reference Bone is explicitly set. */
+	UPROPERTY(EditAnywhere, Category = "Simulation Reference Bone Selection")
+	FChaosClothAssetReferenceBoneSelection ReferenceBone;
+
+	UPROPERTY(EditAnywhere, Transient, Category = "Simulation Reference Bone Selection", Meta = (DisplayName = "Calculate Default"))
+	FDataflowFunctionProperty CalculateDefaultReferenceBone;
+};
+
+
