Chaos Cloth Asset Editor - Fixed crash when deleting the weight map node while the painting tool is active.

[RN] Choas Cloth Asset Editor - Fixed crash when deleting the weight map node while the painting tool is active.

diff --git a/Engine/Plugins/ChaosClothAssetEditor/Source/ChaosClothAssetEditorTools/Private/ChaosClothAsset/ClothEditorToolBuilders.cpp b/Engine/Plugins/ChaosClothAssetEditor/Source/ChaosClothAssetEditorTools/Private/ChaosClothAsset/ClothEditorToolBuilders.cpp
--- a/Engine/Plugins/ChaosClothAssetEditor/Source/ChaosClothAssetEditorTools/Private/ChaosClothAsset/ClothEditorToolBuilders.cpp
+++ b/Engine/Plugins/ChaosClothAssetEditor/Source/ChaosClothAssetEditorTools/Private/ChaosClothAsset/ClothEditorToolBuilders.cpp
@@ -130,7 +130,7 @@ bool UClothEditorWeightMapPaintToolBuilder::CanBuildTool(const FToolBuilderState
 {
 	if (UDataflowContextObject* const DataflowContextObject = SceneState.ToolManager->GetContextObjectStore()->FindContext<UDataflowContextObject>())
 	{
-		return (DataflowContextObject->GetSelectedNode() != nullptr);
+		return (DataflowContextObject->GetSelectedNodeOfType<FChaosClothAssetWeightMapNode>() != nullptr);
 	}
 	return false;
 }
