Cloth Editor : prevent tool launch from crashing when selecting a disconnected weight map node

diff --git a/Engine/Plugins/ChaosClothAssetEditor/Source/ChaosClothAssetEditorTools/Private/ChaosClothAsset/ClothEditorToolBuilders.cpp b/Engine/Plugins/ChaosClothAssetEditor/Source/ChaosClothAssetEditorTools/Private/ChaosClothAsset/ClothEditorToolBuilders.cpp
--- a/Engine/Plugins/ChaosClothAssetEditor/Source/ChaosClothAssetEditorTools/Private/ChaosClothAsset/ClothEditorToolBuilders.cpp
+++ b/Engine/Plugins/ChaosClothAssetEditor/Source/ChaosClothAssetEditorTools/Private/ChaosClothAsset/ClothEditorToolBuilders.cpp
@@ -128,9 +128,12 @@ void UClothEditorWeightMapPaintToolBuilder::GetSupportedViewModes(const UDataflo
 
 bool UClothEditorWeightMapPaintToolBuilder::CanBuildTool(const FToolBuilderState& SceneState) const
 {
-	if (UDataflowContextObject* const DataflowContextObject = SceneState.ToolManager->GetContextObjectStore()->FindContext<UDataflowContextObject>())
+	if (UMeshSurfacePointMeshEditingToolBuilder::CanBuildTool(SceneState))
 	{
-		return (DataflowContextObject->GetSelectedNodeOfType<FChaosClothAssetWeightMapNode>() != nullptr);
+		if (UDataflowContextObject* const DataflowContextObject = SceneState.ToolManager->GetContextObjectStore()->FindContext<UDataflowContextObject>())
+		{
+			return (DataflowContextObject->GetSelectedNodeOfType<FChaosClothAssetWeightMapNode>() != nullptr);
+		}
 	}
 	return false;
 }
@@ -251,11 +254,13 @@ const FToolTargetTypeRequirements& UClothMeshSelectionToolBuilder::GetTargetRequ
 
 bool UClothMeshSelectionToolBuilder::CanBuildTool(const FToolBuilderState& SceneState) const
 {
-	if (UDataflowContextObject* const DataflowContextObject = SceneState.ToolManager->GetContextObjectStore()->FindContext<UDataflowContextObject>())
+	if (UInteractiveToolWithToolTargetsBuilder::CanBuildTool(SceneState))
 	{
-		return DataflowContextObject->GetSelectedNodeOfType<FChaosClothAssetSelectionNode_v2>() != nullptr && (SceneState.TargetManager->CountSelectedAndTargetable(SceneState, GetTargetRequirements()) == 1);
+		if (UDataflowContextObject* const DataflowContextObject = SceneState.ToolManager->GetContextObjectStore()->FindContext<UDataflowContextObject>())
+		{
+			return DataflowContextObject->GetSelectedNodeOfType<FChaosClothAssetSelectionNode_v2>() != nullptr && (SceneState.TargetManager->CountSelectedAndTargetable(SceneState, GetTargetRequirements()) == 1);
+		}
 	}
-
 	return false;
 }
 
