Dataflow : Use DF_ as prefix as default when creating a dataflow asset from another asset

diff --git a/Engine/Plugins/ChaosClothAssetEditor/Source/ChaosClothAssetEditor/Private/ChaosClothAsset/AssetDefinition_ClothAsset.cpp b/Engine/Plugins/ChaosClothAssetEditor/Source/ChaosClothAssetEditor/Private/ChaosClothAsset/AssetDefinition_ClothAsset.cpp
--- a/Engine/Plugins/ChaosClothAssetEditor/Source/ChaosClothAssetEditor/Private/ChaosClothAsset/AssetDefinition_ClothAsset.cpp
+++ b/Engine/Plugins/ChaosClothAssetEditor/Source/ChaosClothAssetEditor/Private/ChaosClothAsset/AssetDefinition_ClothAsset.cpp
@@ -41,7 +41,7 @@ namespace ClothAssetDefinitionHelpers
 			const FString PackageName = ClothAsset->GetOutermost()->GetName();
 			NewDataflowAssetDialogConfig.DefaultPath = FPackageName::GetLongPackagePath(PackageName);
 			const FString ClothName = ClothAsset->GetName();
-			NewDataflowAssetDialogConfig.DefaultAssetName = ClothName + "_Dataflow";
+			NewDataflowAssetDialogConfig.DefaultAssetName = "DF_" + ClothName;
 			NewDataflowAssetDialogConfig.AssetClassNames.Add(DataflowClass->GetClassPathName());
 			NewDataflowAssetDialogConfig.ExistingAssetPolicy = ESaveAssetDialogExistingAssetPolicy::AllowButWarn;
 			NewDataflowAssetDialogConfig.DialogTitleOverride = LOCTEXT("NewDataflowAssetDialogTitle", "Save Dataflow Asset As");
diff --git a/Engine/Plugins/Experimental/Dataflow/Source/DataflowEditor/Private/Dataflow/AssetDefinition_DataflowAsset.cpp b/Engine/Plugins/Experimental/Dataflow/Source/DataflowEditor/Private/Dataflow/AssetDefinition_DataflowAsset.cpp
--- a/Engine/Plugins/Experimental/Dataflow/Source/DataflowEditor/Private/Dataflow/AssetDefinition_DataflowAsset.cpp
+++ b/Engine/Plugins/Experimental/Dataflow/Source/DataflowEditor/Private/Dataflow/AssetDefinition_DataflowAsset.cpp
@@ -34,7 +34,7 @@ namespace UE::DataflowAssetDefinitionHelpers
 			const FString PackageName = Asset->GetOutermost()->GetName();
 			NewDataflowAssetDialogConfig.DefaultPath = FPackageName::GetLongPackagePath(PackageName);
 			const FString AssetName = Asset->GetName();
-			NewDataflowAssetDialogConfig.DefaultAssetName = AssetName + "_Dataflow";
+			NewDataflowAssetDialogConfig.DefaultAssetName = "DF_" + AssetName;
 			NewDataflowAssetDialogConfig.AssetClassNames.Add(DataflowClass->GetClassPathName());
 			NewDataflowAssetDialogConfig.ExistingAssetPolicy = ESaveAssetDialogExistingAssetPolicy::Disallow;
 			NewDataflowAssetDialogConfig.DialogTitleOverride = LOCTEXT("NewDataflowAssetDialogTitle", "Save Dataflow Asset As");
