Cloth tool builder : fix assert when opening a selection tool in Dataflow
- this is a side effect of a recent change

diff --git a/Engine/Plugins/ChaosClothAssetEditor/Source/ChaosClothAssetEditorTools/Private/ChaosClothAsset/ClothEditorToolBuilders.cpp b/Engine/Plugins/ChaosClothAssetEditor/Source/ChaosClothAssetEditorTools/Private/ChaosClothAsset/ClothEditorToolBuilders.cpp
--- a/Engine/Plugins/ChaosClothAssetEditor/Source/ChaosClothAssetEditorTools/Private/ChaosClothAsset/ClothEditorToolBuilders.cpp
+++ b/Engine/Plugins/ChaosClothAssetEditor/Source/ChaosClothAssetEditorTools/Private/ChaosClothAsset/ClothEditorToolBuilders.cpp
@@ -254,12 +254,9 @@ const FToolTargetTypeRequirements& UClothMeshSelectionToolBuilder::GetTargetRequ
 
 bool UClothMeshSelectionToolBuilder::CanBuildTool(const FToolBuilderState& SceneState) const
 {
-	if (UInteractiveToolWithToolTargetsBuilder::CanBuildTool(SceneState))
+	if (UDataflowContextObject* const DataflowContextObject = SceneState.ToolManager->GetContextObjectStore()->FindContext<UDataflowContextObject>())
 	{
-		if (UDataflowContextObject* const DataflowContextObject = SceneState.ToolManager->GetContextObjectStore()->FindContext<UDataflowContextObject>())
-		{
-			return DataflowContextObject->GetSelectedNodeOfType<FChaosClothAssetSelectionNode_v2>() != nullptr && (SceneState.TargetManager->CountSelectedAndTargetable(SceneState, GetTargetRequirements()) == 1);
-		}
+		return DataflowContextObject->GetSelectedNodeOfType<FChaosClothAssetSelectionNode_v2>() != nullptr && (SceneState.TargetManager->CountSelectedAndTargetable(SceneState, GetTargetRequirements()) == 1);
 	}
 	return false;
 }
