Chaos Cloth Asset - Fixed crash in SkeletalMeshImport node when importing a render section indexing a material that has been remapped.

[RN] Chaos Cloth Asset - Fixed crash in SkeletalMeshImport node when importing a render section indexing a material that has been remapped.

diff --git a/Engine/Plugins/ChaosClothAssetEditor/Source/ChaosClothAssetDataflowNodes/Private/ChaosClothAsset/SkeletalMeshImportNode.cpp b/Engine/Plugins/ChaosClothAssetEditor/Source/ChaosClothAssetDataflowNodes/Private/ChaosClothAsset/SkeletalMeshImportNode.cpp
--- a/Engine/Plugins/ChaosClothAssetEditor/Source/ChaosClothAssetDataflowNodes/Private/ChaosClothAsset/SkeletalMeshImportNode.cpp
+++ b/Engine/Plugins/ChaosClothAssetEditor/Source/ChaosClothAssetDataflowNodes/Private/ChaosClothAsset/SkeletalMeshImportNode.cpp
@@ -98,8 +98,18 @@ void FChaosClothAssetSkeletalMeshImportNode_v2::Evaluate(UE::Dataflow::FContext&
 				if (bImportRenderMesh)
 				{
 					const TArray<FSkeletalMaterial>& Materials = InSkeletalMesh->GetMaterials();
-					check(Section < Materials.Num());
-					const FString RenderMaterialPathName = Materials[Section].MaterialInterface ? Materials[Section].MaterialInterface->GetPathName() : "";
+
+					const FSkeletalMeshLODInfo* const SkeletalMeshLODInfo = InSkeletalMesh->GetLODInfo(LODIndex);
+					const int32 MaterialIndex =
+						SkeletalMeshLODInfo &&
+						SkeletalMeshLODInfo->LODMaterialMap.IsValidIndex(Section) &&
+						SkeletalMeshLODInfo->LODMaterialMap[Section] != INDEX_NONE ?
+							SkeletalMeshLODInfo->LODMaterialMap[Section] :
+							LODModel.Sections[Section].MaterialIndex;
+
+					const FString RenderMaterialPathName = Materials.IsValidIndex(MaterialIndex) && Materials[MaterialIndex].MaterialInterface ?
+						 Materials[MaterialIndex].MaterialInterface->GetPathName() : "";
+
 					FClothDataflowTools::AddRenderPatternFromSkeletalMeshSection(ClothCollection, LODModel, Section, RenderMaterialPathName);
 				}
 			}
