Cloth : use the new unified Dataflow Editor by default and add the option to open the legacy panel editor from right clicking on the asset

diff --git a/Engine/Plugins/ChaosClothAssetEditor/Source/ChaosClothAssetEditor/Private/ChaosClothAsset/AssetDefinition_ClothAsset.cpp b/Engine/Plugins/ChaosClothAssetEditor/Source/ChaosClothAssetEditor/Private/ChaosClothAsset/AssetDefinition_ClothAsset.cpp
--- a/Engine/Plugins/ChaosClothAssetEditor/Source/ChaosClothAssetEditor/Private/ChaosClothAsset/AssetDefinition_ClothAsset.cpp
+++ b/Engine/Plugins/ChaosClothAssetEditor/Source/ChaosClothAssetEditor/Private/ChaosClothAsset/AssetDefinition_ClothAsset.cpp
@@ -25,8 +25,8 @@
 
 namespace UE::Chaos::ClothAsset::Private
 {
-static int32 bEnableClothDataflowEditor = false;
-static FAutoConsoleVariableRef CVarEnableClothDataflowEditor(TEXT("p.ChaosCloth.EnableDataflowEditor"), bEnableClothDataflowEditor, TEXT("Enable the use of the core dataflow editor for cloth asset (WIP)"));
+static int32 bEnableLegacyClothPanelEditor = false;
+static FAutoConsoleVariableRef CVarEnableLegacyClothPanelEditor(TEXT("p.ChaosCloth.EnableLegacyClothPanelEditor"), bEnableLegacyClothPanelEditor, TEXT("Enable the use of the legacy Cloth Panel Editor instead of the unified Dataflow one for editing cloth assets"));
 }
 
 namespace ClothAssetDefinitionHelpers
@@ -183,7 +183,7 @@ EAssetCommandResult UAssetDefinition_ClothAsset::OpenAssets(const FAssetOpenArgs
 	if (ClothObjects.Num() > 0)
 	{
 		UAssetEditorSubsystem* const AssetEditorSubsystem = GEditor->GetEditorSubsystem<UAssetEditorSubsystem>();
-		if(!UE::Chaos::ClothAsset::Private::bEnableClothDataflowEditor)
+		if(UE::Chaos::ClothAsset::Private::bEnableLegacyClothPanelEditor)
 		{
 			UChaosClothAssetEditor* const AssetEditor = NewObject<UChaosClothAssetEditor>(AssetEditorSubsystem, NAME_None, RF_Transient);
 			AssetEditor->Initialize({ClothObjects[0] });
diff --git a/Engine/Plugins/ChaosClothAssetEditor/Source/ChaosClothAssetEditor/Private/ChaosClothAsset/ChaosClothAssetEditorModule.cpp b/Engine/Plugins/ChaosClothAssetEditor/Source/ChaosClothAssetEditor/Private/ChaosClothAsset/ChaosClothAssetEditorModule.cpp
--- a/Engine/Plugins/ChaosClothAssetEditor/Source/ChaosClothAssetEditor/Private/ChaosClothAsset/ChaosClothAssetEditorModule.cpp
+++ b/Engine/Plugins/ChaosClothAssetEditor/Source/ChaosClothAssetEditor/Private/ChaosClothAsset/ChaosClothAssetEditorModule.cpp
@@ -2,12 +2,14 @@
 
 #include "ChaosClothAsset/ChaosClothAssetEditorModule.h"
 
+
 #include "ChaosClothAsset/AssetDefinition_ClothAsset.h"
 #include "ChaosClothAsset/ChaosClothAssetThumbnailRenderer.h"
 #include "ChaosClothAsset/ClothAsset.h"
 #include "ChaosClothAsset/ClothComponent.h"
 #include "ChaosClothAsset/ClothDataflowConstructionVisualization.h"
 #include "ChaosClothAsset/ClothDataflowSimulationVisualization.h"
+#include "ChaosClothAsset/ClothEditor.h"
 #include "ChaosClothAsset/ClothSimulationNodeDetailExtender.h"
 #include "ChaosClothAsset/ClothEditorCommands.h"
 #include "ChaosClothAsset/ClothEditorMode.h"
@@ -87,37 +89,20 @@ namespace UE::Chaos::ClothAsset
 					const TSharedPtr<class FUICommandList> CommandListToBind = MakeShared<FUICommandList>();
 
 					CommandListToBind->MapAction(
-						FChaosClothAssetEditorCommands::Get().OpenClothAssetInDataflowEditor,
+						FChaosClothAssetEditorCommands::Get().OpenClothAssetInLegacyClothPanelEditor,
 						FExecuteAction::CreateWeakLambda(DataflowEditorSubsystem, [Context, DataflowEditorSubsystem]()
 						{
 							// When we actually do want to open the editor, trigger the load to get the objects
 							TArray<TObjectPtr<UObject>> AssetsToEdit;
 							AssetsToEdit.Append(Context->LoadSelectedObjects<UObject>());
 
-							UDataflowEditor* const AssetEditor = NewObject<UDataflowEditor>(DataflowEditorSubsystem, NAME_None, RF_Transient);
-							AssetEditor->RegisterToolCategories({"General", "Cloth"});
-
-							// Macke sure the cloth asset has a Dataflow asset
-							UChaosClothAsset* const ClothAsset = CastChecked<UChaosClothAsset>(AssetsToEdit[0]);
-							if (!ClothAsset->GetDataflow())
-							{
-PRAGMA_DISABLE_DEPRECATION_WARNINGS
-								if (UDataflow* const NewDataflowAsset = Cast<UDataflow>(UAssetDefinition_ClothAsset::NewOrOpenDataflowAsset(ClothAsset)))
-								{
-									ClothAsset->SetDataflow(NewDataflowAsset);
-								}
-PRAGMA_ENABLE_DEPRECATION_WARNINGS
-							}
-
-							const TSubclassOf<AActor> PreviewActorClass = StaticLoadClass(AActor::StaticClass(), nullptr,
-								TEXT("/ChaosClothAssetEditor/BP_ClothPreview.BP_ClothPreview_C"), nullptr, LOAD_None, nullptr);
-
-							AssetEditor->Initialize({AssetsToEdit[0]}, PreviewActorClass);
+							UChaosClothAssetEditor* const AssetEditor = NewObject<UChaosClothAssetEditor>(DataflowEditorSubsystem, NAME_None, RF_Transient);
+							AssetEditor->Initialize({ AssetsToEdit[0] });
 						}),
 						FCanExecuteAction::CreateWeakLambda(Context, [Context]() { return Context->bCanBeModified; }));
 
 					const TAttribute<FText> ToolTipOverride = Context->bCanBeModified ? TAttribute<FText>() : LOCTEXT("ReadOnlyAssetWarning", "The selected asset(s) are read-only and cannot be edited.");
-					Section.AddMenuEntryWithCommandList(FChaosClothAssetEditorCommands::Get().OpenClothAssetInDataflowEditor,
+					Section.AddMenuEntryWithCommandList(FChaosClothAssetEditorCommands::Get().OpenClothAssetInLegacyClothPanelEditor,
 						CommandListToBind,
 						TAttribute<FText>(),
 						ToolTipOverride,
diff --git a/Engine/Plugins/ChaosClothAssetEditor/Source/ChaosClothAssetEditor/Private/ChaosClothAsset/ClothEditorCommands.cpp b/Engine/Plugins/ChaosClothAssetEditor/Source/ChaosClothAssetEditor/Private/ChaosClothAsset/ClothEditorCommands.cpp
--- a/Engine/Plugins/ChaosClothAssetEditor/Source/ChaosClothAssetEditor/Private/ChaosClothAsset/ClothEditorCommands.cpp
+++ b/Engine/Plugins/ChaosClothAssetEditor/Source/ChaosClothAssetEditor/Private/ChaosClothAsset/ClothEditorCommands.cpp
@@ -40,7 +40,7 @@ void FChaosClothAssetEditorCommands::RegisterCommands()
 
 	UI_COMMAND(OpenClothEditor, "Cloth Editor", "Open the Cloth Editor window", EUserInterfaceActionType::Button, FInputChord());
 
-	UI_COMMAND(OpenClothAssetInDataflowEditor, "Open in Dataflow Editor (experimental)", "Open the Cloth Asset in Dataflow Editor", EUserInterfaceActionType::Button, FInputChord());
+	UI_COMMAND(OpenClothAssetInLegacyClothPanelEditor, "Open in legacy Cloth Panel Editor", "Open the Cloth Asset in the legacy Cloth Panel Editor instead of the unified Dataflow one", EUserInterfaceActionType::Button, FInputChord());
 
 	UI_COMMAND(BeginRemeshTool, "Remesh", "Remesh the selected mesh", EUserInterfaceActionType::Button, FInputChord());
 	UI_COMMAND(BeginAttributeEditorTool, "AttrEd", "Edit/configure mesh attributes", EUserInterfaceActionType::Button, FInputChord());
diff --git a/Engine/Plugins/ChaosClothAssetEditor/Source/ChaosClothAssetEditor/Private/ChaosClothAsset/ClothEditorCommands.h b/Engine/Plugins/ChaosClothAssetEditor/Source/ChaosClothAssetEditor/Private/ChaosClothAsset/ClothEditorCommands.h
--- a/Engine/Plugins/ChaosClothAssetEditor/Source/ChaosClothAssetEditor/Private/ChaosClothAsset/ClothEditorCommands.h
+++ b/Engine/Plugins/ChaosClothAssetEditor/Source/ChaosClothAssetEditor/Private/ChaosClothAsset/ClothEditorCommands.h
@@ -30,7 +30,7 @@ public:
 
 	TSharedPtr<FUICommandInfo> OpenClothEditor;
 
-	TSharedPtr<FUICommandInfo> OpenClothAssetInDataflowEditor;
+	TSharedPtr<FUICommandInfo> OpenClothAssetInLegacyClothPanelEditor;
 
 	const static FString BeginRemeshToolIdentifier;
 	TSharedPtr<FUICommandInfo> BeginRemeshTool;
