Cloth editor : disabled async caching for cloth through Dataflow editor simulation settings

diff --git a/Engine/Plugins/ChaosClothAssetEditor/Source/ChaosClothAssetEditor/Private/ChaosClothAsset/AssetDefinition_ClothAsset.cpp b/Engine/Plugins/ChaosClothAssetEditor/Source/ChaosClothAssetEditor/Private/ChaosClothAsset/AssetDefinition_ClothAsset.cpp
--- a/Engine/Plugins/ChaosClothAssetEditor/Source/ChaosClothAssetEditor/Private/ChaosClothAsset/AssetDefinition_ClothAsset.cpp
+++ b/Engine/Plugins/ChaosClothAssetEditor/Source/ChaosClothAssetEditor/Private/ChaosClothAsset/AssetDefinition_ClothAsset.cpp
@@ -17,6 +17,7 @@
 #include "ChaosClothAsset/TerminalNode.h"
 #include "Dataflow/DataflowEditor.h"
 #include "Dataflow/DataflowSNode.h"
+#include "Dataflow/DataflowSimulationScene.h"
 #include "Toolkits/SimpleAssetEditor.h"
 
 #include UE_INLINE_GENERATED_CPP_BY_NAME(AssetDefinition_ClothAsset)
@@ -200,6 +201,15 @@ EAssetCommandResult UAssetDefinition_ClothAsset::OpenAssets(const FAssetOpenArgs
 			else
 			{
 				UDataflowEditor* const AssetEditor = NewObject<UDataflowEditor>(AssetEditorSubsystem, NAME_None, RF_Transient);
+
+				if (UDataflowSimulationSettings* SimulationSettings = NewObject<UDataflowSimulationSettings>())
+				{
+					SimulationSettings->bIsSimulationPlayingByDefault = true;
+					SimulationSettings->bIsAsyncCachingSupported = false;
+					SimulationSettings->bIsAsyncCachingEnabledByDefault = false;
+					AssetEditor->AddEditorSettings(SimulationSettings);
+				}
+
 				const TSubclassOf<AActor> ActorClass = StaticLoadClass(AActor::StaticClass(), nullptr,
 					TEXT("/ChaosClothAssetEditor/BP_ClothPreview.BP_ClothPreview_C"), nullptr, LOAD_None, nullptr);
 				AssetEditor->RegisterToolCategories({"General", "Cloth"});
