Dataflow : add cloth collection renderable types for the new rendering system
- note : they are not used not registered or used yet

diff --git a/Engine/Plugins/ChaosClothAssetEditor/Source/ChaosClothAssetDataflowNodes/Private/ChaosClothAsset/ClothRenderableTypes.cpp b/Engine/Plugins/ChaosClothAssetEditor/Source/ChaosClothAssetDataflowNodes/Private/ChaosClothAsset/ClothRenderableTypes.cpp
new file mode 100644
--- /dev/null
+++ b/Engine/Plugins/ChaosClothAssetEditor/Source/ChaosClothAssetDataflowNodes/Private/ChaosClothAsset/ClothRenderableTypes.cpp
@@ -0,0 +1,176 @@
+// Copyright Epic Games, Inc. All Rights Reserved.
+
+#include "ChaosClothAsset/ClothRenderableTypes.h"
+
+#include "Components/DynamicMeshComponent.h"
+#include "Dataflow/DataflowRenderingViewMode.h"
+#include "Dataflow/DataflowRenderingViewMode.h"
+#include "DataflowRendering/DataflowRenderableComponents.h"
+#include "DataflowRendering/DataflowRenderableType.h"
+#include "DataflowRendering/DataflowRenderableTypeInstance.h"
+#include "DataflowRendering/DataflowRenderableTypeRegistry.h"
+#include "DataflowRendering/DataflowRenderableTypeUtils.h"
+
+#include "GeometryCollection/ManagedArrayCollection.h"
+
+#include "UObject/ObjectPtr.h"
+
+#include "ChaosClothAsset/ClothCollectionGroup.h"
+#include "ChaosClothAsset/CollectionClothFacade.h"
+#include "ChaosClothAsset/ClothPatternToDynamicMesh.h"
+
+#include "Materials/Material.h"
+
+namespace UE::Chaos::ClothAsset
+{
+	namespace Private
+	{
+		static UMaterialInterface* GetTwoSidedMaterialForRenderingTypes()
+		{
+			return Cast<UMaterial>(StaticLoadObject(UMaterial::StaticClass(), NULL, TEXT("/Engine/EditorMaterials/Dataflow/DataflowTwoSidedVertexMaterial")));
+		}
+	}
+	////////////////////////////////////////////////////////////////////////////////////////////////////////////////////////////
+
+	class FClothSimSurfaceRenderableType : public UE::Dataflow::IRenderableType
+	{
+		UE_DATAFLOW_IRENDERABLE_PRIMARY_TYPE(FManagedArrayCollection, Collection);
+		UE_DATAFLOW_IRENDERABLE_RENDER_GROUP(SimSurface);
+		UE_DATAFLOW_IRENDERABLE_VIEW_MODE(UE::Dataflow::FDataflowConstruction3DViewMode);
+
+		virtual bool CanRender(const UE::Dataflow::FRenderableTypeInstance& Instance) const override
+		{
+			const FManagedArrayCollection& Collection = GetCollection(Instance);
+			return Collection.HasGroup(ClothCollectionGroup::SimVertices3D);
+		}
+
+		virtual void GetPrimitiveComponents(const UE::Dataflow::FRenderableTypeInstance& Instance, UE::Dataflow::FRenderableComponents& OutComponents) const override
+		{
+			FManagedArrayCollection Collection = GetCollection(Instance);
+
+			const TSharedRef<const FManagedArrayCollection> ClothCollection = MakeShared<FManagedArrayCollection>(MoveTemp(Collection));
+			const FCollectionClothConstFacade ClothFacade(ClothCollection);
+			if (ClothFacade.HasValidSimulationData())
+			{
+				FDynamicMesh3 DynamicMesh;
+				FClothPatternToDynamicMesh Converter;
+				Converter.Convert(ClothCollection, INDEX_NONE, EClothPatternVertexType::Sim3D, DynamicMesh);
+
+				// only create components for the one that have vertices
+				if (DynamicMesh.VertexCount() > 0)
+				{
+					static const FName ComponentName = TEXT("ClothSim3DMesh");
+					if (UDynamicMeshComponent* Component = OutComponents.AddNewComponent<UDynamicMeshComponent>(ComponentName))
+					{
+						Component->SetCastShadow(false);
+						Component->SetMesh(MoveTemp(DynamicMesh));
+
+						UMaterialInterface* TwoSidedMaterial = Private::GetTwoSidedMaterialForRenderingTypes();
+						Component->SetOverrideRenderMaterial(TwoSidedMaterial);
+					}
+				}
+			}
+		}
+	};
+
+	///////////////////////////////////////////////////////////////////////////////////////////////////////////////
+
+	class FClothRenderSurfaceRenderableType : public UE::Dataflow::IRenderableType
+	{
+		UE_DATAFLOW_IRENDERABLE_PRIMARY_TYPE(FManagedArrayCollection, Collection);
+		UE_DATAFLOW_IRENDERABLE_RENDER_GROUP(RenderSurface);
+		UE_DATAFLOW_IRENDERABLE_VIEW_MODE(UE::Dataflow::FDataflowConstruction3DViewMode);
+
+		virtual bool CanRender(const UE::Dataflow::FRenderableTypeInstance& Instance) const override
+		{
+			const FManagedArrayCollection& Collection = GetCollection(Instance);
+			return Collection.HasGroup(ClothCollectionGroup::RenderVertices);
+		}
+
+		virtual void GetPrimitiveComponents(const UE::Dataflow::FRenderableTypeInstance& Instance, UE::Dataflow::FRenderableComponents& OutComponents) const override
+		{
+			FManagedArrayCollection Collection = GetCollection(Instance);
+
+			const TSharedRef<const FManagedArrayCollection> ClothCollection = MakeShared<FManagedArrayCollection>(MoveTemp(Collection));
+			const FCollectionClothConstFacade ClothFacade(ClothCollection);
+			if (ClothFacade.HasValidRenderData())
+			{
+				FDynamicMesh3 DynamicMesh;
+				FClothPatternToDynamicMesh Converter;
+				Converter.Convert(ClothCollection, INDEX_NONE, EClothPatternVertexType::Render, DynamicMesh);
+
+				// only create components for the one that have vertices
+				if (DynamicMesh.VertexCount() > 0)
+				{
+					static const FName ComponentName = TEXT("ClothRender3DMesh");
+					if (UDynamicMeshComponent* Component = OutComponents.AddNewComponent<UDynamicMeshComponent>(ComponentName))
+					{
+						Component->SetCastShadow(false);
+						Component->SetMesh(MoveTemp(DynamicMesh));
+
+						const TConstArrayView<FSoftObjectPath> MaterialPaths = ClothFacade.GetRenderMaterialSoftObjectPathName();
+						const int32 NumMaterials = MaterialPaths.Num();
+
+						for (int32 Index = 0; Index < NumMaterials; ++Index)
+						{
+							UMaterialInterface* const Material = Cast<UMaterialInterface>(MaterialPaths[Index].TryLoad());
+							Component->SetMaterial(Index, Material);
+						}
+					}
+				}
+			}
+		}
+	};
+
+
+	///////////////////////////////////////////////////////////////////////////////////////////////////////////////
+
+	class FClothSimPatternRenderableType : public UE::Dataflow::IRenderableType
+	{
+		UE_DATAFLOW_IRENDERABLE_PRIMARY_TYPE(FManagedArrayCollection, Collection);
+		UE_DATAFLOW_IRENDERABLE_RENDER_GROUP(Pattern);
+		UE_DATAFLOW_IRENDERABLE_VIEW_MODE(UE::Dataflow::FDataflowConstruction2DViewMode);
+
+		virtual bool CanRender(const UE::Dataflow::FRenderableTypeInstance& Instance) const override
+		{
+			const FManagedArrayCollection& Collection = GetCollection(Instance);
+			return Collection.HasGroup(ClothCollectionGroup::SimVertices2D);
+		}
+
+		virtual void GetPrimitiveComponents(const UE::Dataflow::FRenderableTypeInstance& Instance, UE::Dataflow::FRenderableComponents& OutComponents) const override
+		{
+			FManagedArrayCollection Collection = GetCollection(Instance);
+
+			const TSharedRef<const FManagedArrayCollection> ClothCollection = MakeShared<FManagedArrayCollection>(MoveTemp(Collection));
+			const FCollectionClothConstFacade ClothFacade(ClothCollection);
+			if (ClothFacade.HasValidSimulationData())
+			{
+				FDynamicMesh3 DynamicMesh;
+				FClothPatternToDynamicMesh Converter;
+				Converter.Convert(ClothCollection, INDEX_NONE, EClothPatternVertexType::Sim2D, DynamicMesh);
+
+				// only create components for the one that have vertices
+				if (DynamicMesh.VertexCount() > 0)
+				{
+					static const FName ComponentName = TEXT("ClothSim2DMesh");
+					if (UDynamicMeshComponent* Component = OutComponents.AddNewComponent<UDynamicMeshComponent>(ComponentName))
+					{
+						Component->SetCastShadow(false);
+						Component->SetMesh(MoveTemp(DynamicMesh));
+
+						UMaterialInterface* TwoSidedMaterial = Private::GetTwoSidedMaterialForRenderingTypes();
+						Component->SetOverrideRenderMaterial(TwoSidedMaterial);
+					}
+				}
+			}
+		}
+	};
+	///////////////////////////////////////////////////////////////////////////////////////////////////////////////
+
+	void RegisterClothRenderableTypes()
+	{
+		UE_DATAFLOW_REGISTER_RENDERABLE_TYPE(FClothSimSurfaceRenderableType);
+		UE_DATAFLOW_REGISTER_RENDERABLE_TYPE(FClothRenderSurfaceRenderableType);
+		UE_DATAFLOW_REGISTER_RENDERABLE_TYPE(FClothSimPatternRenderableType);
+	}
+}
diff --git a/Engine/Plugins/ChaosClothAssetEditor/Source/ChaosClothAssetDataflowNodes/Private/ChaosClothAsset/ClothRenderableTypes.h b/Engine/Plugins/ChaosClothAssetEditor/Source/ChaosClothAssetDataflowNodes/Private/ChaosClothAsset/ClothRenderableTypes.h
new file mode 100644
--- /dev/null
+++ b/Engine/Plugins/ChaosClothAssetEditor/Source/ChaosClothAssetDataflowNodes/Private/ChaosClothAsset/ClothRenderableTypes.h
@@ -0,0 +1,10 @@
+// Copyright Epic Games, Inc. All Rights Reserved.
+
+#pragma once
+
+namespace UE::Chaos::ClothAsset
+{
+	void RegisterClothRenderableTypes();
+}  // End namespace UE::Chaos::ClothAsset
+
+
