Set correct number of weights on Cloth WeightMapNode when in 2D view mode.

diff --git a/Engine/Plugins/ChaosClothAssetEditor/Source/ChaosClothAssetDataflowNodes/Private/ChaosClothAsset/WeightMapNode.cpp b/Engine/Plugins/ChaosClothAssetEditor/Source/ChaosClothAssetDataflowNodes/Private/ChaosClothAsset/WeightMapNode.cpp
--- a/Engine/Plugins/ChaosClothAssetEditor/Source/ChaosClothAssetDataflowNodes/Private/ChaosClothAsset/WeightMapNode.cpp
+++ b/Engine/Plugins/ChaosClothAssetEditor/Source/ChaosClothAssetDataflowNodes/Private/ChaosClothAsset/WeightMapNode.cpp
@@ -308,20 +308,23 @@ void FChaosClothAssetWeightMapNode::SetVertexAttributeValues(UE::Dataflow::FCont
 	if (ClothFacade.IsValid())  // Can only act on the collection if it is a valid cloth collection
 	{
 		TConstArrayView<float> InputAttributeValues;
+		int32 NumFinalWeights;
 		if (MeshTarget == EChaosClothAssetWeightMapMeshTarget::Simulation)
 		{
 			InputAttributeValues = ClothFacade.GetWeightMap(InInputName);
+			NumFinalWeights = ClothFacade.GetNumSimVertices3D();
 		}
 		else
 		{
 			InputAttributeValues = ClothFacade.GetUserDefinedAttribute<float>(InInputName, ClothCollectionGroup::RenderVertices);
+			NumFinalWeights = ClothFacade.GetNumRenderVertices();
 		}
 
 		if (InWeightIndices.Num() == InValues.Num())
 		{
 			TArray<float> ValuesToApply;
-			ValuesToApply.Init(0.f, InWeightIndices.Num());
-			for (int32 Index = 0; Index < ValuesToApply.Num(); ++Index)
+			ValuesToApply.Init(0.f, NumFinalWeights);
+			for (int32 Index = 0; Index < InWeightIndices.Num(); ++Index)
 			{
 				ValuesToApply[InWeightIndices[Index]] = InValues[Index];
 			}
