Dataflow : Update the rendering callbacks to better report what view modes are supported
- this allow the cloth nodes to better work in the context of the general dataflow where the proper render / sim mesh can show in the construction view

diff --git a/Engine/Plugins/ChaosClothAssetEditor/Source/ChaosClothAssetDataflowNodes/Private/ChaosClothAsset/ChaosClothAssetDataflowNodesModule.cpp b/Engine/Plugins/ChaosClothAssetEditor/Source/ChaosClothAssetDataflowNodes/Private/ChaosClothAsset/ChaosClothAssetDataflowNodesModule.cpp
--- a/Engine/Plugins/ChaosClothAssetEditor/Source/ChaosClothAssetDataflowNodes/Private/ChaosClothAsset/ChaosClothAssetDataflowNodesModule.cpp
+++ b/Engine/Plugins/ChaosClothAssetEditor/Source/ChaosClothAssetDataflowNodes/Private/ChaosClothAsset/ChaosClothAssetDataflowNodesModule.cpp
@@ -236,6 +236,40 @@ namespace UE::Chaos::ClothAsset
 						ViewModeName == Dataflow::FDataflowConstruction2DViewMode::Name);
 			}
 
+			virtual bool CanRenderFromState(const UE::Dataflow::FGraphRenderingState& State) const override
+			{
+				if (State.GetRenderOutputs().Num())
+				{
+					const FManagedArrayCollection Default;
+					const FName PrimaryOutput = State.GetRenderOutputs()[0];
+					FManagedArrayCollection Collection = State.GetValue<FManagedArrayCollection>(PrimaryOutput, Default);
+					const TSharedRef<const FManagedArrayCollection> ClothCollection = MakeShared<FManagedArrayCollection>(MoveTemp(Collection));
+					const FCollectionClothConstFacade ClothFacade(ClothCollection);
+					if (ClothFacade.IsValid())
+					{
+						const FName& ViewModeName = State.GetViewMode().GetName();
+						const bool bHasSimData = ClothFacade.HasValidSimulationData();
+						const bool bHasRenderData = ClothFacade.HasValidRenderData();
+
+						if (bHasSimData && (ViewModeName == FCloth2DSimViewMode::Name || ViewModeName == Dataflow::FDataflowConstruction2DViewMode::Name))
+						{
+							return true;
+						}
+						if (bHasSimData && (ViewModeName == FCloth3DSimViewMode::Name || ViewModeName == Dataflow::FDataflowConstruction3DViewMode::Name))
+						{
+							return true;
+						}
+						if (bHasRenderData && (ViewModeName == FClothRenderViewMode::Name || ViewModeName == Dataflow::FDataflowConstruction3DViewMode::Name))
+						{
+							return true;
+						}
+
+						return false;
+					}
+				}
+				return CanRender(State.GetViewMode());
+			}
+
 			virtual void Render(GeometryCollection::Facades::FRenderingFacade& RenderCollection, const UE::Dataflow::FGraphRenderingState& State) override
 			{
 				if (State.GetRenderOutputs().Num())
@@ -259,13 +293,18 @@ namespace UE::Chaos::ClothAsset
 					TArray<TArray<FVector2f>> UVs;
 					TArray<int32> RenderMaterialIDs;
 
-					if (State.GetViewMode().GetName() == FCloth3DSimViewMode::Name || State.GetViewMode().GetName() == Dataflow::FDataflowConstruction3DViewMode::Name)
+
+					const FName& ViewModeName = State.GetViewMode().GetName();
+					const bool bHasSimData = ClothFacade.HasValidSimulationData();
+					const bool bHasRenderData = ClothFacade.HasValidRenderData();
+
+					if (bHasSimData && (ViewModeName == FCloth3DSimViewMode::Name || ViewModeName == Dataflow::FDataflowConstruction3DViewMode::Name))
 					{
 						Vertices = ClothFacade.GetSimPosition3D();
 						Indices = ClothFacade.GetSimIndices3D();
 						Normals = ClothFacade.GetSimNormal();
 					}
-					else if (State.GetViewMode().GetName() == FCloth2DSimViewMode::Name || State.GetViewMode().GetName() == Dataflow::FDataflowConstruction2DViewMode::Name)
+					else if (bHasSimData && (ViewModeName == FCloth2DSimViewMode::Name || ViewModeName == Dataflow::FDataflowConstruction2DViewMode::Name))
 					{
 						//
 						// Flip the Y coordinate to get the desired visualization with our LVT_OrthoXY viewport
@@ -299,7 +338,7 @@ namespace UE::Chaos::ClothAsset
 						}
 						Normals.Init(Normal2D, ClothFacade.GetNumSimVertices2D());
 					}
-					else if (State.GetViewMode().GetName() == FClothRenderViewMode::Name)
+					else if (bHasRenderData && (ViewModeName == FClothRenderViewMode::Name || ViewModeName == Dataflow::FDataflowConstruction3DViewMode::Name))
 					{
 						Vertices = ClothFacade.GetRenderPosition();
 						Indices = ClothFacade.GetRenderIndices();
@@ -325,9 +364,10 @@ namespace UE::Chaos::ClothAsset
 					TArray<FLinearColor> Colors;
 					Colors.Init(FLinearColor::Gray, Vertices.Num());	// TODO: Choose a vertex color
 
-					const int32 GeometryIndex = RenderCollection.StartGeometryGroup(State.GetGuid().ToString());
+					const FString GeometryName = FString::Printf(TEXT("%s_%s"), *State.GetNodeName().ToString(), *PrimaryOutput.ToString());
+					const int32 GeometryIndex = RenderCollection.StartGeometryGroup(GeometryName);
 
-					if (State.GetViewMode().GetName() == FClothRenderViewMode::Name)
+					if (RenderMaterialIDs.Num() > 0)
 					{
 						TArray<FString> MaterialArray;
 						MaterialArray.Reserve(ClothFacade.GetRenderMaterialSoftObjectPathName().Num());
diff --git a/Engine/Plugins/Experimental/Dataflow/Source/DataflowEditor/Private/Dataflow/DataflowRenderingFactory.cpp b/Engine/Plugins/Experimental/Dataflow/Source/DataflowEditor/Private/Dataflow/DataflowRenderingFactory.cpp
--- a/Engine/Plugins/Experimental/Dataflow/Source/DataflowEditor/Private/Dataflow/DataflowRenderingFactory.cpp
+++ b/Engine/Plugins/Experimental/Dataflow/Source/DataflowEditor/Private/Dataflow/DataflowRenderingFactory.cpp
@@ -31,7 +31,7 @@ namespace UE::Dataflow
 	{
 		if (CallbackMap.Contains(State.GetRenderKey()))
 		{
-			return CallbackMap[State.GetRenderKey()]->CanRender(State.GetViewMode());
+			return CallbackMap[State.GetRenderKey()]->CanRenderFromState(State);
 		}
 
 		return false;
diff --git a/Engine/Plugins/Experimental/Dataflow/Source/DataflowEditor/Public/Dataflow/DataflowRenderingFactory.h b/Engine/Plugins/Experimental/Dataflow/Source/DataflowEditor/Public/Dataflow/DataflowRenderingFactory.h
--- a/Engine/Plugins/Experimental/Dataflow/Source/DataflowEditor/Public/Dataflow/DataflowRenderingFactory.h
+++ b/Engine/Plugins/Experimental/Dataflow/Source/DataflowEditor/Public/Dataflow/DataflowRenderingFactory.h
@@ -84,6 +84,7 @@ namespace UE::Dataflow
 			virtual ~ICallbackInterface() = default;
 			virtual FRenderKey GetRenderKey() const = 0;
 			virtual bool CanRender(const UE::Dataflow::IDataflowConstructionViewMode& ViewMode) const = 0;
+			virtual bool CanRenderFromState(const FGraphRenderingState& State) const { return CanRender(State.GetViewMode()); }
 			virtual void Render(GeometryCollection::Facades::FRenderingFacade& RenderingFacade, const FGraphRenderingState& State) = 0;
 		};
 
