Dataflow : fix cloth transfer tool to better select the view mode based on the node parameters

diff --git a/Engine/Plugins/ChaosClothAssetEditor/Source/ChaosClothAssetEditorTools/Private/ChaosClothAsset/ClothEditorToolBuilders.cpp b/Engine/Plugins/ChaosClothAssetEditor/Source/ChaosClothAssetEditorTools/Private/ChaosClothAsset/ClothEditorToolBuilders.cpp
--- a/Engine/Plugins/ChaosClothAssetEditor/Source/ChaosClothAssetEditorTools/Private/ChaosClothAsset/ClothEditorToolBuilders.cpp
+++ b/Engine/Plugins/ChaosClothAssetEditor/Source/ChaosClothAssetEditorTools/Private/ChaosClothAsset/ClothEditorToolBuilders.cpp
@@ -3,6 +3,7 @@
 #include "ChaosClothAsset/ClothEditorToolBuilders.h"
 #include "ChaosClothAsset/ClothCollectionGroup.h"
 #include "ChaosClothAsset/ClothPatternVertexType.h"
+#include "ChaosClothAsset/TransferSkinWeightsNode.h"
 #include "TargetInterfaces/PrimitiveComponentBackedTarget.h"
 #include "ChaosClothAsset/SelectionNode.h"
 #include "ChaosClothAsset/WeightMapNode.h"
@@ -286,7 +287,27 @@ void UClothTransferSkinWeightsToolBuilder::GetSupportedConstructionViewModes(con
 	using namespace UE::Chaos::ClothAsset;
 	const UE::Dataflow::FRenderingViewModeFactory& Factory = UE::Dataflow::FRenderingViewModeFactory::GetInstance();
 	const UE::Dataflow::IDataflowConstructionViewMode* const Sim3DMode = Factory.GetViewMode(ClothViewModeToDataflowViewModeName(EClothPatternVertexType::Sim3D));
-	Modes.Add(Sim3DMode);
+	const UE::Dataflow::IDataflowConstructionViewMode* const RenderMode = Factory.GetViewMode(ClothViewModeToDataflowViewModeName(EClothPatternVertexType::Render));
+
+	const FChaosClothAssetTransferSkinWeightsNode* const TransferNode = ContextObject.GetSelectedNodeOfType<FChaosClothAssetTransferSkinWeightsNode>();
+	if (TransferNode)
+	{
+		switch (TransferNode->TargetMeshType)
+		{
+		case EChaosClothAssetTransferTargetMeshType::All:
+			Modes.Add(Sim3DMode);
+			Modes.Add(RenderMode);
+			break;
+
+		case EChaosClothAssetTransferTargetMeshType::Simulation:
+			Modes.Add(Sim3DMode);
+			break;
+
+		case EChaosClothAssetTransferTargetMeshType::Render:
+			Modes.Add(RenderMode);
+			break;
+		}
+	}
 }
 
 bool UClothTransferSkinWeightsToolBuilder::CanSceneStateChange(const UInteractiveTool* ActiveTool, const FToolBuilderState& SceneState) const
@@ -302,7 +323,25 @@ void UClothTransferSkinWeightsToolBuilder::SceneStateChanged(UInteractiveTool* A
 
 void UClothTransferSkinWeightsToolBuilder::GetSupportedViewModes(const UDataflowContextObject& ContextObject, TArray<UE::Chaos::ClothAsset::EClothPatternVertexType>& Modes) const
 {
-	Modes.Add(UE::Chaos::ClothAsset::EClothPatternVertexType::Sim3D);
+	const FChaosClothAssetTransferSkinWeightsNode* const TransferNode = ContextObject.GetSelectedNodeOfType<FChaosClothAssetTransferSkinWeightsNode>();
+	if (TransferNode)
+	{
+		switch (TransferNode->TargetMeshType)
+		{
+		case EChaosClothAssetTransferTargetMeshType::All:
+			Modes.Add(UE::Chaos::ClothAsset::EClothPatternVertexType::Sim3D);
+			Modes.Add(UE::Chaos::ClothAsset::EClothPatternVertexType::Render);
+			break;
+
+		case EChaosClothAssetTransferTargetMeshType::Simulation:
+			Modes.Add(UE::Chaos::ClothAsset::EClothPatternVertexType::Sim3D);
+			break;
+
+		case EChaosClothAssetTransferTargetMeshType::Render:
+			Modes.Add(UE::Chaos::ClothAsset::EClothPatternVertexType::Render);
+			break;
+		}
+	}
 }
 
 USingleSelectionMeshEditingTool* UClothTransferSkinWeightsToolBuilder::CreateNewTool(const FToolBuilderState& SceneState) const
