Make the code to toggle between Cloth panel Editor and dataflow editor for cloth asset symmetrical
- set the cloth panel editor to be default  for 5.7

diff --git a/Engine/Plugins/ChaosClothAssetEditor/Source/ChaosClothAssetEditor/Private/ChaosClothAsset/AssetDefinition_ClothAsset.cpp b/Engine/Plugins/ChaosClothAssetEditor/Source/ChaosClothAssetEditor/Private/ChaosClothAsset/AssetDefinition_ClothAsset.cpp
--- a/Engine/Plugins/ChaosClothAssetEditor/Source/ChaosClothAssetEditor/Private/ChaosClothAsset/AssetDefinition_ClothAsset.cpp
+++ b/Engine/Plugins/ChaosClothAssetEditor/Source/ChaosClothAssetEditor/Private/ChaosClothAsset/AssetDefinition_ClothAsset.cpp
@@ -27,8 +27,8 @@
 
 namespace UE::Chaos::ClothAsset::Private
 {
-static int32 bEnableLegacyClothPanelEditor = false;
-static FAutoConsoleVariableRef CVarEnableLegacyClothPanelEditor(TEXT("p.ChaosCloth.EnableLegacyClothPanelEditor"), bEnableLegacyClothPanelEditor, TEXT("Enable the use of the legacy Cloth Panel Editor instead of the unified Dataflow one for editing cloth assets"));
+static bool bUseClothPanelEditorByDefault = true;
+static FAutoConsoleVariableRef CVarUseClothPanelEditorByDefault(TEXT("p.ChaosCloth.UseClothPanelEditorByDefault"), bUseClothPanelEditorByDefault, TEXT("Enable the use of the Cloth Panel Editor instead of the unified Dataflow one for editing cloth assets"));
 }
 
 namespace ClothAssetDefinitionHelpers
@@ -135,7 +135,6 @@ namespace ClothAssetDefinitionHelpers
 	}
 }
 
-
 // Create a new UDataflow if one doesn't already exist for the Cloth Asset
 UObject* UAssetDefinition_ClothAsset::NewOrOpenDataflowAsset(const UChaosClothAsset* ClothAsset)
 {
@@ -181,28 +180,66 @@ EAssetCommandResult UAssetDefinition_ClothAsset::OpenAssets(const FAssetOpenArgs
 
 	// For now the cloth editor only works on one asset at a time
 	ensure(ClothObjects.Num() == 0 || ClothObjects.Num() == 1);
-
 	if (ClothObjects.Num() > 0)
 	{
-		UAssetEditorSubsystem* const AssetEditorSubsystem = GEditor->GetEditorSubsystem<UAssetEditorSubsystem>();
-		if(UE::Chaos::ClothAsset::Private::bEnableLegacyClothPanelEditor)
+		if (ClothObjects[0])
 		{
-			UChaosClothAssetEditor* const AssetEditor = NewObject<UChaosClothAssetEditor>(AssetEditorSubsystem, NAME_None, RF_Transient);
-			AssetEditor->Initialize({ClothObjects[0] });
-		}
-		else
-		{
-			// Check if the cloth asset has a Dataflow asset
-			UChaosClothAsset* const ClothAsset = CastChecked<UChaosClothAsset>(ClothObjects[0]);
-			if (!ClothAsset->GetDataflow())
+			bool bSuccess = false;
+			if (UAssetDefinition_ClothAsset::UseClothPanelEditorByDefault())
 			{
-				// No Dataflow asset, just open a basic properties panel
-				FSimpleAssetEditor::CreateEditor(EToolkitMode::Standalone, OpenArgs.ToolkitHost, ClothObjects[0] );
+				bSuccess = UAssetDefinition_ClothAsset::LaunchClothPanelAssetEditor(ClothObjects[0]);
 			}
 			else
 			{
-				UDataflowEditor* const AssetEditor = NewObject<UDataflowEditor>(AssetEditorSubsystem, NAME_None, RF_Transient);
+				bSuccess = UAssetDefinition_ClothAsset::LaunchClothDataflowAssetEditor(ClothObjects[0]);
+			}
+			if (!bSuccess)
+			{
+				// fallback to basic property panel editor
+				FSimpleAssetEditor::CreateEditor(EToolkitMode::Standalone, OpenArgs.ToolkitHost, ClothObjects[0]);
+			}
+
+			return EAssetCommandResult::Handled;
+		}
+	}
+	return EAssetCommandResult::Unhandled;
+}
+
+bool UAssetDefinition_ClothAsset::UseClothPanelEditorByDefault()
+{
+	return UE::Chaos::ClothAsset::Private::bUseClothPanelEditorByDefault;
+}
+
+bool UAssetDefinition_ClothAsset::LaunchClothPanelAssetEditor(UChaosClothAsset* InClothAsset)
+{
+	if (InClothAsset == nullptr)
+	{
+		return false;
+	};
+	if (UAssetEditorSubsystem* const AssetEditorSubsystem = GEditor->GetEditorSubsystem<UAssetEditorSubsystem>())
+	{
+		if (UChaosClothAssetEditor* const AssetEditor = NewObject<UChaosClothAssetEditor>(AssetEditorSubsystem, NAME_None, RF_Transient))
+		{
+			AssetEditor->Initialize({ InClothAsset });
+			return true;
+		}
+	}
+	return false;
+}
 
+bool UAssetDefinition_ClothAsset::LaunchClothDataflowAssetEditor(UChaosClothAsset* InClothAsset)
+{
+	if (InClothAsset == nullptr)
+	{
+		return false;
+	};
+	// Check if the cloth asset has a Dataflow asset
+	if (InClothAsset->GetDataflow())
+	{
+		if (UAssetEditorSubsystem* const AssetEditorSubsystem = GEditor->GetEditorSubsystem<UAssetEditorSubsystem>())
+		{
+			if (UDataflowEditor* const AssetEditor = NewObject<UDataflowEditor>(AssetEditorSubsystem, NAME_None, RF_Transient))
+			{
 				if (UDataflowSimulationSettings* SimulationSettings = NewObject<UDataflowSimulationSettings>())
 				{
 					SimulationSettings->bIsSimulationPlayingByDefault = true;
@@ -219,16 +256,14 @@ EAssetCommandResult UAssetDefinition_ClothAsset::OpenAssets(const FAssetOpenArgs
 
 				const TSubclassOf<AActor> ActorClass = StaticLoadClass(AActor::StaticClass(), nullptr,
 					TEXT("/ChaosClothAssetEditor/BP_ClothPreview.BP_ClothPreview_C"), nullptr, LOAD_None, nullptr);
-				AssetEditor->RegisterToolCategories({"General", "Cloth"});
-				AssetEditor->Initialize({ ClothObjects[0] }, ActorClass);
+				AssetEditor->RegisterToolCategories({ "General", "Cloth" });
+				AssetEditor->Initialize({ InClothAsset }, ActorClass);
+
+				return true;
 			}
 		}
-
-		return EAssetCommandResult::Handled;
 	}
-
-	return EAssetCommandResult::Unhandled;
+	return false;
 }
 
-
 #undef LOCTEXT_NAMESPACE
diff --git a/Engine/Plugins/ChaosClothAssetEditor/Source/ChaosClothAssetEditor/Private/ChaosClothAsset/AssetDefinition_ClothAsset.h b/Engine/Plugins/ChaosClothAssetEditor/Source/ChaosClothAssetEditor/Private/ChaosClothAsset/AssetDefinition_ClothAsset.h
--- a/Engine/Plugins/ChaosClothAssetEditor/Source/ChaosClothAssetEditor/Private/ChaosClothAsset/AssetDefinition_ClothAsset.h
+++ b/Engine/Plugins/ChaosClothAssetEditor/Source/ChaosClothAssetEditor/Private/ChaosClothAsset/AssetDefinition_ClothAsset.h
@@ -17,6 +17,10 @@ public:
 	UE_DEPRECATED(5.5, "Dataflow assets are now created through the physics menu.")
 	static UObject* NewOrOpenDataflowAsset(const UChaosClothAsset* ClothAsset);
 
+	static bool LaunchClothPanelAssetEditor(UChaosClothAsset* ClothAsset);
+	static bool LaunchClothDataflowAssetEditor(UChaosClothAsset* ClothAsset);
+	static bool UseClothPanelEditorByDefault();
+	
 private:
 
 	virtual FText GetAssetDisplayName() const override;
diff --git a/Engine/Plugins/ChaosClothAssetEditor/Source/ChaosClothAssetEditor/Private/ChaosClothAsset/ChaosClothAssetEditorModule.cpp b/Engine/Plugins/ChaosClothAssetEditor/Source/ChaosClothAssetEditor/Private/ChaosClothAsset/ChaosClothAssetEditorModule.cpp
--- a/Engine/Plugins/ChaosClothAssetEditor/Source/ChaosClothAssetEditor/Private/ChaosClothAsset/ChaosClothAssetEditorModule.cpp
+++ b/Engine/Plugins/ChaosClothAssetEditor/Source/ChaosClothAssetEditor/Private/ChaosClothAsset/ChaosClothAssetEditorModule.cpp
@@ -69,7 +69,7 @@ namespace UE::Chaos::ClothAsset
 
 		UToolMenu* const ClothContextMenu = UToolMenus::Get()->ExtendMenu("ContentBrowser.AssetContextMenu.ChaosClothAsset");
 		FToolMenuSection& Section = ClothContextMenu->FindOrAddSection("GetAssetActions");
-		Section.AddDynamicEntry("OpenInDataflowEditor", FNewToolMenuSectionDelegate::CreateLambda([](FToolMenuSection& Section)
+		Section.AddDynamicEntry("OpenInAlternateEditor", FNewToolMenuSectionDelegate::CreateLambda([](FToolMenuSection& Section)
 		{
 			// We'll need to get the target assets out of the context
 			if (const UContentBrowserAssetContextMenuContext* const Context = Section.FindContext<UContentBrowserAssetContextMenuContext>())
@@ -88,21 +88,42 @@ namespace UE::Chaos::ClothAsset
 
 					const TSharedPtr<class FUICommandList> CommandListToBind = MakeShared<FUICommandList>();
 
+					const bool bLaunchDataflowEditorAsAlternative = UAssetDefinition_ClothAsset::UseClothPanelEditorByDefault();
+
+					TSharedPtr<FUICommandInfo> LaunchCommand =
+						bLaunchDataflowEditorAsAlternative
+						? FChaosClothAssetEditorCommands::Get().OpenClothAssetInDataflowEditor
+						: FChaosClothAssetEditorCommands::Get().OpenClothAssetInClothPanelEditor
+						;
+
 					CommandListToBind->MapAction(
-						FChaosClothAssetEditorCommands::Get().OpenClothAssetInLegacyClothPanelEditor,
-						FExecuteAction::CreateWeakLambda(DataflowEditorSubsystem, [Context, DataflowEditorSubsystem]()
-						{
-							// When we actually do want to open the editor, trigger the load to get the objects
-							TArray<TObjectPtr<UObject>> AssetsToEdit;
-							AssetsToEdit.Append(Context->LoadSelectedObjects<UObject>());
-
-							UChaosClothAssetEditor* const AssetEditor = NewObject<UChaosClothAssetEditor>(DataflowEditorSubsystem, NAME_None, RF_Transient);
-							AssetEditor->Initialize({ AssetsToEdit[0] });
-						}),
+						LaunchCommand,
+						FExecuteAction::CreateWeakLambda(DataflowEditorSubsystem, 
+							[Context, DataflowEditorSubsystem, bLaunchDataflowEditorAsAlternative]()
+							{
+								// When we actually do want to open the editor, trigger the load to get the objects
+								TArray<TObjectPtr<UObject>> AssetsToEdit;
+								AssetsToEdit.Append(Context->LoadSelectedObjects<UObject>());
+								if (AssetsToEdit.Num() > 0)
+								{
+									if (UChaosClothAsset* const ClothAsset = CastChecked<UChaosClothAsset>(AssetsToEdit[0]))
+									{
+										if (bLaunchDataflowEditorAsAlternative)
+										{
+											UAssetDefinition_ClothAsset::LaunchClothDataflowAssetEditor(ClothAsset);
+										}
+										else
+										{
+											UAssetDefinition_ClothAsset::LaunchClothPanelAssetEditor(ClothAsset);
+										}
+									}
+								}
+							}),
 						FCanExecuteAction::CreateWeakLambda(Context, [Context]() { return Context->bCanBeModified; }));
 
 					const TAttribute<FText> ToolTipOverride = Context->bCanBeModified ? TAttribute<FText>() : LOCTEXT("ReadOnlyAssetWarning", "The selected asset(s) are read-only and cannot be edited.");
-					Section.AddMenuEntryWithCommandList(FChaosClothAssetEditorCommands::Get().OpenClothAssetInLegacyClothPanelEditor,
+					Section.AddMenuEntryWithCommandList(
+						LaunchCommand,
 						CommandListToBind,
 						TAttribute<FText>(),
 						ToolTipOverride,
diff --git a/Engine/Plugins/ChaosClothAssetEditor/Source/ChaosClothAssetEditor/Private/ChaosClothAsset/ClothEditorCommands.cpp b/Engine/Plugins/ChaosClothAssetEditor/Source/ChaosClothAssetEditor/Private/ChaosClothAsset/ClothEditorCommands.cpp
--- a/Engine/Plugins/ChaosClothAssetEditor/Source/ChaosClothAssetEditor/Private/ChaosClothAsset/ClothEditorCommands.cpp
+++ b/Engine/Plugins/ChaosClothAssetEditor/Source/ChaosClothAssetEditor/Private/ChaosClothAsset/ClothEditorCommands.cpp
@@ -40,7 +40,8 @@ void FChaosClothAssetEditorCommands::RegisterCommands()
 
 	UI_COMMAND(OpenClothEditor, "Cloth Editor", "Open the Cloth Editor window", EUserInterfaceActionType::Button, FInputChord());
 
-	UI_COMMAND(OpenClothAssetInLegacyClothPanelEditor, "Open in legacy Cloth Panel Editor", "Open the Cloth Asset in the legacy Cloth Panel Editor instead of the unified Dataflow one", EUserInterfaceActionType::Button, FInputChord());
+	UI_COMMAND(OpenClothAssetInClothPanelEditor, "Open in Cloth Panel Editor", "Open the Cloth Asset in the Cloth Panel Editor instead of the unified Dataflow one", EUserInterfaceActionType::Button, FInputChord());
+	UI_COMMAND(OpenClothAssetInDataflowEditor, "Open in Dataflow Editor", "Open the Cloth Asset in the Unified Dataflow Editor instead of the Cloth Panel one", EUserInterfaceActionType::Button, FInputChord());
 
 	UI_COMMAND(BeginRemeshTool, "Remesh", "Remesh the selected mesh", EUserInterfaceActionType::Button, FInputChord());
 	UI_COMMAND(BeginAttributeEditorTool, "AttrEd", "Edit/configure mesh attributes", EUserInterfaceActionType::Button, FInputChord());
diff --git a/Engine/Plugins/ChaosClothAssetEditor/Source/ChaosClothAssetEditor/Private/ChaosClothAsset/ClothEditorCommands.h b/Engine/Plugins/ChaosClothAssetEditor/Source/ChaosClothAssetEditor/Private/ChaosClothAsset/ClothEditorCommands.h
--- a/Engine/Plugins/ChaosClothAssetEditor/Source/ChaosClothAssetEditor/Private/ChaosClothAsset/ClothEditorCommands.h
+++ b/Engine/Plugins/ChaosClothAssetEditor/Source/ChaosClothAssetEditor/Private/ChaosClothAsset/ClothEditorCommands.h
@@ -30,7 +30,8 @@ public:
 
 	TSharedPtr<FUICommandInfo> OpenClothEditor;
 
-	TSharedPtr<FUICommandInfo> OpenClothAssetInLegacyClothPanelEditor;
+	TSharedPtr<FUICommandInfo> OpenClothAssetInClothPanelEditor;
+	TSharedPtr<FUICommandInfo> OpenClothAssetInDataflowEditor;
 
 	const static FString BeginRemeshToolIdentifier;
 	TSharedPtr<FUICommandInfo> BeginRemeshTool;
