Dataflow : cloth make sure we instantly update the viewport when updating cloth visualization menu

diff --git a/Engine/Plugins/ChaosClothAssetEditor/Source/ChaosClothAssetEditor/Private/ChaosClothAsset/ClothEditorSimulationVisualization.cpp b/Engine/Plugins/ChaosClothAssetEditor/Source/ChaosClothAssetEditor/Private/ChaosClothAsset/ClothEditorSimulationVisualization.cpp
--- a/Engine/Plugins/ChaosClothAssetEditor/Source/ChaosClothAssetEditor/Private/ChaosClothAsset/ClothEditorSimulationVisualization.cpp
+++ b/Engine/Plugins/ChaosClothAssetEditor/Source/ChaosClothAssetEditor/Private/ChaosClothAsset/ClothEditorSimulationVisualization.cpp
@@ -16,6 +16,7 @@
 #include "Widgets/Input/SSpinBox.h"
 #include "Widgets/Input/STextComboBox.h"
 #include "Types/SlateEnums.h"
+#include "EditorViewportClient.h"
 
 
 #define LOCTEXT_NAMESPACE "ClothEditorSimulationVisualization"
@@ -501,14 +502,14 @@ FClothEditorSimulationVisualization::FClothEditorSimulationVisualization()
 }
 
 
-void FClothEditorSimulationVisualization::ExtendViewportShowMenu(FMenuBuilder& MenuBuilder, const TFunction<UChaosClothComponent*()>& GetClothComponentFunc)
+void FClothEditorSimulationVisualization::ExtendViewportShowMenu(FMenuBuilder& MenuBuilder, FEditorViewportClient& ViewportClient, const TFunction<UChaosClothComponent*()>& GetClothComponentFunc)
 {
 	MenuBuilder.BeginSection(TEXT("ChaosSimulation_Visualizations"), LOCTEXT("ClothVisualizationSection", "Chaos Cloth Visualization"));
 	{
 		for (uint32 OptionIndex = 0; OptionIndex < Private::FVisualizationOption::Count; ++OptionIndex)
 		{
 			// Handler for visualization entry being clicked
-			const FExecuteAction ExecuteAction = FExecuteAction::CreateLambda([this, OptionIndex, GetClothComponentFunc]()
+			const FExecuteAction ExecuteAction = FExecuteAction::CreateLambda([this, OptionIndex, ViewportClientPtr = &ViewportClient, GetClothComponentFunc]()
 			{
 				Flags[OptionIndex] = !Flags[OptionIndex];
 
@@ -527,6 +528,13 @@ void FClothEditorSimulationVisualization::ExtendViewportShowMenu(FMenuBuilder& M
 						const bool bIsClothSectionsVisible = !Flags[OptionIndex];
 						ShowClothSections(ClothComponent, bIsClothSectionsVisible);
 					}
+
+					// refresh the view
+					if (ViewportClientPtr)
+					{
+						ViewportClientPtr->Invalidate();
+					}
+
 				}
 			});
 
@@ -551,7 +559,7 @@ void FClothEditorSimulationVisualization::ExtendViewportShowMenu(FMenuBuilder& M
 {
 	TWeakPtr<FChaosClothAssetEditor3DViewportClient> ViewportClientWeakPtr(ViewportClient);
 
-	ExtendViewportShowMenu(MenuBuilder, [ViewportClientWeakPtr]() -> UChaosClothComponent*
+	ExtendViewportShowMenu(MenuBuilder, ViewportClient.Get(), [ViewportClientWeakPtr]() -> UChaosClothComponent*
 	{
 		if (TSharedPtr<FChaosClothAssetEditor3DViewportClient> ViewportClient = ViewportClientWeakPtr.Pin())
 		{
@@ -567,7 +575,7 @@ void FClothEditorSimulationVisualization::ExtendViewportShowMenu(FMenuBuilder& M
 {
 	TWeakPtr<FDataflowSimulationViewportClient> ViewportClientWeakPtr(ViewportClient);
 
-	ExtendViewportShowMenu(MenuBuilder, [ViewportClientWeakPtr]() -> UChaosClothComponent*
+	ExtendViewportShowMenu(MenuBuilder, ViewportClient.Get(), [ViewportClientWeakPtr]() -> UChaosClothComponent*
 	{
 		if (TSharedPtr<FDataflowSimulationViewportClient> ViewportClient = ViewportClientWeakPtr.Pin())
 		{
diff --git a/Engine/Plugins/ChaosClothAssetEditor/Source/ChaosClothAssetEditor/Private/ChaosClothAsset/ClothEditorSimulationVisualization.h b/Engine/Plugins/ChaosClothAssetEditor/Source/ChaosClothAssetEditor/Private/ChaosClothAsset/ClothEditorSimulationVisualization.h
--- a/Engine/Plugins/ChaosClothAssetEditor/Source/ChaosClothAssetEditor/Private/ChaosClothAsset/ClothEditorSimulationVisualization.h
+++ b/Engine/Plugins/ChaosClothAssetEditor/Source/ChaosClothAssetEditor/Private/ChaosClothAsset/ClothEditorSimulationVisualization.h
@@ -8,6 +8,7 @@ class FMenuBuilder;
 class UChaosClothComponent;
 class FPrimitiveDrawInterface;
 class FCanvas;
+class FEditorViewportClient;
 class FSceneView;
 class STextComboBox;
 class FDataflowSimulationViewportClient;
@@ -89,7 +90,7 @@ private:
 	};
 
 
-	void ExtendViewportShowMenu(FMenuBuilder& MenuBuilder, const TFunction<UChaosClothComponent*()>& GetClothComponentFunc);
+	void ExtendViewportShowMenu(FMenuBuilder& MenuBuilder, FEditorViewportClient& ViewportClient, const TFunction<UChaosClothComponent*()>& GetClothComponentFunc);
 
 	/** Return whether or not - given the current enabled options - the simulation should be disabled. */
 	bool ShouldDisableSimulation() const;
