Chaos Cloth Asset - Enabled a complete rebuild of the Cloth Asset on initial evaluation or after a forced re-evaluation clears the context cache.
Charles suggested this change so that it better matches the rest of the engine when opening the asset for the first time.

#rnx

diff --git a/Engine/Plugins/ChaosClothAssetEditor/Source/ChaosClothAssetDataflowNodes/Private/ChaosClothAsset/TerminalNode.cpp b/Engine/Plugins/ChaosClothAssetEditor/Source/ChaosClothAssetDataflowNodes/Private/ChaosClothAsset/TerminalNode.cpp
--- a/Engine/Plugins/ChaosClothAssetEditor/Source/ChaosClothAssetDataflowNodes/Private/ChaosClothAsset/TerminalNode.cpp
+++ b/Engine/Plugins/ChaosClothAssetEditor/Source/ChaosClothAssetDataflowNodes/Private/ChaosClothAsset/TerminalNode.cpp
@@ -111,6 +111,12 @@ void FChaosClothAssetTerminalNode_v2::SetAssetValue(TObjectPtr<UObject> Asset, U
 	{
 		using namespace UE::Chaos::ClothAsset;
 
+		// Check whether the context has just been cleared up, in which case force a rebuild by invalidating the checksum
+		if (Context.IsEmpty())
+		{
+			bClothCollectionChecksumValid = false;
+		}
+
 		TArray<TSharedRef<const FManagedArrayCollection>> InClothCollections = GetCleanedCollectionLodValues(Context);
 		TArray<TSharedRef<const FManagedArrayCollection>> ClothCollections = static_cast<const UChaosClothAsset*>(ClothAsset)->GetClothCollections();
 
