Chaos Cloth Asset:
- Updated the USD cloth import v2 to read fabric properties.
- Fixed USD import nodes broken materials caused by the missing soft object path fixup code.

[RN] Chaos Cloth Asset - Updated the USD cloth import v2 to read fabric properties.

diff --git a/Engine/Plugins/ChaosClothAssetEditor/Source/ChaosClothAssetDataflowNodes/Private/ChaosClothAsset/USDImportNode.cpp b/Engine/Plugins/ChaosClothAssetEditor/Source/ChaosClothAssetDataflowNodes/Private/ChaosClothAsset/USDImportNode.cpp
--- a/Engine/Plugins/ChaosClothAssetEditor/Source/ChaosClothAssetDataflowNodes/Private/ChaosClothAsset/USDImportNode.cpp
+++ b/Engine/Plugins/ChaosClothAssetEditor/Source/ChaosClothAssetDataflowNodes/Private/ChaosClothAsset/USDImportNode.cpp
@@ -171,6 +171,10 @@ void FChaosClothAssetUSDImportNode::Serialize(FArchive& Ar)
 		{
 			ClothFacade.DefineSchema();
 		}
+
+		// Also apply any required fixup (e.g. soft object path names)
+		ClothFacade.PostSerialize(Ar);
+
 		CollectionCache = MoveTemp(*ClothCollection);
 
 		if (Ar.CustomVer(FFortniteMainBranchObjectVersion::GUID) < FFortniteMainBranchObjectVersion::ChaosClothAssetUSDImportNodeAddAssetDependencies)
diff --git a/Engine/Plugins/ChaosClothAssetEditor/Source/ChaosClothAssetDataflowNodes/Private/ChaosClothAsset/USDImportNode_v2.cpp b/Engine/Plugins/ChaosClothAssetEditor/Source/ChaosClothAssetDataflowNodes/Private/ChaosClothAsset/USDImportNode_v2.cpp
--- a/Engine/Plugins/ChaosClothAssetEditor/Source/ChaosClothAssetDataflowNodes/Private/ChaosClothAsset/USDImportNode_v2.cpp
+++ b/Engine/Plugins/ChaosClothAssetEditor/Source/ChaosClothAssetDataflowNodes/Private/ChaosClothAsset/USDImportNode_v2.cpp
@@ -9,20 +9,17 @@
 #include "ChaosClothAsset/CollectionClothSelectionFacade.h"
 #include "ChaosClothAsset/RenderMeshImport.h"
 #include "ChaosClothAsset/USDImportNode.h"
+#include "ChaosClothAsset/UsdPrimAttributeAccessor.h"
 #include "AssetRegistry/AssetData.h"
 #include "AssetRegistry/AssetRegistryModule.h"
 #include "Chaos/CollectionPropertyFacade.h"
 #include "Dataflow/DataflowInputOutput.h"
-#include "Dataflow/DataflowGraph.h"
-#include "Dataflow/DataflowObject.h"
 #include "Engine/StaticMesh.h"
-#include "HAL/FileManager.h"
 #include "Materials/Material.h"
-#include "Materials/MaterialInstanceConstant.h"
 #include "Misc/Paths.h"
 #include "Misc/ScopedSlowTask.h"
-#include "Rendering/SkeletalMeshLODModel.h"
 #include "StaticMeshAttributes.h"
+#include "UObject/UE5ReleaseStreamObjectVersion.h"
 #include "UObject/Package.h"
 #include "UnrealUSDWrapper.h"
 #include "USDAssetUserData.h"
@@ -32,7 +29,6 @@
 #include "USDStageImporter.h"
 #include "USDStageImportOptions.h"
 #include "USDTypesConversion.h"
-#include "USDValueConversion.h"
 #include "UsdWrappers/UsdPrim.h"
 #include "UsdWrappers/UsdGeomSubset.h"
 #include "UsdWrappers/UsdRelationship.h"
@@ -54,6 +50,7 @@ namespace UE::Chaos::ClothAsset::Private
 	static const FName SimMeshDataAPI(TEXT("SimMeshDataAPI"));
 	static const FName SimPatternAPI(TEXT("SimPatternAPI"));
 	static const FName SewingAPI(TEXT("SewingAPI"));
+	static const FName CloSolverPropertiesAPI(TEXT("CloSolverPropertiesAPI"));
 
 	// USD import material overrides
 	static TArray<FSoftObjectPath> UsdClothOverrideMaterials_v2 =
@@ -64,7 +61,7 @@ namespace UE::Chaos::ClothAsset::Private
 		FSoftObjectPath(TEXT("/ChaosClothAsset/Materials/USDImportTranslucentTwoSidedMaterial.USDImportTranslucentTwoSidedMaterial")),
 		FSoftObjectPath(TEXT("/ChaosClothAsset/Materials/USDImportDisplayColorMaterial.USDImportDisplayColorMaterial")),
 	};
-	
+
 	// Return the specified UObject's dependencies for top level UAssets that are not in the engine folders
 	static TArray<UObject*> GetAssetDependencies(const UObject* Asset)
 	{
@@ -92,6 +89,40 @@ namespace UE::Chaos::ClothAsset::Private
 		return Dependencies;
 	}
 
+	static TArray<TObjectPtr<UObject>> GetImportedAssetDependencies(const UObject* StaticMesh)
+	{
+		using namespace UE::Chaos::ClothAsset::Private;
+
+		TSet<TObjectPtr<UObject>> ImportedAssets;
+		if (StaticMesh)
+		{
+			TQueue<const UObject*> AssetsToVisit;
+			AssetsToVisit.Enqueue(StaticMesh);
+
+			const UObject* VisitedAsset;
+			while (AssetsToVisit.Dequeue(VisitedAsset))
+			{
+				const FName VisitedAssetPackageName(VisitedAsset->GetPackage()->GetName());
+				const TArray<UObject*> AssetDependencies = GetAssetDependencies(VisitedAsset);
+
+				UE_CLOG(AssetDependencies.Num(), LogChaosClothAssetDataflowNodes, Verbose, TEXT("Dependencies for Object %s - %s:"), *VisitedAsset->GetName(), *VisitedAssetPackageName.ToString());
+				for (UObject* const AssetDependency : AssetDependencies)
+				{
+					if (!ImportedAssets.Contains(AssetDependency))
+					{
+						// Add the dependency
+						UE_LOG(LogChaosClothAssetDataflowNodes, Verbose, TEXT("Found %s"), *AssetDependency->GetPackage()->GetName());
+						ImportedAssets.Emplace(AssetDependency);
+
+						// Visit this asset too
+						AssetsToVisit.Enqueue(AssetDependency);
+					}
+				}
+			}
+		}
+		return ImportedAssets.Array();
+	}
+
 	static void OverrideUsdImportMaterials_v2(const TArray<FSoftObjectPath>& Materials, TArray<FSoftObjectPath>* SavedValues = nullptr)
 	{
 		if (UUsdProjectSettings* UsdProjectSettings = GetMutableDefault<UUsdProjectSettings>())
@@ -175,11 +206,30 @@ namespace UE::Chaos::ClothAsset::Private
 		return ImportContext.ImportedAssets;
 	}
 
-	static FUsdPrim FindClothPrim(const FUsdPrim& RootPrim)
+	static TObjectPtr<UStaticMesh> FindImportedStaticMesh(const TArrayView<TObjectPtr<UObject>> ImportedAssets, const FString& MeshPrimPath)
+	{
+		for (TObjectPtr<UObject>& ImportedAsset : ImportedAssets)
+		{
+			if (const TObjectPtr<UStaticMesh> ImportedStaticMesh = Cast<UStaticMesh>(ImportedAsset))
+			{
+				if (const UUsdMeshAssetUserData* const AssetUserData = Cast<UUsdMeshAssetUserData>(
+					ImportedStaticMesh->GetAssetUserDataOfClass(UUsdMeshAssetUserData::StaticClass())))
+				{
+					if (AssetUserData->PrimPaths.Find(MeshPrimPath) != INDEX_NONE)
+					{
+						return ImportedStaticMesh;
+					}
+				}
+			}
+		}
+		return nullptr;
+	}
+
+	static FUsdPrim FindAPIPrim(const FUsdPrim& RootPrim, const FName APIName)
 	{
 		for (FUsdPrim& ChildPrim : RootPrim.GetChildren())
 		{
-			if (ChildPrim.HasAPI(ClothRootAPI))
+			if (ChildPrim.HasAPI(APIName))
 			{
 				return ChildPrim;
 			}
@@ -244,7 +294,6 @@ namespace UE::Chaos::ClothAsset::Private
 			if (ClothChildPrim.IsA(TEXT("Mesh")))
 			{
 				// Look for all GeomSubsets to see if this is a suitable render mesh prim
-				bool bIsRenderMesh = false;
 				for (FUsdPrim& RenderMeshChildPrim : ClothChildPrim.GetChildren())
 				{
 					if (RenderMeshChildPrim.IsA(TEXT("GeomSubset")) && RenderMeshChildPrim.HasAPI(RenderPatternAPI))
@@ -257,6 +306,29 @@ namespace UE::Chaos::ClothAsset::Private
 		return FUsdPrim();
 	}
 
+	static TArray<FUsdPrim> FindRenderMeshPrims(const FUsdPrim& ClothPrim)
+	{
+		TArray<FUsdPrim> RenderMeshPrims;
+		const TArray<FUsdPrim> ClothChildPrims = ClothPrim.GetChildren();
+		RenderMeshPrims.Reserve(ClothChildPrims.Num());
+		for (const FUsdPrim& ClothChildPrim : ClothChildPrims)
+		{
+			if (ClothChildPrim.IsA(TEXT("Mesh")))
+			{
+				// Look for all GeomSubsets to see if this is a suitable render mesh prim
+				for (const FUsdPrim& RenderMeshChildPrim : ClothChildPrim.GetChildren())
+				{
+					if (RenderMeshChildPrim.IsA(TEXT("GeomSubset")) && RenderMeshChildPrim.HasAPI(RenderPatternAPI))
+					{
+						RenderMeshPrims.Emplace(ClothChildPrim);
+						break;
+					}
+				}
+			}
+		}
+		return RenderMeshPrims;
+	}
+
 	static FVector2f GetSimMeshUVScale(const FUsdPrim& SimMeshPrim)
 	{
 		FVector2f UVScale(1.f);
@@ -339,20 +411,26 @@ namespace UE::Chaos::ClothAsset::Private
 		return false;
 	}
 
-	static bool ImportPatternsFromMeshPrim(const FUsdPrim& MeshPrim, const FName PatternAPI, TMap<FName, TSet<int32>>& Patterns, FText& OutErrorText)
+	static bool GetPatternGeomSubsetsFromMeshPrim(const FUsdPrim& MeshPrim, const FName PatternAPI, TMap<FName, FUsdGeomSubset>& GeomSubsets, FText& OutErrorText)
 	{
-		Patterns.Reset();
+		GeomSubsets.Reset();
 		for (const FUsdPrim& MeshChildPrim : MeshPrim.GetChildren())
 		{
 			if (MeshChildPrim.IsA(TEXT("GeomSubset")) && MeshChildPrim.HasAPI(PatternAPI))
 			{
+				if (GeomSubsets.Contains(MeshChildPrim.GetName()))
+				{
+					OutErrorText = FText::Format(LOCTEXT("DuplicatePatternGeomSubsetName", "Duplicate pattern name for GeomSubset '{0}'. The name needs to be unique."), FText::FromString(MeshChildPrim.GetPrimPath().GetString()));
+					return false;
+				}
+
 				const FUsdGeomSubset GeomSubset(MeshChildPrim);
 
 				// Read FamillyName
 				const FUsdAttribute FamilyNameAttr = GeomSubset.GetFamilyNameAttr();
 				if (GetStringValue(FamilyNameAttr) != TEXT("pattern"))
 				{
-					OutErrorText = FText::Format(LOCTEXT("WrongPatternGeomSubsetFamilyName", "Wrong pattern family name for GeomSubset '{0}'. Needs to be 'pattern'."), FText::FromString(MeshPrim.GetPrimPath().GetString()));
+					OutErrorText = FText::Format(LOCTEXT("WrongPatternGeomSubsetFamilyName", "Wrong pattern family name for GeomSubset '{0}'. Needs to be 'pattern'."), FText::FromString(MeshChildPrim.GetPrimPath().GetString()));
 					return false;
 				}
 
@@ -360,7 +438,7 @@ namespace UE::Chaos::ClothAsset::Private
 				const FUsdAttribute ElementTypeAttr = GeomSubset.GetElementTypeAttr();
 				if (GetStringValue(ElementTypeAttr) != TEXT("face"))
 				{
-					OutErrorText = FText::Format(LOCTEXT("WrongPatternGeomSubsetType", "Wrong pattern type for GeomSubset '{0}'. Needs to be 'face'."), FText::FromString(MeshPrim.GetPrimPath().GetString()));
+					OutErrorText = FText::Format(LOCTEXT("WrongPatternGeomSubsetType", "Wrong pattern type for GeomSubset '{0}'. Needs to be 'face'."), FText::FromString(MeshChildPrim.GetPrimPath().GetString()));
 					return false;
 				}
 
@@ -368,30 +446,63 @@ namespace UE::Chaos::ClothAsset::Private
 				const FUsdAttribute IndicesAttr = GeomSubset.GetIndicesAttr();
 				if (IndicesAttr.GetTypeName() != TEXT("int[]"))
 				{
-					OutErrorText = FText::Format(LOCTEXT("WrongPatternGeomSubsetIndexType", "Wrong pattern index type for GeomSubset '{0}'. Needs to be 'int[]'."), FText::FromString(MeshPrim.GetPrimPath().GetString()));
+					OutErrorText = FText::Format(LOCTEXT("WrongPatternGeomSubsetIndexType", "Wrong pattern index type for GeomSubset '{0}'. Needs to be 'int[]'."), FText::FromString(MeshChildPrim.GetPrimPath().GetString()));
 					return false;
 				}
 
-				if (Patterns.Contains(MeshChildPrim.GetName()))
+				// Make a pattern name that won't collide between sim mesh and render mesh
+				const FName PatternName(MeshChildPrim.GetParent().GetName().ToString() + TEXT("__") + MeshChildPrim.GetName().ToString());
+
+				GeomSubsets.Emplace(PatternName, GeomSubset);
+			}
+		}
+		return true;
+	}
+
+	static bool AddPatternsFromMeshPrim(const FUsdPrim& MeshPrim, const FName PatternAPI, TMap<FName, TSet<int32>>& Patterns, FText& OutErrorText, int32 Offset = 0)
+	{
+		TMap<FName, FUsdGeomSubset> GeomSubsets;
+		if (GetPatternGeomSubsetsFromMeshPrim(MeshPrim, PatternAPI, GeomSubsets, OutErrorText))
+		{
+			for (const TPair<FName, FUsdGeomSubset>& GeomSubset : GeomSubsets)
+			{
+				// Read indices 
+				const FUsdAttribute IndicesAttr = GeomSubset.Value.GetIndicesAttr();
+				if (IndicesAttr.GetTypeName() != TEXT("int[]"))
 				{
-					OutErrorText = FText::Format(LOCTEXT("DuplicatePatternGeomSubsetName", "Duplicate pattern name for GeomSubset '{0}'. The name needs to be unique."), FText::FromString(MeshPrim.GetPrimPath().GetString()));
+					OutErrorText = FText::Format(LOCTEXT("WrongPatternGeomSubsetIndexType", "Wrong pattern index type for GeomSubset '{0}'. Needs to be 'int[]'."), FText::FromString(GeomSubset.Value.GetPrim().GetPrimPath().GetString()));
 					return false;
 				}
-
-				Patterns.Emplace(MeshChildPrim.GetName(), GetIntArrayValues(IndicesAttr));
+				TArray<int32> Indices = GetIntArrayValues(IndicesAttr);
+				if (Offset)
+				{
+					for (int32& Index : Indices)
+					{
+						Index += Offset;
+					}
+				}
+				Patterns.Emplace(GeomSubset.Key, MoveTemp(Indices));
 			}
+			return true;
 		}
-		return true;
+		return false;
 	}
 
-	static bool ImportPatternsFromRenderMeshPrim(const FUsdPrim& RenderMeshPrim, const FUsdPrim& SimMeshPrim, TMap<FName, TSet<int32>>& Patterns, TMap<FName, TSet<FName>>& RenderToSimPatterns, FText& OutErrorText)
+	static bool ImportPatternsFromRenderMeshPrims(const TArray<FUsdPrim>& RenderMeshPrims, const FUsdPrim& SimMeshPrim, TMap<FName, TSet<int32>>& Patterns, TMap<FName, TSet<FName>>& RenderToSimPatterns, FText& OutErrorText)
 	{
+		Patterns.Reset();
 		RenderToSimPatterns.Reset();
 
-		if (ImportPatternsFromMeshPrim(RenderMeshPrim, RenderPatternAPI, Patterns, OutErrorText))
+		const FSdfPath SimMeshPath = SimMeshPrim.GetPrimPath();
+		int32 Offset = 0;
+		for (const FUsdPrim& RenderMeshPrim : RenderMeshPrims)
 		{
-			const FSdfPath SimMeshPath = SimMeshPrim.GetPrimPath();
-
+			// Import the patterns indices
+			if (!AddPatternsFromMeshPrim(RenderMeshPrim, RenderPatternAPI, Patterns, OutErrorText, Offset))
+			{
+				return false;
+			}
+			// Add the sim mesh pattern relationships
 			for (const FUsdPrim& MeshChildPrim : RenderMeshPrim.GetChildren())
 			{
 				if (MeshChildPrim.IsA(TEXT("GeomSubset")) && MeshChildPrim.HasAPI(RenderPatternAPI))
@@ -419,14 +530,28 @@ namespace UE::Chaos::ClothAsset::Private
 					}
 				}
 			}
-			return true;
+			// Increment offset to match the indices in the collapsed (merged) mesh
+			const FUsdAttribute FaceVertexCountsAttr = RenderMeshPrim.GetAttribute(TEXT("faceVertexCounts"));
+			if (!FaceVertexCountsAttr)
+			{
+				OutErrorText = LOCTEXT("MissingRenderMeshFaceCountAttribute", "Missing render mesh faceVertexCounts attribute.");
+				return false;
+			}
+			else if (FaceVertexCountsAttr.GetTypeName() != TEXT("int[]"))
+			{
+				OutErrorText = LOCTEXT("WrongRenderMeshFaceCountTypeName", "Wrong render mesh faceVertexCounts type name. Needs to be 'int[]'.");
+				return false;
+			}
+			const TArray<int32> FaceVertexCounts = GetIntArrayValues(FaceVertexCountsAttr);
+			Offset += FaceVertexCounts.Num();
 		}
-		return false;
+		return true;
 	}
 
 	static bool ImportPatternsFromSimMeshPrim(const FUsdPrim& SimMeshPrim, TMap<FName, TSet<int32>>& Patterns, FText& OutErrorText)
 	{
-		return ImportPatternsFromMeshPrim(SimMeshPrim, SimPatternAPI, Patterns, OutErrorText);
+		Patterns.Reset();
+		return AddPatternsFromMeshPrim(SimMeshPrim, SimPatternAPI, Patterns, OutErrorText);
 	}
 
 	static bool ImportSewingsFromSimMeshPrim(const FUsdPrim& SimMeshPrim, TMap<FName, TSet<FIntVector2>>& Sewings, FText& OutErrorText)
@@ -442,7 +567,7 @@ namespace UE::Chaos::ClothAsset::Private
 				const FUsdAttribute FamilyNameAttr = GeomSubset.GetFamilyNameAttr();
 				if (GetStringValue(FamilyNameAttr) != TEXT("sewing"))
 				{
-					OutErrorText = LOCTEXT("WrongSewingGeomSubsetFamilyName", "Wrong sewing GeomSubset family name. Needs to be 'pattern'.");
+					OutErrorText = LOCTEXT("WrongSewingGeomSubsetFamilyName", "Wrong sewing GeomSubset family name. Needs to be 'sewing'.");
 					return false;
 				}
 
@@ -488,8 +613,218 @@ namespace UE::Chaos::ClothAsset::Private
 		}
 		return true;
 	}
+
+	static bool ImportSimulationProperties(
+		const FUsdPrim& ClothPrim,
+		const FUsdPrim& SimMeshPrim,
+		const EUsdUpAxis UsdUpAxis,
+		FManagedArrayCollection& OutSimulationCollection,
+		TMap<FName, int32>& OutSimPatternFabricIndices,
+		FText& OutErrorText)
+	{
+		OutSimPatternFabricIndices.Reset();
+		
+		const TSharedRef<FManagedArrayCollection> SimulationCollection = MakeShared<FManagedArrayCollection>();
+		FCollectionClothFacade ClothFacade(SimulationCollection);
+		ClothFacade.DefineSchema(EClothCollectionExtendedSchemas::Solvers | EClothCollectionExtendedSchemas::Fabrics);
+
+		// Import the solver properties
+		auto ImportSolverProperties = [&UsdUpAxis, &ClothFacade](const FUsdPrim& ParentPrim) -> bool
+			{
+				// Find the CloSolverPropertiesAPI prim
+				const FUsdPrim CloSolverPropertiesPrim = FindAPIPrim(ParentPrim, CloSolverPropertiesAPI);
+				if (!CloSolverPropertiesPrim)
+				{
+					return false;
+				}
+
+				// Read the attributes
+				const FUsdPrimAttributeAccessor UsdPrimAttributeAccessor(CloSolverPropertiesPrim, UsdUpAxis);
+
+				float AirDamping = 0.1f;
+				AirDamping = UsdPrimAttributeAccessor.GetValue(TEXT("airDamping"), AirDamping);
+
+				FVector3f Gravity(0.f, 0.f, -980.f);
+				constexpr float CentimetersPerMillimeters = 0.1f;  // CloSolverPropertiesAPI Gravity is always given in mm regardless of UsdStageInfo
+				Gravity = UsdPrimAttributeAccessor.GetValue(TEXT("gravity"), Gravity / CentimetersPerMillimeters) * CentimetersPerMillimeters;
+
+				float TimeStep = 0.033333f;
+				TimeStep = UsdPrimAttributeAccessor.GetValue(TEXT("timeStep"), TimeStep);
+
+				uint32 SubStepCount = 1;
+				SubStepCount = UsdPrimAttributeAccessor.GetValue(TEXT("subStepCount"), SubStepCount);
+
+				ClothFacade.SetSolverGravity(Gravity);
+				ClothFacade.SetSolverAirDamping(AirDamping);
+				ClothFacade.SetSolverTimeStep(TimeStep);
+				ClothFacade.SetSolverSubSteps(SubStepCount);
+
+				return true;
+			};
+
+		if (!ImportSolverProperties(ClothPrim))  // Try under the cloth prim first
+		{
+			ImportSolverProperties(SimMeshPrim);  // If that doesn't work, try under the sim mesh prim instead
+		}
+
+		// Import the fabric properties
+		auto ImportFabricProperties = [&UsdUpAxis, &ClothFacade](const FUsdGeomSubset& GeomSubset) -> int32
+			{
+				// Read the attributes
+				const FUsdPrimAttributeAccessor UsdPrimAttributeAccessor(GeomSubset.GetPrim(), UsdUpAxis);
+
+				constexpr float BendingScaling = 1e-5f;  // From g.mm2/s2 to kg.cm2/s2
+				constexpr float StretchShearScaling = 1e-3f;  // From g/s2 to kg/s2
+				constexpr float DensityScaling = 1e+3f;  // From g/mm2 to kg/m2
+				constexpr float ThicknessScaling = 1e-1f;  // From mm to cm
+
+				float BendingBiasLeft = 0.0f;
+				BendingBiasLeft = UsdPrimAttributeAccessor.GetArrayValue<float>(TEXT("primvars:clo:bendingBiasLeft"), BendingBiasLeft) * BendingScaling;
+
+				float BendingBiasRight = 0.0f;
+				BendingBiasRight = UsdPrimAttributeAccessor.GetArrayValue<float>(TEXT("primvars:clo:bendingBiasRight"), BendingBiasRight) * BendingScaling;
+
+				float BendingWarp = 0.0f;
+				BendingWarp = UsdPrimAttributeAccessor.GetArrayValue<float>(TEXT("primvars:clo:bendingWarp"), BendingWarp) * BendingScaling;
+
+				float BendingWeft = 0.0f;
+				BendingWeft = UsdPrimAttributeAccessor.GetArrayValue<float>(TEXT("primvars:clo:bendingWeft"), BendingWeft) * BendingScaling;
+
+				float BucklingRatioBiasLeft = 0.0f;
+				BucklingRatioBiasLeft = UsdPrimAttributeAccessor.GetArrayValue<float>(TEXT("primvars:clo:bucklingRatioBiasLeft"), BucklingRatioBiasLeft);
+
+				float BucklingRatioBiasRight = 0.0f;
+				BucklingRatioBiasRight = UsdPrimAttributeAccessor.GetArrayValue<float>(TEXT("primvars:clo:bucklingRatioBiasRight"), BucklingRatioBiasRight);
+
+				float BucklingRatioWarp = 0.0f;
+				BucklingRatioWarp = UsdPrimAttributeAccessor.GetArrayValue<float>(TEXT("primvars:clo:bucklingRatioWarp"), BucklingRatioWarp);
+
+				float BucklingRatioWeft = 0.0f;
+				BucklingRatioWeft = UsdPrimAttributeAccessor.GetArrayValue<float>(TEXT("primvars:clo:bucklingRatioWeft"), BucklingRatioWeft);
+
+				float BucklingStiffnessBiasLeft = 0.0f;
+				BucklingStiffnessBiasLeft = UsdPrimAttributeAccessor.GetArrayValue<float>(TEXT("primvars:clo:bucklingStiffnessBiasLeft"), BucklingStiffnessBiasLeft);
+
+				float BucklingStiffnessBiasRight = 0.0f;
+				BucklingStiffnessBiasRight = UsdPrimAttributeAccessor.GetArrayValue<float>(TEXT("primvars:clo:bucklingStiffnessBiasRight"), BucklingStiffnessBiasRight);
+
+				float BucklingStiffnessWarp = 0.0f;
+				BucklingStiffnessWarp = UsdPrimAttributeAccessor.GetArrayValue<float>(TEXT("primvars:clo:bucklingStiffnessWarp"), BucklingStiffnessWarp);
+
+				float BucklingStiffnessWeft = 0.0f;
+				BucklingStiffnessWeft = UsdPrimAttributeAccessor.GetArrayValue<float>(TEXT("primvars:clo:bucklingStiffnessWeft"), BucklingStiffnessWeft);
+
+				float Density = 0.0f;
+				Density = UsdPrimAttributeAccessor.GetArrayValue<float>(TEXT("primvars:clo:density"), Density) * DensityScaling;
+
+				float Friction = 0.0f;
+				Friction = UsdPrimAttributeAccessor.GetArrayValue<float>(TEXT("primvars:clo:friction"), Friction);
+
+				float Damping = 0.0f;
+				Damping = UsdPrimAttributeAccessor.GetArrayValue<float>(TEXT("primvars:clo:internalDamping"), Damping);
+
+				float ShearLeft = 0.0f;
+				ShearLeft = UsdPrimAttributeAccessor.GetArrayValue<float>(TEXT("primvars:clo:shearLeft"), ShearLeft) * StretchShearScaling;
+
+				float ShearRight = 0.0f;
+				ShearRight = UsdPrimAttributeAccessor.GetArrayValue<float>(TEXT("primvars:clo:shearRight"), ShearRight) * StretchShearScaling;
+
+				float StretchWarp = 0.0f;
+				StretchWarp = UsdPrimAttributeAccessor.GetArrayValue<float>(TEXT("primvars:clo:stretchWarp"), StretchWarp) * StretchShearScaling;
+
+				float StretchWeft = 0.0f;
+				StretchWeft = UsdPrimAttributeAccessor.GetArrayValue<float>(TEXT("primvars:clo:stretchWeft"), StretchWeft) * StretchShearScaling;
+
+				float Thickness = 0.0f;
+				Thickness = UsdPrimAttributeAccessor.GetArrayValue<float>(TEXT("primvars:clo:thickness"), Thickness)* ThicknessScaling;
+
+				// Initialize a new fabric with these parameters
+				const FCollectionClothFabricFacade::FAnisotropicData BendingStiffness(
+					BendingWeft,
+					BendingWarp,
+					0.5f * (BendingBiasLeft + BendingBiasRight));
+
+				const FCollectionClothFabricFacade::FAnisotropicData StretchStiffness(
+					StretchWeft,
+					StretchWarp,
+					0.5f * (ShearLeft + ShearRight));
+
+				const float BucklingRatio =
+					(BucklingRatioWeft + BucklingRatioWarp +
+						0.5f * (BucklingRatioBiasLeft + BucklingRatioBiasRight)) / 3.f;
+
+				const FCollectionClothFabricFacade::FAnisotropicData BucklingStiffness =
+					BucklingRatio < UE_SMALL_NUMBER ? BendingStiffness :
+					FCollectionClothFabricFacade::FAnisotropicData(
+						BendingStiffness.Weft * BucklingStiffnessWeft,
+						BendingStiffness.Warp * BucklingStiffnessWarp,
+						BendingStiffness.Bias * 0.5f * (BucklingStiffnessBiasLeft + BucklingStiffnessBiasRight));
+
+				const int32 FabricIndex = ClothFacade.AddFabric();
+				FCollectionClothFabricFacade Fabric = ClothFacade.GetFabric(FabricIndex);
+
+				Fabric.Initialize(
+					BendingStiffness,
+					BucklingRatio,
+					BucklingStiffness,
+					StretchStiffness,
+					Density,
+					Friction,
+					Damping,
+					0.0f,
+					0,
+					Thickness);
+
+				return FabricIndex;
+			};
+
+		TMap<FName, FUsdGeomSubset> GeomSubsets;
+		if (!GetPatternGeomSubsetsFromMeshPrim(SimMeshPrim, SimPatternAPI, GeomSubsets, OutErrorText))
+		{
+			return false;
+		}
+
+		for (const TPair<FName, FUsdGeomSubset>& GeomSubset : GeomSubsets)
+		{
+			// Setup a new fabric for this pattern
+			const int32 FabricIndex = ImportFabricProperties(GeomSubset.Value);
+			OutSimPatternFabricIndices.Emplace(GeomSubset.Key, FabricIndex);
+		}
+
+		OutSimulationCollection = MoveTemp(*SimulationCollection);
+
+		return true;
+	}
 }  // End namespace Private
 
+void FChaosClothAssetUsdClothData::Reset()
+{
+	SimPatterns.Reset();
+	Sewings.Reset();
+	RenderPatterns.Reset();
+	RenderToSimPatterns.Reset();
+	SimPatternFabricIndices.Reset();
+	SimulationCollection.Reset();
+}
+
+bool FChaosClothAssetUsdClothData::Serialize(FArchive& Ar)
+{
+	Ar.UsingCustomVersion(FUE5ReleaseStreamObjectVersion::GUID);
+
+	Ar << SimPatterns;
+	Ar << Sewings;
+	Ar << RenderPatterns;
+	Ar << RenderToSimPatterns;
+
+	if (Ar.CustomVer(FUE5ReleaseStreamObjectVersion::GUID) >= FUE5ReleaseStreamObjectVersion::AddSimulationPropertySupportToClothUSDImportNodeV2)
+	{
+		Ar << SimPatternFabricIndices;
+		SimulationCollection.Serialize(Ar);
+	}
+	
+	return true;
+}
+
 FChaosClothAssetUSDImportNode_v2::FChaosClothAssetUSDImportNode_v2(const UE::Dataflow::FNodeParameters& InParam, FGuid InGuid)
 	: FDataflowNode(InParam, InGuid)
 	, UsdFile(
@@ -576,14 +911,18 @@ void FChaosClothAssetUSDImportNode_v2::Serialize(FArchive& Ar)
 		{
 			ClothFacade.DefineSchema();
 		}
+
+		// Also apply any required fixup (e.g. soft object path names)
+		ClothFacade.PostSerialize(Ar);
+		
 		Collection = MoveTemp(*ClothCollection);
 
 		// Regenerate correct dependencies if needed
 		if (!ImportedAssets_DEPRECATED.IsEmpty())
 		{
 			ImportedAssets_DEPRECATED.Empty();
-			ImportedSimAssets = GetImportedAssetDependencies(ImportedSimStaticMesh);
-			ImportedRenderAssets = GetImportedAssetDependencies(ImportedRenderStaticMesh);
+			ImportedSimAssets = Private::GetImportedAssetDependencies(ImportedSimStaticMesh);
+			ImportedRenderAssets = Private::GetImportedAssetDependencies(ImportedRenderStaticMesh);
 		}
 	}
 }
@@ -597,6 +936,7 @@ void FChaosClothAssetUSDImportNode_v2::ResetImport()
 	ImportedUVScale = { 1.f, 1.f };
 	ImportedRenderAssets.Reset();
 	ImportedSimAssets.Reset();
+	UsdClothData.Reset();
 }
 
 // V1 of the USD importer (schemaless)
@@ -659,7 +999,7 @@ bool FChaosClothAssetUSDImportNode_v2::ImportUsdFile(const FString& UsdFilePath,
 	}
 
 	// Start slow task
-	const float NumSteps = 1.f + (bImportSimMesh ? bImportRenderMesh ? 2.f : 1.f : bImportRenderMesh ? 1.f : 0.f);
+	const float NumSteps = bImportSimMesh ? bImportRenderMesh ? 2.f : 1.f : bImportRenderMesh ? 1.f : 0.f;
 	FScopedSlowTask SlowTask(NumSteps, LOCTEXT("ImportingUSDFile", "Importing USD file..."));
 	SlowTask.MakeDialogDelayed(1.f);
 
@@ -675,7 +1015,7 @@ bool FChaosClothAssetUSDImportNode_v2::ImportUsdFile(const FString& UsdFilePath,
 	}
 
 	// Find the cloth prim
-	const FUsdPrim ClothPrim = FindClothPrim(UsdStage.GetPseudoRoot());
+	FUsdPrim ClothPrim = FindAPIPrim(UsdStage.GetPseudoRoot(), ClothRootAPI);
 	if (!ClothPrim)
 	{
 		OutErrorText = LOCTEXT("CantFindClothRootAPI", "Can't find a cloth root inside the specified USD file.");
@@ -683,9 +1023,9 @@ bool FChaosClothAssetUSDImportNode_v2::ImportUsdFile(const FString& UsdFilePath,
 	}
 
 	// Find SimMesh and Render Mesh prims
-	const FUsdPrim SimMeshPrim = FindSimMeshPrim(ClothPrim);
-	const FUsdPrim RenderMeshPrim = FindRenderMeshPrim(ClothPrim);
-	if (!SimMeshPrim && !RenderMeshPrim)
+	FUsdPrim SimMeshPrim = FindSimMeshPrim(ClothPrim);
+	TArray<FUsdPrim> RenderMeshPrims = FindRenderMeshPrims(ClothPrim);
+	if (!SimMeshPrim && !RenderMeshPrims.Num())
 	{
 		OutErrorText = LOCTEXT("CantFindMeshPrims", "Can't find a sim mesh or render mesh prim with valid pattern data.");
 		return false;
@@ -706,13 +1046,18 @@ bool FChaosClothAssetUSDImportNode_v2::ImportUsdFile(const FString& UsdFilePath,
 	const FString PackageName = FString::Printf(TEXT("%s_%08X"), *UsdFileName, UsdPathHash);
 	PackagePath = FPaths::Combine(AssetPath + TEXT("_Import"), PackageName);
 
-	// Import the stage
-	SlowTask.EnterProgressFrame(1.f);
-	TArray<TObjectPtr<UObject>> ImportedAssets = ImportStaticMeshesFromUsdStage(UsdStage, UsdFilePath, PackagePath);
-
-	// Find the imported static meshes in the imported USD assets
-	ImportedSimStaticMesh = FindImportedStaticMesh(ImportedAssets, SimMeshPrim.GetPrimPath().GetString());
-	ImportedRenderStaticMesh = FindImportedStaticMesh(ImportedAssets, RenderMeshPrim.GetPrimPath().GetString());
+	// Mesh selector lambda
+	TArray<FUsdPrim> ClothChildren = ClothPrim.GetChildren();
+	auto SetMeshClothChildrenActive = [&ClothChildren](auto Predicate)
+		{
+			for (FUsdPrim& ClothChildPrim : ClothChildren)
+			{
+				if (ClothChildPrim.IsA(TEXT("Mesh")))
+				{
+					ClothChildPrim.SetActive(Predicate(ClothChildPrim));
+				}
+			}
+		};
 
 	// Import sim mesh from the static mesh
 	if (bImportSimMesh)
@@ -737,7 +1082,25 @@ bool FChaosClothAssetUSDImportNode_v2::ImportUsdFile(const FString& UsdFilePath,
 			return false;
 		}
 
-		// Lastly import the geometry and finalize the patterns
+		// Import the sim static mesh selectively by disabling the render prims
+		SetMeshClothChildrenActive([&SimMeshPrim](FUsdPrim& ClothChildPrim) { return ClothChildPrim == SimMeshPrim; });
+		TArray<TObjectPtr<UObject>> ImportedAssets = ImportStaticMeshesFromUsdStage(UsdStage, UsdFilePath, PackagePath);
+		ImportedSimStaticMesh = FindImportedStaticMesh(ImportedAssets, SimMeshPrim.GetPrimPath().GetString());
+
+		// Import the simulation properties if they exist
+		const FUsdStageInfo UsdStageInfo(UsdStage);
+		if (!ImportSimulationProperties(
+			ClothPrim,
+			SimMeshPrim,
+			UsdStageInfo.UpAxis,
+			UsdClothData.SimulationCollection,
+			UsdClothData.SimPatternFabricIndices,
+			OutErrorText))
+		{
+			return false;
+		}
+
+		// Import the geometry and finalize the patterns
 		if (!ImportSimStaticMesh(ClothCollection, OutErrorText))
 		{
 			return false;
@@ -749,8 +1112,16 @@ bool FChaosClothAssetUSDImportNode_v2::ImportUsdFile(const FString& UsdFilePath,
 	{
 		SlowTask.EnterProgressFrame(1.f);
 
+		// Import the render mesh selectively by disabling the sim prim
+		ClothPrim.SetTypeName(TEXT("Xform"));  // Won't collapse the meshes unless the root is an Xform prim, so make it one
+		UsdUtils::SetDefaultKind(ClothPrim, EUsdDefaultKind::Component);  // Set the cloth prim kind to enable KindsToCollapse
+
+		SetMeshClothChildrenActive([&RenderMeshPrims](FUsdPrim& ClothChildPrim) { return RenderMeshPrims.Contains(ClothChildPrim); });
+		TArray<TObjectPtr<UObject>> ImportedAssets = ImportStaticMeshesFromUsdStage(UsdStage, UsdFilePath, PackagePath);
+		ImportedRenderStaticMesh = FindImportedStaticMesh(ImportedAssets, ClothPrim.GetPrimPath().GetString());  // Collapsed mesh uses the root ClothPrim name
+
 		// Import the render patterns
-		if (!ImportPatternsFromRenderMeshPrim(RenderMeshPrim, SimMeshPrim, UsdClothData.RenderPatterns, UsdClothData.RenderToSimPatterns, OutErrorText))
+		if (!ImportPatternsFromRenderMeshPrims(RenderMeshPrims, SimMeshPrim, UsdClothData.RenderPatterns, UsdClothData.RenderToSimPatterns, OutErrorText))
 		{
 			return false;
 		}
@@ -810,59 +1181,6 @@ void FChaosClothAssetUSDImportNode_v2::UpdateImportedAssets(const FString& SimMe
 	}
 }
 
-TObjectPtr<UStaticMesh> FChaosClothAssetUSDImportNode_v2::FindImportedStaticMesh(const TArrayView<TObjectPtr<UObject>> ImportedAssets, const FString& MeshPrimPath)
-{
-	for (TObjectPtr<UObject>& ImportedAsset : ImportedAssets)
-	{
-		if (const TObjectPtr<UStaticMesh> ImportedStaticMesh = Cast<UStaticMesh>(ImportedAsset))
-		{
-			if (const UUsdMeshAssetUserData* const AssetUserData = Cast<UUsdMeshAssetUserData>(
-				ImportedStaticMesh->GetAssetUserDataOfClass(UUsdMeshAssetUserData::StaticClass())))
-			{
-				if (AssetUserData->PrimPaths.Find(MeshPrimPath) != INDEX_NONE)
-				{
-					return ImportedStaticMesh;
-				}
-			}
-		}
-	}
-	return nullptr;
-}
-
-TArray<TObjectPtr<UObject>> FChaosClothAssetUSDImportNode_v2::GetImportedAssetDependencies(const UObject* StaticMesh)
-{
-	using namespace UE::Chaos::ClothAsset::Private;
-
-	TSet<TObjectPtr<UObject>> ImportedAssets;
-	if (StaticMesh)
-	{
-		TQueue<const UObject*> AssetsToVisit;
-		AssetsToVisit.Enqueue(StaticMesh);
-
-		const UObject* VisitedAsset;
-		while (AssetsToVisit.Dequeue(VisitedAsset))
-		{
-			const FName VisitedAssetPackageName(VisitedAsset->GetPackage()->GetName());
-			const TArray<UObject*> AssetDependencies = GetAssetDependencies(VisitedAsset);
-
-			UE_CLOG(AssetDependencies.Num(), LogChaosClothAssetDataflowNodes, Verbose, TEXT("Dependencies for Object %s - %s:"), *VisitedAsset->GetName(), *VisitedAssetPackageName.ToString());
-			for (UObject* const AssetDependency : AssetDependencies)
-			{
-				if (!ImportedAssets.Contains(AssetDependency))
-				{
-					// Add the dependency
-					UE_LOG(LogChaosClothAssetDataflowNodes, Verbose, TEXT("Found %s"), *AssetDependency->GetPackage()->GetName());
-					ImportedAssets.Emplace(AssetDependency);
-
-					// Visit this asset too
-					AssetsToVisit.Enqueue(AssetDependency);
-				}
-			}
-		}
-	}
-	return ImportedAssets.Array();
-}
-
 bool FChaosClothAssetUSDImportNode_v2::ImportSimStaticMesh(const TSharedRef<FManagedArrayCollection> ClothCollection, FText& OutErrorText)
 {
 	using namespace UE::Chaos::ClothAsset;
@@ -889,7 +1207,7 @@ bool FChaosClothAssetUSDImportNode_v2::ImportSimStaticMesh(const TSharedRef<FMan
 			FClothGeometryTools::BindMeshToRootBone(ClothCollection, bBindSimMesh, bBindRenderMesh);
 
 			// Make sure to clean the dependencies whatever the import status is
-			ImportedSimAssets = GetImportedAssetDependencies(ImportedSimStaticMesh);
+			ImportedSimAssets = Private::GetImportedAssetDependencies(ImportedSimStaticMesh);
 		};
 
 	if (!ImportedSimStaticMesh)
@@ -923,8 +1241,8 @@ bool FChaosClothAssetUSDImportNode_v2::ImportSimStaticMesh(const TSharedRef<FMan
 	const TConstArrayView<FVertexInstanceID> TriangleVertexInstanceIndices = StaticMeshAttributes.GetTriangleVertexInstanceIndices().GetRawArray();
 	const TConstArrayView<FVector2f> VertexInstanceUVs = StaticMeshAttributes.GetVertexInstanceUVs().GetRawArray();
 
-	check(TriangleVertexInstanceIndices.Num() % 3 == 0)
-		TriangleToVertexIndex.SetNumUninitialized(TriangleVertexInstanceIndices.Num() / 3);
+	check(TriangleVertexInstanceIndices.Num() % 3 == 0);
+	TriangleToVertexIndex.SetNumUninitialized(TriangleVertexInstanceIndices.Num() / 3);
 
 	auto SetRestPositions2D = [&RestPositions2D, &VertexInstanceUVs](FVertexID VertexID, FVertexInstanceID VertexInstanceID) -> bool
 		{
@@ -1042,10 +1360,15 @@ bool FChaosClothAssetUSDImportNode_v2::ImportSimStaticMesh(const TSharedRef<FMan
 					}
 				}
 
+				// Find this pattern's fabric if any
+				const int32 FabricIndex = UsdClothData.SimPatternFabricIndices.Contains(PatternNameFaces.Key) ?
+					UsdClothData.SimPatternFabricIndices[PatternNameFaces.Key] :
+					INDEX_NONE;
+
 				// Add this pattern to the cloth collection
 				const int32 SimPatternIndex = ClothFacade.AddSimPattern();
 				FCollectionClothSimPatternFacade SimPattern = ClothFacade.GetSimPattern(SimPatternIndex);
-				SimPattern.Initialize(PatternRestPositions2D, PatternDrapedPositions3D, PatternTriangleToVertexIndex);
+				SimPattern.Initialize(PatternRestPositions2D, PatternDrapedPositions3D, PatternTriangleToVertexIndex, FabricIndex);
 
 				// Keep track of the original triangle indices
 				const TArrayView<TArray<int32>> OriginalTriangles =
@@ -1113,6 +1436,42 @@ bool FChaosClothAssetUSDImportNode_v2::ImportSimStaticMesh(const TSharedRef<FMan
 			FCollectionClothSeamFacade ClothSeamFacade = ClothFacade.AddGetSeam();
 			ClothSeamFacade.Initialize(Indices.Array());
 		}
+
+		// Add the solver properties
+		const FCollectionClothConstFacade SimulationClothFacade(MakeShared<FManagedArrayCollection>(UsdClothData.SimulationCollection));
+		if (SimulationClothFacade.IsValid(EClothCollectionExtendedSchemas::Solvers))
+		{
+			ClothFacade.DefineSchema(EClothCollectionExtendedSchemas::Solvers);
+			ClothFacade.SetSolverAirDamping(SimulationClothFacade.GetSolverAirDamping());
+			ClothFacade.SetSolverGravity(SimulationClothFacade.GetSolverGravity());
+			ClothFacade.SetSolverSubSteps(SimulationClothFacade.GetSolverSubSteps());
+			ClothFacade.SetSolverTimeStep(SimulationClothFacade.GetSolverTimeStep());
+		}
+
+		// Add the fabric properties
+		if (SimulationClothFacade.IsValid(EClothCollectionExtendedSchemas::Fabrics))
+		{
+			ClothFacade.DefineSchema(EClothCollectionExtendedSchemas::Fabrics);
+
+			for (int32 FabricIndex = 0; FabricIndex < SimulationClothFacade.GetNumFabrics(); ++FabricIndex)
+			{
+				verify(ClothFacade.AddFabric() == FabricIndex);
+				FCollectionClothFabricFacade ClothFabricFacade = ClothFacade.GetFabric(FabricIndex);
+				const FCollectionClothFabricConstFacade SimulationClothFabricFacade = SimulationClothFacade.GetFabric(FabricIndex);
+
+				ClothFabricFacade.Initialize(
+					SimulationClothFabricFacade.GetBendingStiffness(),
+					SimulationClothFabricFacade.GetBucklingRatio(),
+					SimulationClothFabricFacade.GetBucklingStiffness(),
+					SimulationClothFabricFacade.GetStretchStiffness(),
+					SimulationClothFabricFacade.GetDensity(),
+					SimulationClothFabricFacade.GetFriction(),
+					SimulationClothFabricFacade.GetDamping(),
+					SimulationClothFabricFacade.GetPressure(),
+					SimulationClothFabricFacade.GetLayer(),
+					SimulationClothFabricFacade.GetCollisionThickness());
+			}
+		}
 	}
 	return true;
 }
@@ -1144,7 +1503,7 @@ bool FChaosClothAssetUSDImportNode_v2::ImportRenderStaticMesh(const TSharedRef<F
 			FClothGeometryTools::BindMeshToRootBone(ClothCollection, bBindSimMesh, bBindRenderMesh);
 
 			// Make sure to clean the dependencies whatever the import status is
-			ImportedRenderAssets = GetImportedAssetDependencies(ImportedRenderStaticMesh);
+			ImportedRenderAssets = Private::GetImportedAssetDependencies(ImportedRenderStaticMesh);
 		};
 
 	// Import the LOD 0
diff --git a/Engine/Plugins/ChaosClothAssetEditor/Source/ChaosClothAssetDataflowNodes/Private/ChaosClothAsset/UsdPrimAttributeAccessor.h b/Engine/Plugins/ChaosClothAssetEditor/Source/ChaosClothAssetDataflowNodes/Private/ChaosClothAsset/UsdPrimAttributeAccessor.h
new file mode 100644
--- /dev/null
+++ b/Engine/Plugins/ChaosClothAssetEditor/Source/ChaosClothAssetDataflowNodes/Private/ChaosClothAsset/UsdPrimAttributeAccessor.h
@@ -0,0 +1,129 @@
+// Copyright Epic Games, Inc. All Rights Reserved.
+#pragma once
+
+#include "USDStageOptions.h"
+#include "USDValueConversion.h"
+#include "UsdWrappers/UsdAttribute.h"
+#include "UsdWrappers/UsdPrim.h"
+#include "UsdWrappers/VtValue.h"
+
+namespace UE::Chaos::ClothAsset
+{
+	class FUsdPrimAttributeAccessor
+	{
+	public:
+		FUsdPrimAttributeAccessor(const FUsdPrim& InUsdPrim, const EUsdUpAxis UsdUpAxis)
+			: UsdPrim(InUsdPrim)
+			, AxesOrder(
+				0,
+				(UsdUpAxis == EUsdUpAxis::ZAxis) ? 1 : 2,
+				(UsdUpAxis == EUsdUpAxis::ZAxis) ? 2 : 1)
+		{
+		}
+
+		template<typename ValueType>
+		ValueType GetValue(const TCHAR* AttributeName, const ValueType DefaultValue = ValueType(0)) const
+		{
+			if (UsdPrim.HasAttribute(AttributeName))
+			{
+				const FUsdAttribute UsdAttribute = UsdPrim.GetAttribute(AttributeName);
+				if (UsdAttribute.HasValue())
+				{
+					FVtValue VtValue;
+					UsdAttribute.Get(VtValue);
+					return GetValueAs<ValueType>(VtValue, DefaultValue);
+				}
+			}
+			return DefaultValue;
+		}
+
+		template<typename ValueType>
+		TArray<ValueType> GetArray(const TCHAR* AttributeName) const
+		{
+			if (UsdPrim.HasAttribute(AttributeName))
+			{
+				const FUsdAttribute UsdAttribute = UsdPrim.GetAttribute(AttributeName);
+				if (UsdAttribute.HasValue())
+				{
+					FVtValue VtValue;
+					UsdAttribute.Get(VtValue);
+					return GetArrayAs<ValueType>(VtValue);
+				}
+			}
+			return TArray<ValueType>();
+		}
+
+		template<typename ValueType>
+		ValueType GetArrayValue(const TCHAR* AttributeName, const ValueType DefaultValue = ValueType(0), const int32 ValueIndex = 0) const
+		{
+			const TArray<ValueType> Array = GetArray<ValueType>(AttributeName);
+			return Array.IsValidIndex(ValueIndex) ? Array[ValueIndex] : DefaultValue;
+		}
+
+	private:
+		template<typename ValueType>
+		ValueType GetValueAs(const FVtValue& VtValue, const ValueType DefaultValue = ValueType(0)) const
+		{
+			unimplemented();
+			return DefaultValue;
+		}
+
+		template<typename ValueType>
+		TArray<ValueType> GetArrayAs(const FVtValue& VtValue) const
+		{
+			unimplemented();
+			return TArray<ValueType>();
+		}
+
+		const FUsdPrim UsdPrim;
+		const FIntVector AxesOrder;
+	};
+
+	template<>
+	uint32 FUsdPrimAttributeAccessor::GetValueAs<uint32>(const FVtValue& VtValue, const uint32 DefaultValue) const
+	{
+		return VtValue.GetTypeName() == TEXT("unsigned int") ? UsdUtils::GetUnderlyingValue<uint32>(VtValue).Get(DefaultValue) : DefaultValue;  // Note: FUsdAttribute.GetTypeName() would return "uint" instead!
+	}
+
+	template<>
+	float FUsdPrimAttributeAccessor::GetValueAs<float>(const FVtValue& VtValue, const float DefaultValue) const
+	{
+		return VtValue.GetTypeName() == TEXT("float") ? UsdUtils::GetUnderlyingValue<float>(VtValue).Get(DefaultValue) : DefaultValue;
+	}
+
+	template<>
+	FVector3f FUsdPrimAttributeAccessor::GetValueAs<FVector3f>(const FVtValue& VtValue, const FVector3f DefaultValue) const
+	{
+		if (VtValue.GetTypeName() == TEXT("GfVec3f"))  // Note: FUsdAttribute.GetTypeName() would return "float3" instead!
+		{
+			UsdUtils::FConvertedVtValue ConvertedVtValue;
+			if (UsdToUnreal::ConvertValue(VtValue, ConvertedVtValue) && !ConvertedVtValue.bIsArrayValued && !ConvertedVtValue.bIsEmpty)
+			{
+				return FVector3f(
+					ConvertedVtValue.Entries[0][AxesOrder[0]].Get<float>(),
+					ConvertedVtValue.Entries[0][AxesOrder[1]].Get<float>(),
+					ConvertedVtValue.Entries[0][AxesOrder[2]].Get<float>());
+			}
+		}
+		return DefaultValue;
+	}
+
+	template<>
+	TArray<float> FUsdPrimAttributeAccessor::GetArrayAs<float>(const FVtValue& VtValue) const
+	{
+		TArray<float> Array;
+		if (VtValue.GetTypeName() == TEXT("VtArray<float>"))  // Note: FUsdAttribute.GetTypeName() would return "float[]" instead!
+		{
+			UsdUtils::FConvertedVtValue ConvertedVtValue;
+			if (UsdToUnreal::ConvertValue(VtValue, ConvertedVtValue) && ConvertedVtValue.bIsArrayValued && !ConvertedVtValue.bIsEmpty)
+			{
+				Array.Reserve(ConvertedVtValue.Entries.Num());
+				for (int32 Index = 0; Index < ConvertedVtValue.Entries.Num(); ++Index)
+				{
+					Array.Emplace(ConvertedVtValue.Entries[Index][0].Get<float>());
+				}
+			}
+		}
+		return Array;
+	}
+}
diff --git a/Engine/Plugins/ChaosClothAssetEditor/Source/ChaosClothAssetDataflowNodes/Public/ChaosClothAsset/USDImportNode_v2.h b/Engine/Plugins/ChaosClothAssetEditor/Source/ChaosClothAssetDataflowNodes/Public/ChaosClothAsset/USDImportNode_v2.h
--- a/Engine/Plugins/ChaosClothAssetEditor/Source/ChaosClothAssetDataflowNodes/Public/ChaosClothAsset/USDImportNode_v2.h
+++ b/Engine/Plugins/ChaosClothAssetEditor/Source/ChaosClothAssetDataflowNodes/Public/ChaosClothAsset/USDImportNode_v2.h
@@ -6,17 +6,8 @@
 #include "Dataflow/DataflowFunctionProperty.h"
 #include "GeometryCollection/ManagedArrayCollection.h"
 
-#if UE_ENABLE_INCLUDE_ORDER_DEPRECATED_IN_5_5
-#include "Misc/SecureHash.h"
-#endif
 #include "USDImportNode_v2.generated.h"
 
-#if UE_ENABLE_INCLUDE_ORDER_DEPRECATED_IN_5_5
-namespace Dataflow = UE::Dataflow;
-#else
-namespace UE_DEPRECATED(5.5, "Use UE::Dataflow instead.") Dataflow {}
-#endif
-
 class UStaticMesh;
 
 /** Support struct for storage of the extra data that doesn't fit inside a static mesh. */
@@ -29,15 +20,11 @@ struct FChaosClothAssetUsdClothData final
 	TMap<FName, TSet<FIntVector2>> Sewings;
 	TMap<FName, TSet<int32>> RenderPatterns;
 	TMap<FName, TSet<FName>> RenderToSimPatterns;
+	TMap<FName, int32> SimPatternFabricIndices;
+	FManagedArrayCollection SimulationCollection;
 
-	bool Serialize(FArchive& Ar)
-	{
-		Ar << SimPatterns;
-		Ar << Sewings;
-		Ar << RenderPatterns;
-		Ar << RenderToSimPatterns;
-		return true;
-	}
+	void Reset();
+	bool Serialize(FArchive& Ar);
 };
 
 template<> struct TStructOpsTypeTraits<FChaosClothAssetUsdClothData> : public TStructOpsTypeTraitsBase2<FChaosClothAssetUsdClothData>
@@ -111,12 +98,6 @@ private:
 	// Find the two imported static meshes
 	void UpdateImportedAssets(const FString& SimMeshName, const FString& RenderMeshName);
 
-	// Find an imported static mesh from the given usd prim path
-	TObjectPtr<UStaticMesh> FindImportedStaticMesh(const TArrayView<TObjectPtr<UObject>> ImportedAssetsconst, const FString& MeshPrimPath);
-
-	// Find the asset dependencies for the specified static mesh
-	TArray<TObjectPtr<UObject>> GetImportedAssetDependencies(const UObject* StaticMesh);
-
 	/** Content folder where all the USD assets are imported. */
 	UPROPERTY(VisibleAnywhere, Category = "USD Import")
 	FString PackagePath;
diff --git a/Engine/Plugins/Runtime/USDCore/Source/UnrealUSDWrapper/Private/UsdWrappers/UsdGeomSubset.cpp b/Engine/Plugins/Runtime/USDCore/Source/UnrealUSDWrapper/Private/UsdWrappers/UsdGeomSubset.cpp
--- a/Engine/Plugins/Runtime/USDCore/Source/UnrealUSDWrapper/Private/UsdWrappers/UsdGeomSubset.cpp
+++ b/Engine/Plugins/Runtime/USDCore/Source/UnrealUSDWrapper/Private/UsdWrappers/UsdGeomSubset.cpp
@@ -137,6 +137,15 @@ namespace UE
 	}
 #endif	  // #if USE_USD_SDK
 
+	FUsdPrim FUsdGeomSubset::GetPrim() const
+	{
+#if USE_USD_SDK
+		return FUsdPrim(Impl->PxrUsdGeomSubset.Get().GetPrim());
+#else
+		return FUsdPrim();
+#endif	  // #if USE_USD_SDK
+	}
+
 	FUsdAttribute FUsdGeomSubset::GetElementTypeAttr() const
 	{
 #if USE_USD_SDK
diff --git a/Engine/Plugins/Runtime/USDCore/Source/UnrealUSDWrapper/Public/UsdWrappers/UsdGeomSubset.h b/Engine/Plugins/Runtime/USDCore/Source/UnrealUSDWrapper/Public/UsdWrappers/UsdGeomSubset.h
--- a/Engine/Plugins/Runtime/USDCore/Source/UnrealUSDWrapper/Public/UsdWrappers/UsdGeomSubset.h
+++ b/Engine/Plugins/Runtime/USDCore/Source/UnrealUSDWrapper/Public/UsdWrappers/UsdGeomSubset.h
@@ -52,6 +52,8 @@ namespace UE
 
 		// Wrapped pxr::UsdGeomSubset functions, refer to the USD SDK documentation
 	public:
+		FUsdPrim GetPrim() const;
+
 		FUsdAttribute GetElementTypeAttr() const;
 
 		FUsdAttribute GetFamilyNameAttr() const;
diff --git a/Engine/Source/Runtime/Core/Public/UObject/UE5ReleaseStreamObjectVersions.inl b/Engine/Source/Runtime/Core/Public/UObject/UE5ReleaseStreamObjectVersions.inl
--- a/Engine/Source/Runtime/Core/Public/UObject/UE5ReleaseStreamObjectVersions.inl
+++ b/Engine/Source/Runtime/Core/Public/UObject/UE5ReleaseStreamObjectVersions.inl
@@ -189,5 +189,8 @@ UE5_RELEASE_VERSION(MediaProfilePluginCaptureCameraSoftPtr, 59)
 // Reparameterize Spline in SplineComponent based on SplineCurves parameterization
 UE5_RELEASE_VERSION(SplineComponentReparameterizeOnLoad, 60)
 
+// Add solver and fabric property support to the schema based Cloth USD importer
+UE5_RELEASE_VERSION(AddSimulationPropertySupportToClothUSDImportNodeV2, 61)
+
 // -----<new versions can be added above this line>-------------------------------------------------
 #undef UE5_RELEASE_VERSION
\ No newline at end of file
