Chaos Cloth Asset - Fixed USD cloth turned sewing import that weren't adding the corresponding springs.

[RN] Chaos Cloth Asset - Fixed USD cloth turned sewing import that weren't adding the corresponding springs.

diff --git a/Engine/Plugins/ChaosClothAsset/Source/ChaosClothAsset/Private/ChaosClothAsset/ClothGeometryTools.cpp b/Engine/Plugins/ChaosClothAsset/Source/ChaosClothAsset/Private/ChaosClothAsset/ClothGeometryTools.cpp
--- a/Engine/Plugins/ChaosClothAsset/Source/ChaosClothAsset/Private/ChaosClothAsset/ClothGeometryTools.cpp
+++ b/Engine/Plugins/ChaosClothAsset/Source/ChaosClothAsset/Private/ChaosClothAsset/ClothGeometryTools.cpp
@@ -489,6 +489,9 @@ namespace UE::Chaos::ClothAsset
 
 	void FClothGeometryTools::DeleteSimMesh(const TSharedRef<FManagedArrayCollection>& ClothCollection)
 	{
+		::Chaos::Softs::FEmbeddedSpringFacade SpringFacade(ClothCollection.Get(), ClothCollectionGroup::SimVertices3D);
+		SpringFacade.SetNumSpringConstraints(0);
+
 		FCollectionClothFacade ClothFacade(ClothCollection);
 		ClothFacade.SetNumSimPatterns(0);
 		ClothFacade.RemoveAllSimVertices3D();
diff --git a/Engine/Plugins/ChaosClothAssetEditor/Source/ChaosClothAssetDataflowNodes/Private/ChaosClothAsset/RemeshNode.cpp b/Engine/Plugins/ChaosClothAssetEditor/Source/ChaosClothAssetDataflowNodes/Private/ChaosClothAsset/RemeshNode.cpp
--- a/Engine/Plugins/ChaosClothAssetEditor/Source/ChaosClothAssetDataflowNodes/Private/ChaosClothAsset/RemeshNode.cpp
+++ b/Engine/Plugins/ChaosClothAssetEditor/Source/ChaosClothAssetDataflowNodes/Private/ChaosClothAsset/RemeshNode.cpp
@@ -3,6 +3,7 @@
 #include "ChaosClothAsset/RemeshNode.h"
 
 #include "ChaosClothAsset/ClothPatternToDynamicMesh.h"
+#include "Chaos/CollectionEmbeddedSpringConstraintFacade.h"
 #include "ChaosClothAsset/CollectionClothFacade.h"
 #include "ChaosClothAsset/CollectionClothSelectionFacade.h"
 #include "ChaosClothAsset/ClothCollectionGroup.h"
@@ -1055,6 +1056,12 @@ namespace UE::Chaos::ClothAsset::Private
 		}
 	}
 
+	static void EmptySprings(const TSharedRef<FManagedArrayCollection>& ClothCollection)
+	{
+		::Chaos::Softs::FEmbeddedSpringFacade SpringFacade(ClothCollection.Get(), ClothCollectionGroup::SimVertices3D);
+		SpringFacade.SetNumSpringConstraints(0);
+	}
+
 	static void RebuildTopologyDependentSimData(const TSharedRef<const FManagedArrayCollection>& InClothCollection,
 		const TSharedRef<FManagedArrayCollection>& OutClothCollection)
 	{
@@ -1988,6 +1995,7 @@ void FChaosClothAssetRemeshNode::Evaluate(UE::Dataflow::FContext& Context, const
 				else
 				{
 					Private::EmptySimSelections(OutputClothCollectionRef);
+					Private::EmptySprings(OutputClothCollectionRef);
 					Private::RemeshSimMesh(ClothCollection, IterationsSim, SmoothingSim, SimDensityMapName, FVector2f(TargetPercentSim / 100.0f, TargetPercentSim / 50.0f), OutputClothCollectionRef);
 					Private::RebuildTopologyDependentSimData(ClothCollection, OutputClothCollectionRef);
 				}
diff --git a/Engine/Plugins/ChaosClothAssetEditor/Source/ChaosClothAssetDataflowNodes/Private/ChaosClothAsset/USDImportNode_v2.cpp b/Engine/Plugins/ChaosClothAssetEditor/Source/ChaosClothAssetDataflowNodes/Private/ChaosClothAsset/USDImportNode_v2.cpp
--- a/Engine/Plugins/ChaosClothAssetEditor/Source/ChaosClothAssetDataflowNodes/Private/ChaosClothAsset/USDImportNode_v2.cpp
+++ b/Engine/Plugins/ChaosClothAssetEditor/Source/ChaosClothAssetDataflowNodes/Private/ChaosClothAsset/USDImportNode_v2.cpp
@@ -1,6 +1,7 @@
 // Copyright Epic Games, Inc. All Rights Reserved.
 
 #include "ChaosClothAsset/USDImportNode_v2.h"
+#include "Chaos/CollectionEmbeddedSpringConstraintFacade.h"
 #include "ChaosClothAsset/ClothAsset.h"
 #include "ChaosClothAsset/ClothCollectionGroup.h"
 #include "ChaosClothAsset/ClothDataflowTools.h"
@@ -50,6 +51,7 @@ namespace UE::Chaos::ClothAsset::Private
 	static const FName SimMeshDataAPI(TEXT("SimMeshDataAPI"));
 	static const FName SimPatternAPI(TEXT("SimPatternAPI"));
 	static const FName SewingAPI(TEXT("SewingAPI"));
+	static const FName SpringAPI(TEXT("SpringAPI"));
 	static const FName CloSolverPropertiesAPI(TEXT("CloSolverPropertiesAPI"));
 
 	// USD import material overrides
@@ -614,6 +616,145 @@ namespace UE::Chaos::ClothAsset::Private
 		return true;
 	}
 
+	static bool ImportSpringsFromSimMeshPrim(
+		const FUsdPrim& SimMeshPrim,
+		FManagedArrayCollection& OutSimulationCollection,
+		FText& OutErrorText)
+	{
+		using namespace ::Chaos::Softs;
+
+		const TSharedRef<FManagedArrayCollection> SimulationCollection = MakeShared<FManagedArrayCollection>(MoveTemp(OutSimulationCollection));
+		ON_SCOPE_EXIT
+		{
+			OutSimulationCollection = MoveTemp(*SimulationCollection);
+		};
+
+		TOptional<FEmbeddedSpringConstraintFacade> SpringConstraintFacade;
+		TOptional<FCollectionPropertyMutableFacade> PropertyFacade;
+
+		for (const FUsdPrim& SimMeshChildPrim : SimMeshPrim.GetChildren())
+		{
+			if (SimMeshChildPrim.IsA(TEXT("GeomSubset")) && SimMeshChildPrim.HasAPI(SpringAPI))
+			{
+				const FUsdGeomSubset GeomSubset(SimMeshChildPrim);
+
+				// Read FamilyName
+				const FUsdAttribute FamilyNameAttr = GeomSubset.GetFamilyNameAttr();
+				if (GetStringValue(FamilyNameAttr) != TEXT("spring"))
+				{
+					OutErrorText = LOCTEXT("WrongSpringGeomSubsetFamilyName", "Wrong spring GeomSubset family name. Needs to be 'spring'.");
+					return false;
+				}
+
+				// Read ElementTypeAttr
+				const FUsdAttribute ElementTypeAttr = GeomSubset.GetElementTypeAttr();
+				if (GetStringValue(ElementTypeAttr) != TEXT("edge"))
+				{
+					OutErrorText = LOCTEXT("WrongSpringGeomSubsetType", "Wrong spring GeomSubset type. Needs to be edge.");
+					return false;
+				}
+
+				// Read indices 
+				const FUsdAttribute IndicesAttr = GeomSubset.GetIndicesAttr();
+				if (IndicesAttr.GetTypeName() != TEXT("int[]"))
+				{
+					OutErrorText = LOCTEXT("WrongSpringGeomSubsetIndexType", "Wrong spring GeomSubset index type. Needs to be int[].");
+					return false;
+				}
+
+				const TArray<int32> IntArrayValues = GetIntArrayValues(IndicesAttr);
+				const int32 NumSprings = IntArrayValues.Num() / 2;
+				if (NumSprings * 2 != IntArrayValues.Num())
+				{
+					OutErrorText = LOCTEXT("OddSpringGeomSubsetIndices", "Odd number of indices for the spring edges.");
+					return false;
+				}
+
+				if (NumSprings)
+				{
+					// Read the primvars
+					FUsdPrimAttributeAccessor UsdPrimAttributeAccessor(SimMeshChildPrim);
+					const TArray<float> RestLengths = UsdPrimAttributeAccessor.GetArray<float>(TEXT("primvars:restLength"));
+					if (RestLengths.Num() != NumSprings)
+					{
+						OutErrorText = LOCTEXT("MistmatchedSpringGeomSubsetRestLengths", "The number of rest length values doesn't match the number of springs.");
+						return false;
+					}
+					const TArray<float> CloSpringDamp = UsdPrimAttributeAccessor.GetArray<float>(TEXT("primvars:clo:springDamp"));
+					const TArray<float> CloSpringStiffness = UsdPrimAttributeAccessor.GetArray<float>(TEXT("primvars:clo:springStiffness"));
+					const bool bHasCloSpringPrimvars = CloSpringDamp.Num() && CloSpringStiffness.Num();  // Only use the first value of the array, that's how it is currently exported
+
+					// Initialize the spring constraint facade
+					if (!SpringConstraintFacade.IsSet())
+					{
+						FEmbeddedSpringFacade SpringFacade(*SimulationCollection, ClothCollectionGroup::SimVertices3D);
+						checkf(SpringFacade.IsValid(), TEXT("FEmbeddedSpringFacade constructor should have defined the schema."));
+						for (int32 ConstraintIndex = 0; ConstraintIndex < SpringFacade.GetNumSpringConstraints(); ++ConstraintIndex)
+						{
+							FEmbeddedSpringConstraintFacade TempSpringFacade = SpringFacade.GetSpringConstraint(ConstraintIndex);
+							if (TempSpringFacade.GetConstraintEndPointNumIndices() == FUintVector2(1, 1) && TempSpringFacade.GetConstraintName() == TEXT("VertexSpringConstraint"))
+							{
+								SpringConstraintFacade.Emplace(MoveTemp(TempSpringFacade));
+								break;
+							}
+						}
+						if (!SpringConstraintFacade.IsSet())
+						{
+							SpringConstraintFacade.Emplace(SpringFacade.AddGetSpringConstraint());
+							SpringConstraintFacade->Initialize(
+								TConstArrayView<FIntVector2>(),
+								TConstArrayView<float>(),
+								TConstArrayView<float>(),
+								TConstArrayView<float>(),
+								TConstArrayView<float>(),
+								TEXT("VertexSpringConstraint"));
+						}
+					}
+
+					TArray<FIntVector2> ConstraintVertices;
+					ConstraintVertices.Reserve(NumSprings);
+					for (int32 Index = 0; Index < NumSprings; ++Index)
+					{
+						const int32 Index0 = IntArrayValues[Index * 2];
+						const int32 Index1 = IntArrayValues[Index * 2 + 1];
+						ConstraintVertices.Emplace(Index0 <= Index1 ? FIntVector2(Index0, Index1) : FIntVector2(Index1, Index0));
+					}
+
+					SpringConstraintFacade->Append(
+						TConstArrayView<FIntVector2>(ConstraintVertices.GetData(), NumSprings),
+						TConstArrayView<float>(RestLengths.GetData(), NumSprings));
+
+					// Initialize the property facade
+					if (bHasCloSpringPrimvars)
+					{
+						if (!PropertyFacade.IsSet())
+						{
+							PropertyFacade.Emplace(SimulationCollection);
+							PropertyFacade->DefineSchema();
+						}
+
+						auto SetPropertyValue = [&PropertyFacade](const FName PropertyName, const float Value)
+							{
+								constexpr ECollectionPropertyFlags CollectionPropertyFlags = ECollectionPropertyFlags::Interpolable | ECollectionPropertyFlags::Animatable | ECollectionPropertyFlags::Enabled;
+								int32 VertexSpringExtensionStiffnessIndex = PropertyFacade->GetKeyNameIndex(PropertyName);
+								if (VertexSpringExtensionStiffnessIndex == INDEX_NONE)
+								{
+									VertexSpringExtensionStiffnessIndex = PropertyFacade->AddProperty(PropertyName, CollectionPropertyFlags);
+								}
+								PropertyFacade->SetValue(VertexSpringExtensionStiffnessIndex, Value);
+							};
+
+						SetPropertyValue("VertexSpringExtensionStiffness", CloSpringStiffness[0]);
+						SetPropertyValue("VertexSpringCompressionStiffness", CloSpringStiffness[0]);
+						SetPropertyValue("VertexSpringDamping", CloSpringDamp[0]);
+					}
+					break;  // Only does one springs geomesubset since that's how it is currently exported, additional ones would also require creating a new weight map (see fabric code)
+				}
+			}
+		}
+		return true;
+	}
+
 	static bool ImportSimulationProperties(
 		const FUsdPrim& ClothPrim,
 		const FUsdPrim& SimMeshPrim,
@@ -622,9 +763,12 @@ namespace UE::Chaos::ClothAsset::Private
 		TMap<FName, int32>& OutSimPatternFabricIndices,
 		FText& OutErrorText)
 	{
-		OutSimPatternFabricIndices.Reset();
-		
-		const TSharedRef<FManagedArrayCollection> SimulationCollection = MakeShared<FManagedArrayCollection>();
+		const TSharedRef<FManagedArrayCollection> SimulationCollection = MakeShared<FManagedArrayCollection>(MoveTemp(OutSimulationCollection));
+		ON_SCOPE_EXIT
+		{
+			OutSimulationCollection = MoveTemp(*SimulationCollection);
+		};
+
 		FCollectionClothFacade ClothFacade(SimulationCollection);
 		ClothFacade.DefineSchema(EClothCollectionExtendedSchemas::Solvers | EClothCollectionExtendedSchemas::Fabrics);
 
@@ -790,9 +934,6 @@ namespace UE::Chaos::ClothAsset::Private
 			const int32 FabricIndex = ImportFabricProperties(GeomSubset.Value);
 			OutSimPatternFabricIndices.Emplace(GeomSubset.Key, FabricIndex);
 		}
-
-		OutSimulationCollection = MoveTemp(*SimulationCollection);
-
 		return true;
 	}
 }  // End namespace Private
@@ -1082,6 +1223,12 @@ bool FChaosClothAssetUSDImportNode_v2::ImportUsdFile(const FString& UsdFilePath,
 			return false;
 		}
 
+		// Import the springs
+		if (!ImportSpringsFromSimMeshPrim(SimMeshPrim, UsdClothData.SimulationCollection, OutErrorText))
+		{
+			return false;
+		}
+
 		// Import the sim static mesh selectively by disabling the render prims
 		SetMeshClothChildrenActive([&SimMeshPrim](FUsdPrim& ClothChildPrim) { return ClothChildPrim == SimMeshPrim; });
 		TArray<TObjectPtr<UObject>> ImportedAssets = ImportStaticMeshesFromUsdStage(UsdStage, UsdFilePath, PackagePath);
@@ -1184,6 +1331,7 @@ void FChaosClothAssetUSDImportNode_v2::UpdateImportedAssets(const FString& SimMe
 bool FChaosClothAssetUSDImportNode_v2::ImportSimStaticMesh(const TSharedRef<FManagedArrayCollection> ClothCollection, FText& OutErrorText)
 {
 	using namespace UE::Chaos::ClothAsset;
+	using namespace ::Chaos::Softs;
 
 	FCollectionClothFacade ClothFacade(ClothCollection);
 	check(ClothFacade.IsValid());  // The Cloth Collection schema must be valid at this point
@@ -1437,8 +1585,62 @@ bool FChaosClothAssetUSDImportNode_v2::ImportSimStaticMesh(const TSharedRef<FMan
 			ClothSeamFacade.Initialize(Indices.Array());
 		}
 
+		// Add the springs
+		const TSharedRef<FManagedArrayCollection> UsdClothDataCollection = MakeShared<FManagedArrayCollection>(UsdClothData.SimulationCollection);
+		const FEmbeddedSpringFacade UsdClothDataSpringFacade(*UsdClothDataCollection, ClothCollectionGroup::SimVertices3D);
+		if (UsdClothDataSpringFacade.IsValid())
+		{
+			// Initialize the spring facade
+			FEmbeddedSpringFacade SpringFacade(*ClothCollection, ClothCollectionGroup::SimVertices3D);
+			checkf(SpringFacade.IsValid(), TEXT("FEmbeddedSpringFacade constructor should have defined the schema."));
+			FEmbeddedSpringConstraintFacade SpringConstraintFacade = SpringFacade.AddGetSpringConstraint();
+			SpringConstraintFacade.Initialize(
+				TConstArrayView<FIntVector2>(),
+				TConstArrayView<float>(),
+				TConstArrayView<float>(),
+				TConstArrayView<float>(),
+				TConstArrayView<float>(),
+				TEXT("VertexSpringConstraint"));
+
+			for (int32 ConstraintIndex = 0; ConstraintIndex < UsdClothDataSpringFacade.GetNumSpringConstraints(); ++ConstraintIndex)
+			{
+				const FEmbeddedSpringConstraintFacade UsdClothDataSpringConstraintFacade = UsdClothDataSpringFacade.GetSpringConstraintConst(ConstraintIndex);
+				if (UsdClothDataSpringConstraintFacade.GetConstraintEndPointNumIndices() == FUintVector2(1, 1) &&
+					UsdClothDataSpringConstraintFacade.GetConstraintName() == TEXT("VertexSpringConstraint"))
+				{
+					const TConstArrayView<int32> SimVertex3DLookup = ClothFacade.GetSimVertex3DLookup();
+
+					TArray<TArray<int32>> SourceIndices(UsdClothDataSpringConstraintFacade.GetSourceIndexConst());
+					TArray<TArray<int32>> TargetIndices(UsdClothDataSpringConstraintFacade.GetTargetIndexConst());
+					for (TArray<int32>& SourceIndex : SourceIndices)
+					{
+						SourceIndex[0] = OriginalToNewVertexIndices.IsValidIndex(SourceIndex[0]) && OriginalToNewVertexIndices[SourceIndex[0]] != INDEX_NONE ?
+							SimVertex3DLookup[OriginalToNewVertexIndices[SourceIndex[0]]] : INDEX_NONE;
+					}
+					for (TArray<int32>& TargetIndex : TargetIndices)
+					{
+						TargetIndex[0] = OriginalToNewVertexIndices.IsValidIndex(TargetIndex[0]) && OriginalToNewVertexIndices[TargetIndex[0]] != INDEX_NONE ?
+							SimVertex3DLookup[OriginalToNewVertexIndices[TargetIndex[0]]] : INDEX_NONE;
+					}
+					SpringConstraintFacade.Append(
+						SourceIndices,
+						UsdClothDataSpringConstraintFacade.GetSourceWeightsConst(),
+						TargetIndices,
+						UsdClothDataSpringConstraintFacade.GetTargetWeightsConst(),
+						UsdClothDataSpringConstraintFacade.GetSpringLengthConst());
+				}
+			}
+			// Copy the properties
+			const FCollectionPropertyConstFacade UsdClothDataPropertyFacade(UsdClothDataCollection);
+			if (UsdClothDataPropertyFacade.IsValid())
+			{
+				FCollectionPropertyMutableFacade PropertyFacade(ClothCollection);
+				PropertyFacade.Copy(UsdClothData.SimulationCollection);
+			}
+		}
+
 		// Add the solver properties
-		const FCollectionClothConstFacade SimulationClothFacade(MakeShared<FManagedArrayCollection>(UsdClothData.SimulationCollection));
+		const FCollectionClothConstFacade SimulationClothFacade(UsdClothDataCollection);
 		if (SimulationClothFacade.IsValid(EClothCollectionExtendedSchemas::Solvers))
 		{
 			ClothFacade.DefineSchema(EClothCollectionExtendedSchemas::Solvers);
diff --git a/Engine/Plugins/ChaosClothAssetEditor/Source/ChaosClothAssetDataflowNodes/Private/ChaosClothAsset/UsdPrimAttributeAccessor.h b/Engine/Plugins/ChaosClothAssetEditor/Source/ChaosClothAssetDataflowNodes/Private/ChaosClothAsset/UsdPrimAttributeAccessor.h
--- a/Engine/Plugins/ChaosClothAssetEditor/Source/ChaosClothAssetDataflowNodes/Private/ChaosClothAsset/UsdPrimAttributeAccessor.h
+++ b/Engine/Plugins/ChaosClothAssetEditor/Source/ChaosClothAssetDataflowNodes/Private/ChaosClothAsset/UsdPrimAttributeAccessor.h
@@ -12,7 +12,7 @@ namespace UE::Chaos::ClothAsset
 	class FUsdPrimAttributeAccessor
 	{
 	public:
-		FUsdPrimAttributeAccessor(const FUsdPrim& InUsdPrim, const EUsdUpAxis UsdUpAxis)
+		explicit FUsdPrimAttributeAccessor(const FUsdPrim& InUsdPrim, const EUsdUpAxis UsdUpAxis = EUsdUpAxis::ZAxis) 
 			: UsdPrim(InUsdPrim)
 			, AxesOrder(
 				0,
