Chaos - Removed redundant return on void function.
Chaos Cloth Asset - Fixed typo in error message.

#rnx

diff --git a/Engine/Plugins/ChaosClothAssetEditor/Source/ChaosClothAssetDataflowNodes/Private/ChaosClothAsset/USDImportNode_v2.cpp b/Engine/Plugins/ChaosClothAssetEditor/Source/ChaosClothAssetDataflowNodes/Private/ChaosClothAsset/USDImportNode_v2.cpp
--- a/Engine/Plugins/ChaosClothAssetEditor/Source/ChaosClothAssetDataflowNodes/Private/ChaosClothAsset/USDImportNode_v2.cpp
+++ b/Engine/Plugins/ChaosClothAssetEditor/Source/ChaosClothAssetDataflowNodes/Private/ChaosClothAsset/USDImportNode_v2.cpp
@@ -1570,7 +1570,7 @@ bool FChaosClothAssetUSDImportNode_v2::ImportSimStaticMesh(const TSharedRef<FMan
 			{
 				if (!OriginalToNewVertexIndices.IsValidIndex(Stitch[0]) || !OriginalToNewVertexIndices.IsValidIndex(Stitch[1]))
 				{
-					OutErrorText = LOCTEXT("BadSewingIndex", "An out of renge sewing index has been found.");
+					OutErrorText = LOCTEXT("BadSewingIndex", "An out of range sewing index has been found.");
 					return false;
 				}
 				const int32 StitchIndex0 = OriginalToNewVertexIndices[Stitch[0]];
diff --git a/Engine/Source/Runtime/Experimental/Chaos/Public/Chaos/CollectionEmbeddedSpringConstraintFacade.h b/Engine/Source/Runtime/Experimental/Chaos/Public/Chaos/CollectionEmbeddedSpringConstraintFacade.h
--- a/Engine/Source/Runtime/Experimental/Chaos/Public/Chaos/CollectionEmbeddedSpringConstraintFacade.h
+++ b/Engine/Source/Runtime/Experimental/Chaos/Public/Chaos/CollectionEmbeddedSpringConstraintFacade.h
@@ -111,7 +111,7 @@ namespace Chaos::Softs
 
 		void SetConstraintName(const FString& InName)
 		{
-			return ConstraintNameAttribute.ModifyAt(ConstraintIndex, InName);
+			ConstraintNameAttribute.ModifyAt(ConstraintIndex, InName);
 		}
 
 		//~ Group PositionTargets
