Fix indexing crash in ClothCollectionToDynamicMeshNode when transferring normals and actually transfer normals. There was a copy-paste error which  made it so we likely weren't transferring normals before (and could crash in the case where you have more  vertices than faces).

diff --git a/Engine/Plugins/ChaosClothAssetEditor/Source/ChaosClothAssetDataflowNodes/Private/ChaosClothAsset/ClothCollectionToDynamicMeshNode.cpp b/Engine/Plugins/ChaosClothAssetEditor/Source/ChaosClothAssetDataflowNodes/Private/ChaosClothAsset/ClothCollectionToDynamicMeshNode.cpp
--- a/Engine/Plugins/ChaosClothAssetEditor/Source/ChaosClothAssetDataflowNodes/Private/ChaosClothAsset/ClothCollectionToDynamicMeshNode.cpp
+++ b/Engine/Plugins/ChaosClothAssetEditor/Source/ChaosClothAssetDataflowNodes/Private/ChaosClothAsset/ClothCollectionToDynamicMeshNode.cpp
@@ -124,7 +124,7 @@ void FChaosClothAssetUpdateClothFromDynamicMeshNode::Evaluate(UE::Dataflow::FCon
 									[&Normals, &ClothMapping, &TriangleIndices, ClothVertexID](int32 TriangleID, int32 ElementID, const FVector3f& NormalValue)->bool
 									{
 										const int32 ClothMeshTriID = ClothMapping.GetOriginalTriangleID(TriangleID);
-										const FIntVector& ClothTri = TriangleIndices[ClothVertexID];
+										const FIntVector& ClothTri = TriangleIndices[ClothMeshTriID];
 										for (int32 LocalIndex = 0; LocalIndex < 3; ++LocalIndex)
 										{
 											if (ClothTri[LocalIndex] == ClothVertexID)
