Chaos Cloth Asset - Fixed a crash in the TransferSkinWeight node when a non-manifold mesh is used with the inpaint method, which is now prevented.

[RN] Chaos Cloth Asset - Fixed a crash in the TransferSkinWeight node when a non-manifold mesh is used with the inpaint method, which is now prevented.

diff --git a/Engine/Plugins/ChaosClothAssetEditor/Source/ChaosClothAssetDataflowNodes/Private/ChaosClothAsset/TransferSkinWeightsNode.cpp b/Engine/Plugins/ChaosClothAssetEditor/Source/ChaosClothAssetDataflowNodes/Private/ChaosClothAsset/TransferSkinWeightsNode.cpp
--- a/Engine/Plugins/ChaosClothAssetEditor/Source/ChaosClothAssetDataflowNodes/Private/ChaosClothAsset/TransferSkinWeightsNode.cpp
+++ b/Engine/Plugins/ChaosClothAssetEditor/Source/ChaosClothAssetDataflowNodes/Private/ChaosClothAsset/TransferSkinWeightsNode.cpp
@@ -616,6 +616,13 @@ namespace UE::Chaos::ClothAsset::Private
 		return true;
 	}
 
+	static bool TransferClosestPointOnSurface(
+		const FReferenceSkeleton& TargetRefSkeleton,
+		const FDynamicMesh3& SkeletalDynamicMesh,
+		const TSharedRef<FManagedArrayCollection>& ClothCollection,
+		const TransferBoneWeightsSettings& TransferSettings,
+		const TSharedPtr<const FManagedArrayCollection>& SimClothCollection);
+
 	/** Transfer skin weights to render cloth. */
 	static bool TransferInpaintWeightsToRender(
 		const FReferenceSkeleton& TargetRefSkeleton,
@@ -633,11 +640,25 @@ namespace UE::Chaos::ClothAsset::Private
 			return false;
 		}
 
-		// Transfer weights to render mesh
-		if (!TransferInpaintWeightsToMeshComponents(TargetRefSkeleton, SourceDynamicMesh, RenderDynamicMesh, TransferSettings))
+		const FNonManifoldMappingSupport NonManifoldMapping(RenderDynamicMesh);
+		if (NonManifoldMapping.IsNonManifoldVertexInSource())
 		{
-			UE_LOG(LogChaosClothAssetDataflowNodes, Warning, TEXT("TransferSkinWeightsNode: Failed to transfer skin weights to render mesh."));
-			return false;
+			UE_LOG(LogChaosClothAssetDataflowNodes, Warning, TEXT("TransferSkinWeightsNode: Render mesh inpaint method failed due to non manifold geometry, falling back to closest point method instead."));
+			checkf(!TransferSettings.bTransferToRenderFromSim, TEXT("Logic error, TransferClosestPointOnSurface should have been called instead transferring from Sim mesh."));
+			TransferBoneWeightsSettings NewTransferSettings = TransferSettings;
+			NewTransferSettings.bTransferToSim = false;
+			NewTransferSettings.TransferMethod = FTransferBoneWeights::ETransferBoneWeightsMethod::ClosestPointOnSurface;
+			TSharedPtr<const FManagedArrayCollection> EmptySimClothCollection;
+			return TransferClosestPointOnSurface(TargetRefSkeleton, SourceDynamicMesh, ClothCollection, NewTransferSettings, EmptySimClothCollection);
+		}
+		else
+		{
+			// Transfer weights to render mesh
+			if (!TransferInpaintWeightsToMeshComponents(TargetRefSkeleton, SourceDynamicMesh, RenderDynamicMesh, TransferSettings))
+			{
+				UE_LOG(LogChaosClothAssetDataflowNodes, Warning, TEXT("TransferSkinWeightsNode: Failed to transfer skin weights to render mesh."));
+				return false;
+			}
 		}
 
 		CopySkinWeightsFromDynamicMeshToRenderCloth(RenderDynamicMesh, TransferSettings.bUseParallel, ClothCollection);
@@ -645,13 +666,6 @@ namespace UE::Chaos::ClothAsset::Private
 		return true;
 	}
 
-	static bool TransferClosestPointOnSurface(
-		const FReferenceSkeleton& TargetRefSkeleton,
-		const FDynamicMesh3& SkeletalDynamicMesh,
-		const TSharedRef<FManagedArrayCollection>& ClothCollection,
-		const TransferBoneWeightsSettings& TransferSettings,
-		const TSharedPtr<const FManagedArrayCollection>& SimClothCollection);
-
 	/** Transfer skin weights to sim and render cloth. */
 	static bool TransferInpaintWeights(
 		const FReferenceSkeleton& TargetRefSkeleton,
