Dataflow : make sure we mark the dataflow package dirty when the tools are commiting their values
- before this change editing a map without doing anything in the dataflow woudl not mark the package dirty and only modify the asset in memory but not on disk after saving

diff --git a/Engine/Plugins/ChaosClothAssetEditor/Source/ChaosClothAssetEditorTools/Private/ChaosClothAsset/ClothMeshSelectionTool.cpp b/Engine/Plugins/ChaosClothAssetEditor/Source/ChaosClothAssetEditorTools/Private/ChaosClothAsset/ClothMeshSelectionTool.cpp
--- a/Engine/Plugins/ChaosClothAssetEditor/Source/ChaosClothAssetEditorTools/Private/ChaosClothAsset/ClothMeshSelectionTool.cpp
+++ b/Engine/Plugins/ChaosClothAssetEditor/Source/ChaosClothAssetEditorTools/Private/ChaosClothAsset/ClothMeshSelectionTool.cpp
@@ -541,6 +541,8 @@ void UClothMeshSelectionTool::UpdateSelectedNode()
 		GetToolManager()->GetContextTransactionsAPI()->AppendChange(Dataflow, 
 			FChaosClothAssetSelectionNode_v2::MakeSelectedNodeChange(*SelectionNodeToUpdate),
 			LOCTEXT("SelectionNodeChangeDescription", "Update Selection Node"));
+
+		Dataflow->MarkPackageDirty();
 	}
 
 	const UE::Geometry::FGroupTopologySelection& Selection = SelectionMechanic->GetActiveSelection();
diff --git a/Engine/Plugins/ChaosClothAssetEditor/Source/ChaosClothAssetEditorTools/Private/ChaosClothAsset/ClothTransferSkinWeightsTool.cpp b/Engine/Plugins/ChaosClothAssetEditor/Source/ChaosClothAssetEditorTools/Private/ChaosClothAsset/ClothTransferSkinWeightsTool.cpp
--- a/Engine/Plugins/ChaosClothAssetEditor/Source/ChaosClothAssetEditorTools/Private/ChaosClothAsset/ClothTransferSkinWeightsTool.cpp
+++ b/Engine/Plugins/ChaosClothAssetEditor/Source/ChaosClothAssetEditorTools/Private/ChaosClothAsset/ClothTransferSkinWeightsTool.cpp
@@ -10,6 +10,7 @@
 #include "ContextObjectStore.h"
 #include "Dataflow/DataflowContextObject.h"
 #include "Dataflow/DataflowEdNode.h"
+#include "Dataflow/DataflowObject.h"
 #include "ChaosClothAsset/TransferSkinWeightsNode.h"
 #include "BaseGizmos/TransformGizmoUtil.h"
 #include "InteractiveToolObjects.h"
@@ -192,6 +193,13 @@ void UClothTransferSkinWeightsTool::Shutdown(EToolShutdownType ShutdownType)
 		TransferSkinWeightsNode->SkeletalMesh = ToolProperties->SourceMesh;
 		TransferSkinWeightsNode->Transform = TransformFromProperties();
 		TransferSkinWeightsNode->Invalidate();
+		if (DataflowContextObject)
+		{
+			if (TObjectPtr<UDataflow> const Dataflow = DataflowContextObject->GetDataflowAsset())
+			{
+				Dataflow->MarkPackageDirty();
+			}
+		}
 	}
 
 	if (SourceMeshTransformProxy)
diff --git a/Engine/Plugins/Experimental/Dataflow/Source/DataflowEditor/Private/DataflowEditorTools/DataflowEditorSkinWeightsPaintTool.cpp b/Engine/Plugins/Experimental/Dataflow/Source/DataflowEditor/Private/DataflowEditorTools/DataflowEditorSkinWeightsPaintTool.cpp
--- a/Engine/Plugins/Experimental/Dataflow/Source/DataflowEditor/Private/DataflowEditorTools/DataflowEditorSkinWeightsPaintTool.cpp
+++ b/Engine/Plugins/Experimental/Dataflow/Source/DataflowEditor/Private/DataflowEditorTools/DataflowEditorSkinWeightsPaintTool.cpp
@@ -321,6 +321,7 @@ void UDataflowEditorSkinWeightsPaintTool::OnShutdown(EToolShutdownType ShutdownT
 						);
 					}
 				}
+				Dataflow->MarkPackageDirty();
 			}
 
 			TArray<TArray<int32>> CurrentIndices;
diff --git a/Engine/Plugins/Experimental/Dataflow/Source/DataflowEditor/Private/DataflowEditorTools/DataflowEditorVertexAttributePaintTool.cpp b/Engine/Plugins/Experimental/Dataflow/Source/DataflowEditor/Private/DataflowEditorTools/DataflowEditorVertexAttributePaintTool.cpp
--- a/Engine/Plugins/Experimental/Dataflow/Source/DataflowEditor/Private/DataflowEditorTools/DataflowEditorVertexAttributePaintTool.cpp
+++ b/Engine/Plugins/Experimental/Dataflow/Source/DataflowEditor/Private/DataflowEditorTools/DataflowEditorVertexAttributePaintTool.cpp
@@ -558,6 +558,7 @@ void FDataflowEditorVertexAttributePaintToolData::UpdateNode(TObjectPtr<UDataflo
 							LOCTEXT("DataflowEditorVertexAttributePaintTool_ChangeDescription", "Update Weight Map Node")
 						);
 					}
+					Dataflow->MarkPackageDirty();
 				}
 
 				// apply the new values to the node
