Chaos Cloth Asset - Fixed an infinite loop while merging weightmaps using the MergeClothCollections node.
Regression caused by the change to FName in CL 44486948.

[RN] Chaos Cloth Asset - Fixed an infinite loop while merging weightmaps using the MergeClothCollections node.

diff --git a/Engine/Plugins/ChaosClothAssetEditor/Source/ChaosClothAssetDataflowNodes/Private/ChaosClothAsset/MergeClothCollectionsNode.cpp b/Engine/Plugins/ChaosClothAssetEditor/Source/ChaosClothAssetDataflowNodes/Private/ChaosClothAsset/MergeClothCollectionsNode.cpp
--- a/Engine/Plugins/ChaosClothAssetEditor/Source/ChaosClothAssetDataflowNodes/Private/ChaosClothAsset/MergeClothCollectionsNode.cpp
+++ b/Engine/Plugins/ChaosClothAssetEditor/Source/ChaosClothAssetDataflowNodes/Private/ChaosClothAsset/MergeClothCollectionsNode.cpp
@@ -100,7 +100,7 @@ namespace UE::Chaos::ClothAsset::Private
 		// the weight map could already been stored on the out collection and linked to different bounds
 		// coming from the out collection itself or from previous merge with in collection
 		// Since we don't want to break them we need to create a new one on the first available slot
-		while(OutClothFacade.GetWeightMap(PropertyName).Num() > 0)
+		while(OutClothFacade.GetWeightMap(FName(WeightMapName)).Num() > 0)
 		{
 			WeightMapName = PropertyName.ToString();
 			WeightMapName.AppendInt(++WeightMapCount);
