Cloth Remesh Node: rollback behavior changes in V1 of the node

#Synced-CL 48155487

diff --git a/Engine/Plugins/ChaosClothAssetEditor/Source/ChaosClothAssetDataflowNodes/Private/ChaosClothAsset/RemeshNode.cpp b/Engine/Plugins/ChaosClothAssetEditor/Source/ChaosClothAssetDataflowNodes/Private/ChaosClothAsset/RemeshNode.cpp
--- a/Engine/Plugins/ChaosClothAssetEditor/Source/ChaosClothAssetDataflowNodes/Private/ChaosClothAsset/RemeshNode.cpp
+++ b/Engine/Plugins/ChaosClothAssetEditor/Source/ChaosClothAssetDataflowNodes/Private/ChaosClothAsset/RemeshNode.cpp
@@ -1923,15 +1923,1014 @@ void FChaosClothAssetRemeshNode_v2::Evaluate(UE::Dataflow::FContext& Context, co
 ///////////////////////////////////////////////////////////////////////////////////////////////////////////////////////////////
 
 
+
+namespace UE::Chaos::ClothAsset::RemeshNodeV1::Private
+{
+
+	using UE::Chaos::ClothAsset::Private::FSeamCollapseParameters;
+	using UE::Chaos::ClothAsset::Private::FSeamEdge;
+	using UE::Chaos::ClothAsset::Private::IsSharpCorner;
+	using UE::Chaos::ClothAsset::Private::GetSeamEdgeCollapseParameters;
+	using UE::Chaos::ClothAsset::Private::CanSplitSeamEdge;
+	using UE::Chaos::ClothAsset::Private::FindCoincidentBoundaryVertices;
+	using UE::Chaos::ClothAsset::Private::FEdgeCollapseParameters;
+	using UE::Chaos::ClothAsset::Private::GetBoundaryEdgeCollapseParameters;
+
+
+	/** Given weight values for the edge vertices, compute a new edge length multiplier to be used in the remesher.
+		A greater edge length means the edge is more likely to be split, resulting in a higher overall vertex density */
+		// TODO: Experiment with other forms
+		// TODO: Add parameters to rescale the 0-1 weight values
+	static double EdgeScalingFunction(float WeightValueA, float WeightValueB)
+	{
+		const float Avg = 0.5f * (WeightValueA + WeightValueB);
+		const float ClampedWeight = FMath::Clamp(Avg, 0.0f, 1.0f);
+		const double LengthScale = 1.0f + ClampedWeight;
+		return LengthScale;
+	}
+
+	//
+	// Functions to support seam remeshing
+	//
+
+	static void RemeshSeams(UE::Geometry::FDynamicMesh3& Mesh, TArray<TArray<FIntVector2>>& Seams, double TargetEdgeLength, const UE::Geometry::FDynamicMeshWeightAttribute* const DensityMapLayer)
+	{
+		using namespace UE::Geometry;
+
+		// constants pulled from FRemesher::SetTargetEdgeLength
+		const double MinLength = 0.66 * TargetEdgeLength;
+		const double MaxLength = 1.33 * TargetEdgeLength;
+
+		for (int32 SeamID = 0; SeamID < Seams.Num(); ++SeamID)
+		{
+			TArray<FIntVector2>& Seam = Seams[SeamID];
+
+			for (int32 StitchID = 0; StitchID < Seam.Num() - 1; ++StitchID)
+			{
+				const int32 SideAVertexA = Seam[StitchID][0];
+				const int32 SideAVertexB = Seam[StitchID + 1][0];
+
+				if (Mesh.FindEdge(SideAVertexA, SideAVertexB) == UE::Geometry::FDynamicMesh3::InvalidID)
+				{
+					continue;
+				}
+
+				double EdgeALength = (Mesh.GetVertex(SideAVertexA) - Mesh.GetVertex(SideAVertexB)).Length();
+
+				if (DensityMapLayer)
+				{
+					float WeightValueA;
+					DensityMapLayer->GetValue(SideAVertexA, &WeightValueA);
+					float WeightValueB;
+					DensityMapLayer->GetValue(SideAVertexB, &WeightValueB);
+
+					EdgeALength = EdgeScalingFunction(WeightValueA, WeightValueB) * EdgeALength;
+				}
+
+				const int32 SideBVertexA = Seam[StitchID][1];
+				const int32 SideBVertexB = Seam[StitchID + 1][1];
+
+				if (Mesh.FindEdge(SideBVertexA, SideBVertexB) == UE::Geometry::FDynamicMesh3::InvalidID)
+				{
+					continue;
+				}
+
+				double EdgeBLength = (Mesh.GetVertex(SideBVertexA) - Mesh.GetVertex(SideBVertexB)).Length();
+
+				if (DensityMapLayer)
+				{
+					float WeightValueA;
+					DensityMapLayer->GetValue(SideBVertexA, &WeightValueA);
+					float WeightValueB;
+					DensityMapLayer->GetValue(SideBVertexB, &WeightValueB);
+
+					EdgeBLength = EdgeScalingFunction(WeightValueA, WeightValueB) * EdgeBLength;
+				}
+
+				if ((EdgeALength < MinLength) && (EdgeBLength < MinLength))
+				{
+					//
+					// Collapse
+					//
+
+					const FSeamCollapseParameters CollapseConstraints = GetSeamEdgeCollapseParameters(SeamID, StitchID, Mesh, Seams);
+
+					if (!CollapseConstraints.bCanCollapse)
+					{
+						continue;
+					}
+
+					const int32 PatternAKeepVertex = Seam[CollapseConstraints.KeepStitchIndex][0];
+					const int32 PatternADeleteVertex = Seam[CollapseConstraints.DeleteStitchIndex][0];
+					const int32 PatternBKeepVertex = Seam[CollapseConstraints.KeepStitchIndex][1];
+					const int32 PatternBDeleteVertex = Seam[CollapseConstraints.DeleteStitchIndex][1];
+
+					FDynamicMesh3::FEdgeCollapseInfo CollapseInfoA;
+					const EMeshResult ResultA = Mesh.CollapseEdge(PatternAKeepVertex, PatternADeleteVertex, CollapseConstraints.NewPositionEdgeParameter, CollapseInfoA);
+
+					FDynamicMesh3::FEdgeCollapseInfo CollapseInfoB;
+					const EMeshResult ResultB = Mesh.CollapseEdge(PatternBKeepVertex, PatternBDeleteVertex, CollapseConstraints.NewPositionEdgeParameter, CollapseInfoB);
+
+					check(ResultA == EMeshResult::Ok && ResultB == EMeshResult::Ok)
+
+						Seam.RemoveAt(CollapseConstraints.DeleteStitchIndex);
+				}
+				else if ((EdgeALength > MaxLength) && (EdgeBLength > MaxLength))
+				{
+					//
+					// Split
+					//
+
+
+					const bool bCanSplit = CanSplitSeamEdge(SeamID, StitchID, Seams);
+
+					if (!bCanSplit)
+					{
+						continue;
+					}
+
+					FDynamicMesh3::FEdgeSplitInfo SplitInfoA;
+					const EMeshResult ResultA = Mesh.SplitEdge(SideAVertexA, SideAVertexB, SplitInfoA);
+
+					if (ResultA == EMeshResult::Ok)
+					{
+						if (Mesh.FindEdge(SideBVertexA, SideBVertexB) == UE::Geometry::FDynamicMesh3::InvalidID)
+						{
+							// Don't split the same edge twice
+							continue;
+						}
+
+						FDynamicMesh3::FEdgeSplitInfo SplitInfoB;
+						const EMeshResult ResultB = Mesh.SplitEdge(SideBVertexA, SideBVertexB, SplitInfoB);
+						check(ResultB == EMeshResult::Ok);
+
+						const FIntVector2 NewStitch{ SplitInfoA.NewVertex, SplitInfoB.NewVertex };
+
+						Seam.Insert(NewStitch, StitchID + 1);
+					}
+				}
+			}
+		}
+	}
+
+
+	//
+	// Boundary remeshing
+	//
+
+	static void RemeshBoundaries(UE::Geometry::FDynamicMesh3& Mesh, const TArray<TArray<FIntVector2>>& Seams, double TargetEdgeLength, const UE::Geometry::FDynamicMeshWeightAttribute* const DensityMapLayer)
+	{
+		using namespace UE::Geometry;
+
+		auto IsSeamEdge = [](const TArray<TArray<FIntVector2>>& Seams, int EdgeVertexA, int EdgeVertexB)
+			{
+				for (const TArray<FIntVector2>& Seam : Seams)
+				{
+					for (int32 StitchID = 0; StitchID < Seam.Num() - 1; ++StitchID)
+					{
+						const FIntVector2& Stitch = Seam[StitchID];
+						const FIntVector2& NextStitch = Seam[StitchID + 1];
+
+						for (int32 Side = 0; Side < 2; ++Side)
+						{
+							if ((int)Stitch[Side] == EdgeVertexA && (int)NextStitch[Side] == EdgeVertexB)
+							{
+								return true;
+							}
+
+							if ((int)Stitch[Side] == EdgeVertexB && (int)NextStitch[Side] == EdgeVertexA)
+							{
+								return true;
+							}
+						}
+					}
+				}
+
+				return false;
+			};
+
+		// constants pulled from FRemesher::SetTargetEdgeLength
+		const double MinLength = 0.66 * TargetEdgeLength;
+		const double MaxLength = 1.33 * TargetEdgeLength;
+
+		// Get the set of boundary edges up front and then check if they get invalidated later. 
+		// If we process edges inside this loop it tends to collapse a bunch of sequential edges, which can lead to very high-valence vertices
+		TArray<FIndex2i> BoundaryEdges;
+		for (const int EdgeID : Mesh.BoundaryEdgeIndicesItr())
+		{
+			BoundaryEdges.Add(Mesh.GetEdge(EdgeID).Vert);
+		}
+
+		for (const FIndex2i& EdgeVerts : BoundaryEdges)
+		{
+			const int EdgeID = Mesh.FindEdge(EdgeVerts[0], EdgeVerts[1]);
+			if (EdgeID == FDynamicMesh3::InvalidID || !Mesh.IsBoundaryEdge(EdgeID))
+			{
+				continue;
+			}
+
+			if (IsSeamEdge(Seams, EdgeVerts[0], EdgeVerts[1]))
+			{
+				continue;
+			}
+
+			double EdgeLength = (Mesh.GetVertex(EdgeVerts[0]) - Mesh.GetVertex(EdgeVerts[1])).Length();
+
+			if (DensityMapLayer)
+			{
+				float WeightValueA;
+				DensityMapLayer->GetValue(EdgeVerts[0], &WeightValueA);
+				float WeightValueB;
+				DensityMapLayer->GetValue(EdgeVerts[1], &WeightValueB);
+
+				EdgeLength = EdgeScalingFunction(WeightValueA, WeightValueB) * EdgeLength;
+			}
+
+
+			if (EdgeLength < MinLength)
+			{
+				const FEdgeCollapseParameters CollapseParams = GetBoundaryEdgeCollapseParameters(Mesh, Seams, EdgeVerts);
+
+				if (CollapseParams.bCanCollapse)
+				{
+					FDynamicMesh3::FEdgeCollapseInfo CollapseInfo;
+					EMeshResult CollapseResult = Mesh.CollapseEdge(CollapseParams.KeepVertexIndex, CollapseParams.DeleteVertexIndex, CollapseParams.NewPositionEdgeParameter, CollapseInfo);
+				}
+			}
+			else if (EdgeLength > MaxLength)
+			{
+				FDynamicMesh3::FEdgeSplitInfo SplitInfo;
+				EMeshResult SplitResult = Mesh.SplitEdge(EdgeVerts[0], EdgeVerts[1], SplitInfo);
+			}
+		};
+	}
+
+
+	//
+	// Remeshing away from Seams/Boundaries
+	// 
+
+	static bool Remesh(UE::Geometry::FDynamicMesh3& Mesh, double TargetEdgeLength, int32 Iterations, float SmoothingRate, bool bUniformSmoothing, const TArray<TArray<FIntVector2>>& Seams, const FString& DensityMapName, UE::Geometry::FCompactMaps* CompactMaps = nullptr)
+	{
+		using namespace UE::Geometry;
+
+		//
+		// These consts control overall remeshing behavior and are analogs of the properties exposed to the user in the actual Remesh tool
+		// (i.e. things we might wish to add to the node properties in the future)
+		//
+
+		constexpr bool bReprojectToInputMesh = true;
+		constexpr bool bDiscardAttributes = false;
+		constexpr bool bUseFullRemeshPasses = true;
+		constexpr bool bAllowFlips = true;
+		constexpr bool bAllowSplits = true;
+		constexpr bool bAllowCollapses = true;
+		constexpr bool bPreventNormalFlips = true;
+		constexpr bool bPreventTinyTriangles = true;
+		constexpr bool bAutoCompact = true;
+		constexpr bool bCoarsenBoundaries = false;
+
+		const ERemeshSmoothingType SmoothingType = bUniformSmoothing ? ERemeshSmoothingType::Uniform : ERemeshSmoothingType::MeanValue;
+
+		// Mesh seam behavior
+		// Here we are talking UV, Normal, and Color seams, not Cloth Seams
+		// This controls bAllowSeamCollapse and bAllowSeamSmoothing on those overlay seams
+		// NOTE: These seams are not affected by bReprojectConstraints
+		constexpr bool bPreserveSharpEdges = false;
+
+		// Mesh boundaries
+		const EEdgeRefineFlags MeshBoundaryConstraint = bCoarsenBoundaries ? EEdgeRefineFlags::NoFlip : EEdgeRefineFlags::FullyConstrained;
+
+		// Group ID boundaries
+		const EEdgeRefineFlags PolyGroupBoundaryConstraint = bCoarsenBoundaries ? EEdgeRefineFlags::NoFlip : EEdgeRefineFlags::FullyConstrained;
+
+		// Material ID boundaries
+		const EEdgeRefineFlags MaterialBoundaryConstraint = bCoarsenBoundaries ? EEdgeRefineFlags::NoFlip : EEdgeRefineFlags::FullyConstrained;
+
+		// Whether to move boundary vertices back onto the poly-line defined by the original boundary in case of collapse
+		const bool bReprojectConstraints = bCoarsenBoundaries;
+
+		// Seam "corners" are held fixed. We use this angle threshold to determine what constitutes a seam corner
+		constexpr float SeamCornerThresholdAngleDegrees = 45.0;
+
+		FRemeshMeshOp RemeshOp;
+
+		TSharedPtr<FDynamicMesh3> SourceMesh = MakeShared<FDynamicMesh3>(MoveTemp(Mesh));
+		TSharedPtr<FDynamicMeshAABBTree3> SourceSpatial;
+		if (bReprojectToInputMesh)
+		{
+			// acceleration structure is only used for reprojecting
+			SourceSpatial = MakeShared<FDynamicMeshAABBTree3>(SourceMesh.Get(), true);
+		}
+
+		RemeshOp.OriginalMesh = SourceMesh;
+		RemeshOp.OriginalMeshSpatial = SourceSpatial;
+
+		RemeshOp.bDiscardAttributes = bDiscardAttributes;
+		RemeshOp.RemeshType = (bUseFullRemeshPasses) ? ERemeshType::FullPass : ERemeshType::Standard;
+		RemeshOp.RemeshIterations = Iterations;
+		RemeshOp.MaxRemeshIterations = Iterations;
+		RemeshOp.ExtraProjectionIterations = 0;		// unused for regular remeshing
+		RemeshOp.TriangleCountHint = 0;				// unused for regular remeshing
+		RemeshOp.SmoothingStrength = FMath::Clamp(SmoothingRate, 0.0f, 1.0f);
+		RemeshOp.SmoothingType = SmoothingType;
+
+		RemeshOp.TargetEdgeLength = TargetEdgeLength;
+		RemeshOp.bPreserveSharpEdges = bPreserveSharpEdges;
+		RemeshOp.bFlips = bAllowFlips;
+		RemeshOp.bSplits = bAllowSplits;
+		RemeshOp.bCollapses = bAllowCollapses;
+		RemeshOp.bPreventNormalFlips = bPreventNormalFlips;
+		RemeshOp.bPreventTinyTriangles = bPreventTinyTriangles;
+		RemeshOp.MeshBoundaryConstraint = MeshBoundaryConstraint;
+		RemeshOp.GroupBoundaryConstraint = PolyGroupBoundaryConstraint;
+		RemeshOp.MaterialBoundaryConstraint = MaterialBoundaryConstraint;
+		RemeshOp.bReproject = bReprojectToInputMesh;
+		RemeshOp.ProjectionTarget = SourceMesh.Get();
+		RemeshOp.ProjectionTargetSpatial = SourceSpatial.Get();
+		RemeshOp.bReprojectConstraints = bReprojectConstraints;
+		RemeshOp.BoundaryCornerAngleThreshold = SeamCornerThresholdAngleDegrees;
+		RemeshOp.TargetMeshLocalToWorld = FTransformSRT3d::Identity();
+		RemeshOp.ToolMeshLocalToWorld = FTransformSRT3d::Identity();
+		RemeshOp.bUseWorldSpace = false;
+		RemeshOp.bParallel = true;
+
+		// RemeshOp makes a copy of the mesh to operate on, so we can't just pass the FDynamicMeshWeightAttribute pointer into CustomEdgeLengthScaleF
+		int32 DensityMapLayerIndex = -1;
+		bool bFoundDensityMapLayer = false;
+		if (!DensityMapName.IsEmpty() && SourceMesh->HasAttributes())
+		{
+			for (int32 WeightLayerIndex = 0; WeightLayerIndex < SourceMesh->Attributes()->NumWeightLayers(); ++WeightLayerIndex)
+			{
+				if (SourceMesh->Attributes()->GetWeightLayer(WeightLayerIndex)->GetName() == DensityMapName)
+				{
+					DensityMapLayerIndex = WeightLayerIndex;
+					bFoundDensityMapLayer = true;
+					break;
+				}
+			}
+		}
+
+		RemeshOp.CustomEdgeLengthScaleF = [bFoundDensityMapLayer, DensityMapLayerIndex](const FDynamicMesh3& Mesh, int VertexA, int VertexB) -> double
+			{
+				if (bFoundDensityMapLayer)
+				{
+					check(Mesh.HasAttributes());
+					const FDynamicMeshWeightAttribute* const DensityMapLayer = Mesh.Attributes()->GetWeightLayer(DensityMapLayerIndex);
+					check(DensityMapLayer);
+
+					float WeightValueA;
+					DensityMapLayer->GetValue(VertexA, &WeightValueA);
+
+					float WeightValueB;
+					DensityMapLayer->GetValue(VertexB, &WeightValueB);
+
+					return EdgeScalingFunction(WeightValueA, WeightValueB);
+				}
+
+				return 1.0;
+			};
+
+		// Set up constraints for Cloth Seam edges
+		UE::Geometry::FMeshConstraints Constraints;
+		for (const TArray<FIntVector2>& Seam : Seams)
+		{
+			if (Seam.Num() == 1)
+			{
+				for (int32 Side = 0; Side < 2; ++Side)
+				{
+					constexpr bool bCannotDelete = true;
+					constexpr bool bCanMove = false;
+					FVertexConstraint VertexConstraint(bCannotDelete, bCanMove);
+					Constraints.SetOrCombineVertexConstraint(Seam[0][Side], VertexConstraint);
+				}
+			}
+			else
+			{
+				for (int32 StitchIndex = 0; StitchIndex < Seam.Num() - 1; ++StitchIndex)
+				{
+					for (int32 Side = 0; Side < 2; ++Side)
+					{
+						const int EdgeID = SourceMesh->FindEdge(Seam[StitchIndex][Side], Seam[StitchIndex + 1][Side]);
+						check(EdgeID != FDynamicMesh3::InvalidID);
+						FEdgeConstraint EdgeConstraint(EEdgeRefineFlags::FullyConstrained);
+						Constraints.SetOrUpdateEdgeConstraint(EdgeID, EdgeConstraint);
+
+						constexpr bool bCannotDelete = true;
+						constexpr bool bCanMove = false;
+						FVertexConstraint VertexConstraint(bCannotDelete, bCanMove);
+						Constraints.SetOrCombineVertexConstraint(Seam[StitchIndex][Side], VertexConstraint);
+						Constraints.SetOrCombineVertexConstraint(Seam[StitchIndex + 1][Side], VertexConstraint);
+					}
+				}
+			}
+		}
+		RemeshOp.SetUserSpecifiedConstraints(Constraints);
+
+		constexpr FProgressCancel* Progress = nullptr;		// Don't allow cancel or report progress for now
+		RemeshOp.CalculateResult(Progress);
+
+		if (RemeshOp.GetResultInfo().Result == EGeometryResultType::Success)
+		{
+			TUniquePtr<FDynamicMesh3> ResultMesh = RemeshOp.ExtractResult();
+			Mesh = MoveTemp(*ResultMesh);
+		}
+		else
+		{
+			return false;
+		}
+
+		// compact the input mesh if enabled
+		if (bAutoCompact)
+		{
+			Mesh.CompactInPlace(CompactMaps);
+		}
+
+		return true;
+	}
+
+}
+
+
 FChaosClothAssetRemeshNode::FChaosClothAssetRemeshNode(const UE::Dataflow::FNodeParameters& InParam, FGuid InGuid)
 	: FDataflowNode(InParam, InGuid)
 {
-	RegisterInputConnection(&Collection);
-	RegisterOutputConnection(&Collection, &Collection);
-	RegisterInputConnection(&DensityMapSim.StringValue, GET_MEMBER_NAME_CHECKED(FChaosClothAssetConnectableIStringValue, StringValue));
-	RegisterInputConnection(&DensityMapRender.StringValue, GET_MEMBER_NAME_CHECKED(FChaosClothAssetConnectableIStringValue, StringValue));
+	RegisterInputConnection(&Collection);
+	RegisterOutputConnection(&Collection, &Collection);
+	RegisterInputConnection(&DensityMapSim.StringValue, GET_MEMBER_NAME_CHECKED(FChaosClothAssetConnectableIStringValue, StringValue));
+	RegisterInputConnection(&DensityMapRender.StringValue, GET_MEMBER_NAME_CHECKED(FChaosClothAssetConnectableIStringValue, StringValue));
+}
+
+
+void FChaosClothAssetRemeshNode::RemeshSimMesh(const TSharedRef<const FManagedArrayCollection>& ClothCollection,
+	const FString& DensityMapName,
+	const TSharedRef<FManagedArrayCollection>& OutClothCollection) const
+{
+	using namespace UE::Geometry;
+	using namespace UE::Chaos::ClothAsset;
+
+	FCollectionClothConstFacade InClothFacade(ClothCollection);
+
+	if (InClothFacade.GetNumSimPatterns() == 0)
+	{
+		FClothGeometryTools::DeleteSimMesh(OutClothCollection);
+		return;
+	}
+
+	// Convert input patterns to a DynamicMesh
+
+	FClothPatternToDynamicMesh Converter;
+
+	FDynamicMesh3 Mesh2D;
+	Converter.Convert(ClothCollection, INDEX_NONE, EClothPatternVertexType::Sim2D, Mesh2D);
+
+	const double TotalArea = TMeshQueries<FDynamicMesh3>::GetVolumeArea(Mesh2D).Y;
+	checkf(TotalArea > 0, TEXT("FChaosClothAssetRemeshNode::RemeshSimMesh: Expected 2D Sim mesh to have a positive area"));
+	const int TriangleCount = Mesh2D.TriangleCount();
+
+	// Copy pattern IDs into polygroup layer
+
+	Mesh2D.EnableAttributes();
+	const int32 PatternIndexLayerID = Mesh2D.Attributes()->NumPolygroupLayers();
+	Mesh2D.Attributes()->SetNumPolygroupLayers(PatternIndexLayerID + 1);
+	FDynamicMeshPolygroupAttribute* const PatternIndexLayer = Mesh2D.Attributes()->GetPolygroupLayer((int)PatternIndexLayerID);
+	check(Mesh2D.TriangleCount() == InClothFacade.GetNumSimFaces());
+	for (int32 FaceIndex = 0; FaceIndex < InClothFacade.GetNumSimFaces(); ++FaceIndex)
+	{
+		const int32 PatternID = InClothFacade.FindSimPatternByFaceIndex(FaceIndex);
+		PatternIndexLayer->SetValue(FaceIndex, PatternID);
+	}
+
+	// Look for the DensityMap weight layer
+	const FDynamicMeshWeightAttribute* DensityMapLayer = nullptr;
+	if (Mesh2D.HasAttributes())
+	{
+		for (int32 WeightLayerIndex = 0; WeightLayerIndex < Mesh2D.Attributes()->NumWeightLayers(); ++WeightLayerIndex)
+		{
+			if (Mesh2D.Attributes()->GetWeightLayer(WeightLayerIndex)->GetName() == DensityMapName)
+			{
+				DensityMapLayer = Mesh2D.Attributes()->GetWeightLayer(WeightLayerIndex);
+				break;
+			}
+		}
+	}
+
+	// Grab the seam information from the input collection
+
+	// Seams are comprised of a set of Stitches. Each Stitch is simply a pair of vertex indices indicating vertices that should be welded to form the 3D mesh.
+
+	// Stitches are given in random order within the Seam. To make remeshing them easier, we will find connected strips of stitches and store them in sequential order.
+	// So the vertices in Stitch N are connected to the vertices in Stitch N+1.
+
+	TArray<TArray<FIntVector2>> Seams;
+	for (int32 SeamIndex = 0; SeamIndex < InClothFacade.GetNumSeams(); ++SeamIndex)
+	{
+		FCollectionClothSeamConstFacade SeamFacade = InClothFacade.GetSeam(SeamIndex);
+		FClothGeometryTools::BuildConnectedSeams2D(ClothCollection, SeamIndex, Mesh2D, Seams);
+
+		// Check seams are valid
+		for (int32 SeamID = 0; SeamID < Seams.Num(); ++SeamID)
+		{
+			const TArray<FIntVector2>& SubSeam = Seams[SeamID];
+			for (int32 StitchID = 0; StitchID < SubSeam.Num() - 1; ++StitchID)
+			{
+				const int32 NextStitchID = StitchID + 1;
+				for (int32 Side = 0; Side < 2; ++Side)
+				{
+					const int32 StitchVert = SubSeam[StitchID][Side];
+					const int32 NextStitchVert = SubSeam[NextStitchID][Side];
+
+					if (StitchVert == NextStitchVert)
+					{
+						continue;
+					}
+
+					const int32 FoundEdge = Mesh2D.FindEdge(StitchVert, NextStitchVert);
+
+					// This would indicate a problem in BuildConnectedSeams
+					checkf(FoundEdge != UE::Geometry::FDynamicMesh3::InvalidID, TEXT("Could not find a mesh edge between sequential seam vertices %d, %d"), StitchVert, NextStitchVert);
+				}
+			}
+		}
+	}
+
+
+	// For remeshing purposes, we also want to find and constrain any vertices that connect the two seam sides together
+	// (This is only relevant for "internal" seams, which connect vertices within the same pattern)
+	for (TArray<FIntVector2>& Seam : Seams)
+	{
+		if (Seam.Num() == 0)
+		{
+			continue;
+		}
+
+		for (const FIntVector2& EndStitch : { Seam[0], Seam.Last() })
+		{
+			for (const int NeighborA : Mesh2D.VtxVerticesItr(EndStitch[0]))
+			{
+				bool bMatchFound = false;
+
+				for (const int NeighborB : Mesh2D.VtxVerticesItr(EndStitch[1]))
+				{
+					if (NeighborA == NeighborB)
+					{
+						if (Seam.Contains(FIntVector2(NeighborA, NeighborB)))
+						{
+							continue;
+						}
+
+						if (EndStitch == Seam[0])
+						{
+							Seam.Insert(FIntVector2(NeighborA, NeighborB), 0);
+						}
+						else
+						{
+							Seam.Add(FIntVector2(NeighborA, NeighborB));
+						}
+
+						bMatchFound = true;
+						break;
+					}
+				}
+
+				if (bMatchFound)
+				{
+					break;
+				}
+			}
+		}
+	}
+
+	// Remesh seams
+
+	const int TargetTriangleCount = FMath::RoundToInt(static_cast<float>(TargetPercentSim) / 100.0f * static_cast<float>(TriangleCount));
+	const double TargetEdgeLength = FRemeshMeshOp::CalculateTargetEdgeLength(nullptr, TargetTriangleCount, TotalArea);
+
+	for (int32 ResampleIter = 0; ResampleIter < IterationsSim; ++ResampleIter)
+	{
+		UE::Chaos::ClothAsset::RemeshNodeV1::Private::RemeshSeams(Mesh2D, Seams, TargetEdgeLength, DensityMapLayer);
+	}
+
+	// Remesh boundaries
+
+	for (int32 ResampleIter = 0; ResampleIter < IterationsSim; ++ResampleIter)
+	{
+		UE::Chaos::ClothAsset::RemeshNodeV1::Private::RemeshBoundaries(Mesh2D, Seams, TargetEdgeLength, DensityMapLayer);
+	}
+
+	// Do the remeshing of the rest of the mesh
+
+	UE::Geometry::FCompactMaps CompactMaps;
+	constexpr bool bUniformSmoothing = true;
+	UE::Chaos::ClothAsset::RemeshNodeV1::Private::Remesh(Mesh2D, TargetEdgeLength, IterationsSim, SmoothingSim, bUniformSmoothing, Seams, DensityMapName, &CompactMaps);
+
+	// Update stitches
+	for (TArray<FIntVector2>& Seam : Seams)
+	{
+		for (FIntVector2& Stitch : Seam)
+		{
+			Stitch[0] = CompactMaps.GetVertexMapping(Stitch[0]);
+			Stitch[1] = CompactMaps.GetVertexMapping(Stitch[1]);
+			checkf(Stitch[0] != FDynamicMesh3::InvalidID, TEXT("Stitch vertex %d was deleted by remeshing"), Stitch[0]);
+			checkf(Stitch[1] != FDynamicMesh3::InvalidID, TEXT("Stitch vertex %d was deleted by remeshing"), Stitch[1]);
+		}
+	}
+
+
+	// Project the 3D vertices onto the input 3D mesh
+
+	// For each 2D vertex, we will find the closest triangle on the input 2D mesh, then look up that triangle on the input 3D mesh to get the final 3D location.
+	// We will do this pattern-by-pattern to handle issues where the patterns overlap in 2D space.
+
+	const FDynamicMeshPolygroupAttribute* const NewPatternIndexLayer = Mesh2D.Attributes()->GetPolygroupLayer((int)PatternIndexLayerID);
+	check(NewPatternIndexLayer);
+
+	FDynamicMesh3 Mesh3D;
+	Mesh3D.Copy(Mesh2D);
+
+	TMap<int32, TSet<int32>> PatternVertexIDs;
+	for (int32 TID : Mesh2D.TriangleIndicesItr())
+	{
+		const int32 PatternID = NewPatternIndexLayer->GetValue(TID);
+		if (!PatternVertexIDs.Contains(PatternID))
+		{
+			PatternVertexIDs.Add(PatternID, TSet<int32>());
+		}
+
+		const FIndex3i Tri = Mesh2D.GetTriangle(TID);
+		PatternVertexIDs[PatternID].Add(Tri[0]);
+		PatternVertexIDs[PatternID].Add(Tri[1]);
+		PatternVertexIDs[PatternID].Add(Tri[2]);
+	}
+
+	for (int32 PatternID = 0; PatternID < InClothFacade.GetNumSimPatterns(); ++PatternID)
+	{
+		if (!PatternVertexIDs.Contains(PatternID))
+		{
+			continue;
+		}
+
+		FDynamicMesh3 ProjectionTarget2D;
+		Converter.Convert(ClothCollection, PatternID, EClothPatternVertexType::Sim2D, ProjectionTarget2D);
+		FDynamicMesh3 ProjectionTarget3D;
+		Converter.Convert(ClothCollection, PatternID, EClothPatternVertexType::Sim3D, ProjectionTarget3D);
+		TSharedPtr<UE::Geometry::FDynamicMeshAABBTree3> ProjectionTargetSpatial = MakeShared<UE::Geometry::FDynamicMeshAABBTree3>(&ProjectionTarget2D, true);
+
+		for (const int32 VID : PatternVertexIDs[PatternID])
+		{
+			const FVector3d SrcVert = Mesh2D.GetVertex(VID);
+
+			double Distance;
+			const int NearestTriangle = ProjectionTargetSpatial->FindNearestTriangle(SrcVert, Distance);
+
+			const FDistPoint3Triangle3d Dist = TMeshQueries<FDynamicMesh3>::TriangleDistance(ProjectionTarget2D, NearestTriangle, SrcVert);
+			const FVector3d Bary = Dist.TriangleBaryCoords;
+			const FVector3d InterpolatedPoint = ProjectionTarget3D.GetTriBaryPoint(NearestTriangle, Bary[0], Bary[1], Bary[2]);
+
+			Mesh3D.SetVertex(VID, InterpolatedPoint);
+		}
+	}
+
+	// Build the output cloth sim mesh
+
+	constexpr bool bAppendToExistingMesh = false;
+	constexpr bool bTransferWeightMaps = true;
+	constexpr bool bTransferSimSkinningData = true;
+	TMap<int, int32> DynamicMeshToClothVertexMap;
+	FClothGeometryTools::BuildSimMeshFromDynamicMeshes(OutClothCollection, Mesh2D, Mesh3D, PatternIndexLayerID, bTransferWeightMaps, bTransferSimSkinningData, bAppendToExistingMesh, DynamicMeshToClothVertexMap);
+
+
+	// Re-apply the seam info from the input sim mesh. This will create a new Seam for each set of connected stitches
+
+	FCollectionClothFacade OutClothFacade(OutClothCollection);
+	for (int32 SeamIndex = 0; SeamIndex < Seams.Num(); ++SeamIndex)
+	{
+		const TArray<FIntVector2>& Seam = Seams[SeamIndex];
+
+		TArray<FIntVector2> NewSeam;
+		for (const FIntVector2& Stitch : Seam)
+		{
+			if (Stitch[0] == Stitch[1])
+			{
+				continue;
+			}
+
+			FIntVector2 NewStitch;
+			NewStitch[0] = DynamicMeshToClothVertexMap[Stitch[0]];
+			NewStitch[1] = DynamicMeshToClothVertexMap[Stitch[1]];
+			NewSeam.Add(NewStitch);
+		}
+
+		FCollectionClothSeamFacade NewSeamFacade = OutClothFacade.AddGetSeam();
+		NewSeamFacade.Initialize(NewSeam);
+	}
+}
+
+
+void FChaosClothAssetRemeshNode::RemeshRenderMesh(const TSharedRef<const FManagedArrayCollection>& ClothCollection,
+	const FString& DensityMapName,
+	const TSharedRef<FManagedArrayCollection>& OutClothCollection) const
+{
+	using namespace UE::Geometry;
+	using namespace UE::Chaos::ClothAsset;
+
+	// Get the source mesh
+	FClothPatternToDynamicMesh Converter;
+	FDynamicMesh3 DynamicMesh;
+
+	// NOTE: When applied to the Render mesh, this Convert function will assign PatternIDs to the MaterialID attribute of the DynamicMesh. 
+	// After remeshing we will use the MatrialID attribute to determine which triangles should go into which output pattern.
+	Converter.Convert(ClothCollection, INDEX_NONE, EClothPatternVertexType::Render, DynamicMesh);
+
+	check(DynamicMesh.HasAttributes());
+
+	const int InputMeshVertexCount = DynamicMesh.VertexCount();
+	const int InputMeshTriangleCount = DynamicMesh.TriangleCount();
+
+	const bool bHasUVs = (DynamicMesh.Attributes()->PrimaryUV() != nullptr);
+
+	const int TargetTriangleCount = FMath::RoundToInt(static_cast<float>(TargetPercentRender) / 100.0f * static_cast<float>(InputMeshTriangleCount));
+	const double TargetEdgeLength = FRemeshMeshOp::CalculateTargetEdgeLength(&DynamicMesh, TargetTriangleCount);
+
+	TArray<TArray<FIntVector2>> Seams;
+
+	if (bRemeshRenderSeams)
+	{
+		// Create pseudo-stitches based on boundary vertex proximity. These stitches aren't going to actually weld vertices together, but they will guide boundary remeshing.
+		// The goal is to maintain a vertex pairing along boundaries in order to avoid holes opening up when the mesh deforms due to skinning.
+		TArray<FIntVector2> Stitches;
+		UE::Chaos::ClothAsset::RemeshNodeV1::Private::FindCoincidentBoundaryVertices(DynamicMesh, Stitches);
+
+		FClothGeometryTools::BuildConnectedSeams(Stitches, DynamicMesh, Seams);
+
+		// Add density map for for Render mesh
+		const UE::Geometry::FDynamicMeshWeightAttribute* DensityMapLayer = nullptr;
+		if (DynamicMesh.HasAttributes())
+		{
+			for (int32 WeightLayerIndex = 0; WeightLayerIndex < DynamicMesh.Attributes()->NumWeightLayers(); ++WeightLayerIndex)
+			{
+				if (DynamicMesh.Attributes()->GetWeightLayer(WeightLayerIndex)->GetName() == DensityMapName)
+				{
+					DensityMapLayer = DynamicMesh.Attributes()->GetWeightLayer(WeightLayerIndex);
+					break;
+				}
+			}
+		}
+
+		for (int RemeshPass = 0; RemeshPass < RenderSeamRemeshIterations; ++RemeshPass)
+		{
+			UE::Chaos::ClothAsset::RemeshNodeV1::Private::RemeshSeams(DynamicMesh, Seams, TargetEdgeLength, DensityMapLayer);
+		}
+
+		// Also remesh the open boundaries that are not constrained by seams
+		for (int RemeshPass = 0; RemeshPass < RenderSeamRemeshIterations; ++RemeshPass)
+		{
+			UE::Chaos::ClothAsset::RemeshNodeV1::Private::RemeshBoundaries(DynamicMesh, Seams, TargetEdgeLength, DensityMapLayer);
+		}
+	}
+
+
+	UE::Geometry::FCompactMaps CompactMaps;
+	if (RemeshMethodRender == EChaosClothAssetRemeshMethod::Remesh)
+	{
+		constexpr bool bUniformSmoothing = false;	// uniform smoothing can distort the UV layer pretty badly
+		const bool bSuccess = UE::Chaos::ClothAsset::RemeshNodeV1::Private::Remesh(DynamicMesh, TargetEdgeLength, IterationsRender, SmoothingRender, bUniformSmoothing, Seams, DensityMapName, &CompactMaps);
+		check(bSuccess);
+	}
+	else
+	{
+		const bool bCoarsenBoundariesDuringSimplify = !bRemeshRenderSeams;
+		const int TargetVertexCount = FMath::RoundToInt(static_cast<float>(TargetPercentRender) / 100.0f * static_cast<float>(InputMeshVertexCount));
+		UE::Chaos::ClothAsset::Private::Simplify(DynamicMesh, TargetVertexCount, bCoarsenBoundariesDuringSimplify, &CompactMaps);
+	}
+
+	// Collect outputs
+
+	//
+	// Normals
+	//
+
+	const bool bHasNormals = (DynamicMesh.Attributes()->PrimaryNormals() != nullptr);
+	check(bHasNormals);
+
+	TArray<FVector3f> Normals;
+	Normals.SetNum(DynamicMesh.VertexCount());
+	const FDynamicMeshNormalOverlay* const NormalOverlay = DynamicMesh.Attributes()->PrimaryNormals();
+	for (const int TriangleIndex : DynamicMesh.TriangleIndicesItr())
+	{
+		const FIndex3i Tri = DynamicMesh.GetTriangle(TriangleIndex);
+
+		for (int TriangleVertexIndex = 0; TriangleVertexIndex < 3; ++TriangleVertexIndex)
+		{
+			const int VertexIndex = Tri[TriangleVertexIndex];
+
+			// NOTE: This assumes one normal per vertex in the overlay (i.e. no "hard edges")
+			Normals[VertexIndex] = NormalOverlay->GetElementAtVertex(TriangleIndex, VertexIndex);
+		}
+	}
+
+	//
+	// Tangents
+	//
+
+	TArray<FVector3f> TangentUs;
+	TArray<FVector3f> TangentVs;
+	{
+		const bool bHasTangentUs = (DynamicMesh.Attributes()->PrimaryTangents() != nullptr);
+		const bool bHasTangentVs = (DynamicMesh.Attributes()->PrimaryBiTangents() != nullptr);
+		if (!bHasTangentUs || !bHasTangentVs)
+		{
+			FMeshTangentsf::ComputeDefaultOverlayTangents(DynamicMesh);
+		}
+		TangentUs.SetNumZeroed(DynamicMesh.VertexCount());
+		TangentVs.SetNumZeroed(DynamicMesh.VertexCount());
+
+		const FDynamicMeshNormalOverlay* const TangentUOverlay = DynamicMesh.Attributes()->PrimaryTangents();
+		const FDynamicMeshNormalOverlay* const TangentVOverlay = DynamicMesh.Attributes()->PrimaryBiTangents();
+
+		for (const int TriangleIndex : DynamicMesh.TriangleIndicesItr())
+		{
+			const FIndex3i Tri = DynamicMesh.GetTriangle(TriangleIndex);
+
+			for (int TriangleVertexIndex = 0; TriangleVertexIndex < 3; ++TriangleVertexIndex)
+			{
+				const int VertexIndex = Tri[TriangleVertexIndex];
+
+				TangentUs[VertexIndex] += TangentUOverlay->GetElementAtVertex(TriangleIndex, VertexIndex);
+				TangentVs[VertexIndex] += TangentVOverlay->GetElementAtVertex(TriangleIndex, VertexIndex);
+			}
+		}
+
+		for (int32 VertexIndex = 0; VertexIndex < DynamicMesh.VertexCount(); ++VertexIndex)
+		{
+			TangentUs[VertexIndex].Normalize();
+			TangentVs[VertexIndex].Normalize();
+		}
+	}
+
+	//
+	// UVs
+	//
+
+	// TODO: Account for multiple UV layers
+	TArray<FVector2f> UVs;
+	if (bHasUVs)
+	{
+		UVs.SetNum(DynamicMesh.VertexCount());
+
+		const FDynamicMeshUVOverlay* const UVOverlay = DynamicMesh.Attributes()->PrimaryUV();
+
+		// Assume no seams in the dynamic mesh UV overlay
+		checkSlow(!UVOverlay->HasInteriorSeamEdges());
+
+		for (const int TriangleIndex : DynamicMesh.TriangleIndicesItr())
+		{
+			const FIndex3i Tri = DynamicMesh.GetTriangle(TriangleIndex);
+
+			for (int TriangleVertexIndex = 0; TriangleVertexIndex < 3; ++TriangleVertexIndex)
+			{
+				const int VertexIndex = Tri[TriangleVertexIndex];
+				UVs[VertexIndex] = UVOverlay->GetElementAtVertex(TriangleIndex, VertexIndex);
+			}
+		}
+	}
+
+	//
+	// Skin Weights
+	//
+
+	FDynamicMeshAttributeSet* Attributes = DynamicMesh.Attributes();
+	check(Attributes);
+
+	TArray<TArray<int32>> BoneIndices;
+	TArray<TArray<float>> BoneWeights;
+	BoneIndices.SetNum(DynamicMesh.VertexCount());
+	BoneWeights.SetNum(DynamicMesh.VertexCount());
+
+	for (const TPair<FName, TUniquePtr<FDynamicMeshVertexSkinWeightsAttribute>>& SkinWeightLayer : Attributes->GetSkinWeightsAttributes())
+	{
+		const FDynamicMeshVertexSkinWeightsAttribute* const SkinWeightAttribute = SkinWeightLayer.Value.Get();
+
+		for (const int TriangleIndex : DynamicMesh.TriangleIndicesItr())
+		{
+			const FIndex3i Tri = DynamicMesh.GetTriangle(TriangleIndex);
+
+			for (int TriangleVertexIndex = 0; TriangleVertexIndex < 3; ++TriangleVertexIndex)
+			{
+				const int VertexIndex = Tri[TriangleVertexIndex];
+				SkinWeightAttribute->GetValue(VertexIndex, BoneIndices[VertexIndex], BoneWeights[VertexIndex]);
+			}
+		}
+	}
+
+
+	// Find the set of triangles per MaterialID
+
+	TMap<int32, TArray<int32>> MaterialTriangles;
+
+	const FDynamicMeshMaterialAttribute* const MaterialAttribute = Attributes->GetMaterialID();
+	check(MaterialAttribute);
+
+	for (const int32 TriangleID : DynamicMesh.TriangleIndicesItr())
+	{
+		const int32 MaterialID = MaterialAttribute->GetValue(TriangleID);
+		if (!MaterialTriangles.Contains(MaterialID))
+		{
+			MaterialTriangles.Add(MaterialID);
+		}
+		MaterialTriangles[MaterialID].Add(TriangleID);
+	}
+
+	TArray<int32> MaterialIDs;
+	MaterialTriangles.GetKeys(MaterialIDs);
+
+	const int32 NumMaterials = MaterialIDs.Num();
+
+	//
+	// Populate output cloth collection
+	// 
+
+	FClothGeometryTools::DeleteRenderMesh(OutClothCollection);
+	FCollectionClothFacade OutClothFacade(OutClothCollection);
+	OutClothFacade.SetNumRenderPatterns(NumMaterials);
+
+	for (int32 DestPatternID = 0; DestPatternID < MaterialIDs.Num(); ++DestPatternID)
+	{
+		FCollectionClothRenderPatternFacade OutClothPatternFacade = OutClothFacade.GetRenderPattern(DestPatternID);
+		check(OutClothPatternFacade.GetNumRenderFaces() == 0);
+		check(OutClothPatternFacade.GetNumRenderVertices() == 0);
+
+		const int32 SourceMaterialID = MaterialIDs[DestPatternID];
+		check(MaterialTriangles.Contains(SourceMaterialID));
+		const TArray<int32>& TriangleIDs = MaterialTriangles[SourceMaterialID];
+
+		TSet<int32> VertexIndices;
+		for (int32 PatternTriangleIndex = 0; PatternTriangleIndex < TriangleIDs.Num(); ++PatternTriangleIndex)
+		{
+			const int32 TInd = TriangleIDs[PatternTriangleIndex];
+			const FIndex3i Tri = DynamicMesh.GetTriangle(TInd);
+			VertexIndices.Add(Tri[0]);
+			VertexIndices.Add(Tri[1]);
+			VertexIndices.Add(Tri[2]);
+		}
+		const TArray<int32> SourceVertexIndicesArray = VertexIndices.Array();
+		const int32 NumVerticesThisPattern = SourceVertexIndicesArray.Num();
+
+		OutClothPatternFacade.SetNumRenderVertices(NumVerticesThisPattern);
+		TArrayView<FVector3f> RenderPosition = OutClothPatternFacade.GetRenderPosition();
+		TArrayView<FVector3f> RenderNormal = OutClothPatternFacade.GetRenderNormal();
+		TArrayView<FVector3f> RenderTangentU = OutClothPatternFacade.GetRenderTangentU();
+		TArrayView<FVector3f> RenderTangentV = OutClothPatternFacade.GetRenderTangentV();
+		TArrayView<TArray<FVector2f>> RenderUVs = OutClothPatternFacade.GetRenderUVs();
+		TArrayView<FLinearColor> RenderColor = OutClothPatternFacade.GetRenderColor();
+		TArrayView<TArray<int32>> RenderBoneIndices = OutClothPatternFacade.GetRenderBoneIndices();
+		TArrayView<TArray<float>> RenderBoneWeights = OutClothPatternFacade.GetRenderBoneWeights();
+
+		TMap<int32, int32> SourceToDestVertexMap;
+		for (int32 PatternVertexIndex = 0; PatternVertexIndex < NumVerticesThisPattern; ++PatternVertexIndex)
+		{
+			const int32 SourceVertexIndex = SourceVertexIndicesArray[PatternVertexIndex];
+
+			SourceToDestVertexMap.Add(SourceVertexIndex, PatternVertexIndex);
+
+			RenderPosition[PatternVertexIndex] = FVector3f(DynamicMesh.GetVertex(SourceVertexIndex));
+			if (bHasUVs)
+			{
+				RenderUVs[PatternVertexIndex].SetNum(MAX_TEXCOORDS);
+				RenderUVs[PatternVertexIndex][0] = UVs[SourceVertexIndex];
+			}
+			if (bHasNormals)
+			{
+				RenderNormal[PatternVertexIndex] = Normals[SourceVertexIndex];
+			}
+			RenderTangentU[PatternVertexIndex] = TangentUs[SourceVertexIndex];
+			RenderTangentV[PatternVertexIndex] = TangentVs[SourceVertexIndex];
+			RenderColor[PatternVertexIndex] = FLinearColor::White;
+			RenderBoneIndices[PatternVertexIndex] = BoneIndices[SourceVertexIndex];
+			RenderBoneWeights[PatternVertexIndex] = BoneWeights[SourceVertexIndex];
+		}
+
+		OutClothPatternFacade.SetNumRenderFaces(TriangleIDs.Num());
+		for (int32 PatternTriangleIndex = 0; PatternTriangleIndex < TriangleIDs.Num(); ++PatternTriangleIndex)
+		{
+			const int32 VertexOffset = OutClothPatternFacade.GetRenderVerticesOffset();
+			TArrayView<FIntVector3> RenderIndices = OutClothPatternFacade.GetRenderIndices();
+
+			const int32 TInd = TriangleIDs[PatternTriangleIndex];
+			const FIndex3i SourceTri = DynamicMesh.GetTriangle(TInd);
+
+			RenderIndices[PatternTriangleIndex][0] = VertexOffset + SourceToDestVertexMap[SourceTri[0]];
+			RenderIndices[PatternTriangleIndex][1] = VertexOffset + SourceToDestVertexMap[SourceTri[1]];
+			RenderIndices[PatternTriangleIndex][2] = VertexOffset + SourceToDestVertexMap[SourceTri[2]];
+		}
+
+		FCollectionClothConstFacade InClothFacade(ClothCollection);
+		FCollectionClothRenderPatternConstFacade InPatternFacade = InClothFacade.GetRenderPattern(SourceMaterialID);
+		OutClothPatternFacade.SetRenderMaterialSoftObjectPathName(InPatternFacade.GetRenderMaterialSoftObjectPathName());
+	}
 }
 
+
 void FChaosClothAssetRemeshNode::Evaluate(UE::Dataflow::FContext& Context, const FDataflowOutput* Out) const
 {
 	if (Out->IsA<FManagedArrayCollection>(&Collection))
@@ -1996,7 +2995,7 @@ void FChaosClothAssetRemeshNode::Evaluate(UE::Dataflow::FContext& Context, const
 				{
 					Private::EmptySimSelections(OutputClothCollectionRef);
 					Private::EmptySprings(OutputClothCollectionRef);
-					Private::RemeshSimMesh(ClothCollection, IterationsSim, SmoothingSim, SimDensityMapName, FVector2f(TargetPercentSim / 100.0f, TargetPercentSim / 50.0f), OutputClothCollectionRef);
+					RemeshSimMesh(ClothCollection, SimDensityMapName, OutputClothCollectionRef);
 					Private::RebuildTopologyDependentSimData(ClothCollection, OutputClothCollectionRef);
 				}
 			}
@@ -2004,16 +3003,7 @@ void FChaosClothAssetRemeshNode::Evaluate(UE::Dataflow::FContext& Context, const
 			if (bRemeshRender)
 			{
 				Private::EmptyRenderSelections(OutputClothCollectionRef);
-				Private::RemeshRenderMesh(ClothCollection,
-					RemeshMethodRender,
-					bRemeshRenderSeams,
-					RenderSeamRemeshIterations,
-					IterationsRender,
-					SmoothingRender,
-					RenderDensityMapName,
-					FVector2f(TargetPercentRender / 100.0f, TargetPercentRender / 100.0f),
-					TargetPercentRender,
-					OutputClothCollectionRef);
+				RemeshRenderMesh(ClothCollection, RenderDensityMapName, OutputClothCollectionRef);
 			}
 		}
 
@@ -2021,5 +3011,6 @@ void FChaosClothAssetRemeshNode::Evaluate(UE::Dataflow::FContext& Context, const
 	}
 }
 
+
 #undef LOCTEXT_NAMESPACE
 
diff --git a/Engine/Plugins/ChaosClothAssetEditor/Source/ChaosClothAssetDataflowNodes/Public/ChaosClothAsset/RemeshNode.h b/Engine/Plugins/ChaosClothAssetEditor/Source/ChaosClothAssetDataflowNodes/Public/ChaosClothAsset/RemeshNode.h
--- a/Engine/Plugins/ChaosClothAssetEditor/Source/ChaosClothAssetDataflowNodes/Public/ChaosClothAsset/RemeshNode.h
+++ b/Engine/Plugins/ChaosClothAssetEditor/Source/ChaosClothAssetDataflowNodes/Public/ChaosClothAsset/RemeshNode.h
@@ -154,4 +154,15 @@ public:
 	FChaosClothAssetRemeshNode(const UE::Dataflow::FNodeParameters& InParam, FGuid InGuid = FGuid::NewGuid());
 
 	virtual void Evaluate(UE::Dataflow::FContext& Context, const FDataflowOutput* Out) const override;
+
+private:
+
+	void RemeshSimMesh(const TSharedRef<const FManagedArrayCollection>& ClothCollection,
+		const FString& DensityMapName,
+		const TSharedRef<FManagedArrayCollection>& OutClothCollection) const;
+
+	void RemeshRenderMesh(const TSharedRef<const FManagedArrayCollection>& ClothCollection,
+		const FString& DensityMapName,
+		const TSharedRef<FManagedArrayCollection>& OutClothCollection) const;
+
 };
