Chaos Cloth Asset - Allows for render less cloth assets.

[RN] Chaos Cloth Asset - Allows for render less cloth assets.
#Synced-CL 48653274

diff --git a/Engine/Plugins/ChaosClothAsset/Source/ChaosClothAssetEngine/Private/ChaosClothAsset/ClothAsset.cpp b/Engine/Plugins/ChaosClothAsset/Source/ChaosClothAssetEngine/Private/ChaosClothAsset/ClothAsset.cpp
--- a/Engine/Plugins/ChaosClothAsset/Source/ChaosClothAssetEngine/Private/ChaosClothAsset/ClothAsset.cpp
+++ b/Engine/Plugins/ChaosClothAsset/Source/ChaosClothAssetEngine/Private/ChaosClothAsset/ClothAsset.cpp
@@ -544,7 +544,7 @@ void UChaosClothAsset::Build(
 		ClothFacade.DefineSchema();
 
 		const FCollectionClothConstFacade InClothFacade(InClothCollections[LodIndex]);
-		if (!InClothFacade.HasValidRenderData())  // The cloth collection must at least have a render mesh
+		if (!InClothFacade.HasValidRenderData() && !InClothFacade.HasValidSimulationData())  // The cloth collection must at least have a mesh
 		{
 			if (ErrorText && ErrorText->IsEmpty())
 			{
diff --git a/Engine/Plugins/ChaosClothAssetEditor/Source/ChaosClothAssetDataflowNodes/Private/ChaosClothAsset/TerminalNode.cpp b/Engine/Plugins/ChaosClothAssetEditor/Source/ChaosClothAssetDataflowNodes/Private/ChaosClothAsset/TerminalNode.cpp
--- a/Engine/Plugins/ChaosClothAssetEditor/Source/ChaosClothAssetDataflowNodes/Private/ChaosClothAsset/TerminalNode.cpp
+++ b/Engine/Plugins/ChaosClothAssetEditor/Source/ChaosClothAssetDataflowNodes/Private/ChaosClothAsset/TerminalNode.cpp
@@ -206,7 +206,7 @@ TArray<TSharedRef<const FManagedArrayCollection>> FChaosClothAssetTerminalNode_v
 		TSharedRef<FManagedArrayCollection> CollectionLodValue = MakeShared<FManagedArrayCollection>(GetValue<FManagedArrayCollection>(Context, GetConnectionReference(LodIndex)));
 
 		FCollectionClothFacade ClothFacade(CollectionLodValue);
-		if (ClothFacade.HasValidRenderData())  // The cloth collection must at least have a render mesh
+		if (ClothFacade.HasValidRenderData() || ClothFacade.HasValidSimulationData())  // The cloth collection must at least have a mesh
 		{
 			FClothGeometryTools::CleanupAndCompactMesh(CollectionLodValue);
 			LastValidLodIndex = LodIndex;
diff --git a/Engine/Plugins/ChaosClothAssetEditor/Source/ChaosClothAssetEditor/Private/ChaosClothAsset/ChaosClothAssetThumbnailRenderer.cpp b/Engine/Plugins/ChaosClothAssetEditor/Source/ChaosClothAssetEditor/Private/ChaosClothAsset/ChaosClothAssetThumbnailRenderer.cpp
--- a/Engine/Plugins/ChaosClothAssetEditor/Source/ChaosClothAssetEditor/Private/ChaosClothAsset/ChaosClothAssetThumbnailRenderer.cpp
+++ b/Engine/Plugins/ChaosClothAssetEditor/Source/ChaosClothAssetEditor/Private/ChaosClothAsset/ChaosClothAssetThumbnailRenderer.cpp
@@ -2,13 +2,16 @@
 
 #include "ChaosClothAsset/ChaosClothAssetThumbnailRenderer.h"
 
-#include "SceneView.h"
-#include "ShowFlags.h"
+#include "ChaosClothAsset/ClothAsset.h"
 #include "ChaosClothAsset/ClothAssetBase.h"
 #include "ChaosClothAsset/ClothComponent.h"
+#include "ChaosClothAsset/ClothGeometryTools.h"
 #include "Components/SkeletalMeshComponent.h"
 #include "Engine/SkeletalMesh.h"
+#include "Rendering/SkeletalMeshRenderData.h"
 #include "ThumbnailRendering/SceneThumbnailInfo.h"
+#include "SceneView.h"
+#include "ShowFlags.h"
 
 #include UE_INLINE_GENERATED_CPP_BY_NAME(ChaosClothAssetThumbnailRenderer)
 
@@ -78,15 +81,57 @@ namespace UE::Chaos::ClothAsset
 
 void UChaosClothAssetThumbnailRenderer::Draw(UObject* Object, int32 X, int32 Y, uint32 Width, uint32 Height, FRenderTarget* Viewport, FCanvas* Canvas, bool bAdditionalViewFamily)
 {
-	UChaosClothAssetBase* AsClothAsset = Cast<UChaosClothAssetBase>(Object);
-	if (!AsClothAsset)
+	UChaosClothAssetBase* const ClothAsset = Cast<UChaosClothAssetBase>(Object);
+	if (!ClothAsset)
 	{
 		return;
 	}
 	
 	TSharedRef<UE::Chaos::ClothAsset::FThumbnailScene> ThumbnailScene = ClothThumbnailSceneCache.EnsureThumbnailScene(Object);
 
-	ThumbnailScene->SetClothAsset(AsClothAsset);
+	// Check if the cloth asset has render data, if not render the sim mesh instead
+	FSkeletalMeshRenderData* const RenderData = ClothAsset->GetResourceForRendering();
+	const bool bHasRenderData = RenderData && RenderData->LODRenderData.Num() && RenderData->LODRenderData[0].GetTotalFaces();
+
+	if (!bHasRenderData && ClothAsset->HasValidClothSimulationModels() && ClothAsset->GetCollections(0).Num())
+	{
+		// Render sim mesh
+		const FSoftObjectPath RenderMaterialPathName(TEXT("/Engine/EditorMaterials/Cloth/CameraLitDoubleSided.CameraLitDoubleSided"));
+
+		const TSharedRef<const FManagedArrayCollection>& Collection = ClothAsset->GetCollections(0)[0];
+		UChaosClothAsset* const ClothAssetCopy = NewObject<UChaosClothAsset>(this, UChaosClothAsset::StaticClass());
+
+		TSharedRef<FManagedArrayCollection> CollectionCopy = MakeShared<FManagedArrayCollection>(*Collection);
+
+		// Reverse mesh normals to keep the same render look as inverted sim mesh in flat shaded visualization
+		constexpr bool bReverseSimMeshNormals = true;
+		constexpr bool bReverseSimMeshWindingOrder = false;
+		constexpr bool bReverseRenderMeshNormals = false;
+		constexpr bool bReverseRenderMeshWindingOrder = false;
+		const TArray<int32> EmptySelection;
+
+		UE::Chaos::ClothAsset::FClothGeometryTools::ReverseMesh(
+			CollectionCopy,
+			bReverseSimMeshNormals,
+			bReverseSimMeshWindingOrder,
+			bReverseRenderMeshNormals,
+			bReverseRenderMeshWindingOrder,
+			EmptySelection,
+			EmptySelection);
+
+		constexpr bool bSingleRenderPattern = true;
+		UE::Chaos::ClothAsset::FClothGeometryTools::CopySimMeshToRenderMesh(CollectionCopy, RenderMaterialPathName, bSingleRenderPattern);
+
+		const TArray<TSharedRef<const FManagedArrayCollection>> ClothCollectionsCopy = { MoveTemp(CollectionCopy) };
+
+		ClothAssetCopy->Build(ClothCollectionsCopy);
+		ThumbnailScene->SetClothAsset(ClothAssetCopy);
+	}
+	else
+	{
+		// Render render mesh
+		ThumbnailScene->SetClothAsset(ClothAsset);
+	}
 
 	FSceneViewFamilyContext ViewFamily( FSceneViewFamily::ConstructionValues( Viewport, ThumbnailScene->GetScene(), FEngineShowFlags(ESFIM_Game) )
 		.SetTime(UThumbnailRenderer::GetTime())
