Chaos Cloth Asset - Fixed Static Mesh  and Skeletal Mesh Import nodes vertex color space conversions.

[RN] Chaos Cloth Asset - Fixed Static Mesh and Skeletal Mesh Import nodes vertex color space conversions.

#Synced-CL 50002989

diff --git a/Engine/Plugins/ChaosClothAssetEditor/Source/ChaosClothAssetDataflowNodes/Private/ChaosClothAsset/ClothDataflowTools.cpp b/Engine/Plugins/ChaosClothAssetEditor/Source/ChaosClothAssetDataflowNodes/Private/ChaosClothAsset/ClothDataflowTools.cpp
--- a/Engine/Plugins/ChaosClothAssetEditor/Source/ChaosClothAssetDataflowNodes/Private/ChaosClothAsset/ClothDataflowTools.cpp
+++ b/Engine/Plugins/ChaosClothAssetEditor/Source/ChaosClothAssetDataflowNodes/Private/ChaosClothAsset/ClothDataflowTools.cpp
@@ -475,7 +475,7 @@ namespace UE::Chaos::ClothAsset
 				RenderUVs[VertexIndex][TexCoordIndex] = SoftVertex.UVs[TexCoordIndex];
 			}
 
-			RenderColor[VertexIndex] = FLinearColor(SoftVertex.Color);
+			RenderColor[VertexIndex] = SoftVertex.Color.ReinterpretAsLinear();
 
 			const int32 NumBones = Section.MaxBoneInfluences;
 			RenderBoneIndices[VertexIndex].SetNum(NumBones);
@@ -586,6 +586,31 @@ namespace UE::Chaos::ClothAsset
 			SkinWeights.Set(VertexID, RootBinding);
 		}
 
+		// Convert weird static mesh inverse sRGB gamma to linear.
+		// FIXME: Remove once static mesh color space has been fixed to be linear again.
+		TVertexInstanceAttributesRef<FVector4f> VertexInstanceColors = SkeletalMeshAttributes.GetVertexInstanceColors();
+		auto ConvertLinearToSRGBGamma = [](float V)
+			{
+				V = FMath::Clamp(V, 0.0f, 1.0f);
+				if (V <= 0.0031308)
+				{
+					return V * 12.92f;
+				}
+				else
+				{
+					return 1.055f * FMath::Pow(V, 1.0f / 2.4f) - 0.055f;
+				}
+			};
+
+		for (FVertexInstanceID VertexInstanceID : SkeletalMeshGeometry.VertexInstances().GetElementIDs())
+		{
+			FLinearColor VertexColor = VertexInstanceColors.Get(VertexInstanceID);
+			VertexColor.R = ConvertLinearToSRGBGamma(VertexColor.R);
+			VertexColor.G = ConvertLinearToSRGBGamma(VertexColor.G);
+			VertexColor.B = ConvertLinearToSRGBGamma(VertexColor.B);
+			VertexInstanceColors.Set(VertexInstanceID, VertexColor);
+		}
+
 		FSkeletalMeshImportData SkeletalMeshImportGeometry = FSkeletalMeshImportData::CreateFromMeshDescription(SkeletalMeshGeometry);
 		// Data needed by BuildSkeletalMesh
 		TArray<FVector3f> LODPoints;
