Dataflow Cloth : fix weight map node crashing when no transfer collection has been assigned to the node

#Synced-CL 50358557

diff --git a/Engine/Plugins/ChaosClothAssetEditor/Source/ChaosClothAssetDataflowNodes/Private/ChaosClothAsset/WeightMapNode.cpp b/Engine/Plugins/ChaosClothAssetEditor/Source/ChaosClothAssetDataflowNodes/Private/ChaosClothAsset/WeightMapNode.cpp
--- a/Engine/Plugins/ChaosClothAssetEditor/Source/ChaosClothAssetDataflowNodes/Private/ChaosClothAsset/WeightMapNode.cpp
+++ b/Engine/Plugins/ChaosClothAssetEditor/Source/ChaosClothAssetDataflowNodes/Private/ChaosClothAsset/WeightMapNode.cpp
@@ -63,7 +63,16 @@ void FChaosClothAssetWeightMapNode::OnTransfer(UE::Dataflow::FContext& Context)
 	if (ClothFacade.IsValid())  // Can only act on the collection if it is a valid cloth collection
 	{
 		FManagedArrayCollection InTransferCollection = GetValue<FManagedArrayCollection>(Context, &TransferCollection);
+		if (InTransferCollection.IsEmpty())
+		{
+			FClothDataflowTools::LogAndToastWarning(*this, 
+				LOCTEXT("OnTransferError_TransferError", "Transfer Error"),
+				LOCTEXT("OnTransferError_TransferCollectionEmpty", "Transfer collection is empty, no source data to transfer from")
+			);
+			return;
+		}
 		const TSharedRef<const FManagedArrayCollection> TransferClothCollection = MakeShared<const FManagedArrayCollection>(MoveTemp(InTransferCollection));
+
 		FCollectionClothConstFacade TransferClothFacade(TransferClothCollection);
 		const FName InInputName = GetInputName(Context);
 
@@ -103,6 +112,14 @@ void FChaosClothAssetWeightMapNode::OnTransfer(UE::Dataflow::FContext& Context)
 
 				SetVertexWeights(ClothFacade.GetWeightMap(InInputName), RemappedWeights);
 			}
+			else
+			{
+				FClothDataflowTools::LogAndToastWarning(*this, 
+					LOCTEXT("OnTransferError_TransferError", "Transfer Error"),
+					FText::Format(LOCTEXT("OnTransferError_InvalidSimWeightmap", "Invalid transfer collection sim weightmap [{0}], no source data to transfer from"),
+						FText::FromName(InInputName)
+					));
+			}
 		}
 		else 
 		{
@@ -127,6 +144,14 @@ void FChaosClothAssetWeightMapNode::OnTransfer(UE::Dataflow::FContext& Context)
 
 				SetVertexWeights(ClothFacade.GetUserDefinedAttribute<float>(InInputName, ClothCollectionGroup::RenderVertices), RemappedWeights);
 			}
+			else
+			{
+				FClothDataflowTools::LogAndToastWarning(*this, 
+					LOCTEXT("OnTransferError_TransferError", "Transfer Error"),
+					FText::Format(LOCTEXT("InvalidRenderWeightmap", "Invalid transfer collection render weightmap [{0}], no source data to transfer from"),
+						FText::FromName(InInputName)
+					));
+			}
 		}
 	}
 }
