Cloth Asset : enable Dataflow as a default editor in the main branch ( disabled fro 5.7 )

diff --git a/Engine/Plugins/ChaosClothAssetEditor/Source/ChaosClothAssetEditor/Private/ChaosClothAsset/AssetDefinition_ClothAsset.cpp b/Engine/Plugins/ChaosClothAssetEditor/Source/ChaosClothAssetEditor/Private/ChaosClothAsset/AssetDefinition_ClothAsset.cpp
--- a/Engine/Plugins/ChaosClothAssetEditor/Source/ChaosClothAssetEditor/Private/ChaosClothAsset/AssetDefinition_ClothAsset.cpp
+++ b/Engine/Plugins/ChaosClothAssetEditor/Source/ChaosClothAssetEditor/Private/ChaosClothAsset/AssetDefinition_ClothAsset.cpp
@@ -27,7 +27,7 @@
 
 namespace UE::Chaos::ClothAsset::Private
 {
-static bool bUseClothPanelEditorByDefault = true;
+static bool bUseClothPanelEditorByDefault = false;
 static FAutoConsoleVariableRef CVarUseClothPanelEditorByDefault(TEXT("p.ChaosCloth.UseClothPanelEditorByDefault"), bUseClothPanelEditorByDefault, TEXT("Enable the use of the Cloth Panel Editor instead of the unified Dataflow one for editing cloth assets"));
 }
 
