Chaos Cloth Asset - Added a missing index safeguard on creating the render selection set.

#rnx

diff --git a/Engine/Plugins/ChaosClothAssetEditor/Source/ChaosClothAssetDataflowNodes/Private/ChaosClothAsset/USDImportNode_v2.cpp b/Engine/Plugins/ChaosClothAssetEditor/Source/ChaosClothAssetDataflowNodes/Private/ChaosClothAsset/USDImportNode_v2.cpp
--- a/Engine/Plugins/ChaosClothAssetEditor/Source/ChaosClothAssetDataflowNodes/Private/ChaosClothAsset/USDImportNode_v2.cpp
+++ b/Engine/Plugins/ChaosClothAssetEditor/Source/ChaosClothAssetDataflowNodes/Private/ChaosClothAsset/USDImportNode_v2.cpp
@@ -1736,7 +1736,10 @@ bool FChaosClothAssetUSDImportNode_v2::ImportRenderStaticMesh(const TSharedRef<F
 					SelectionSet.Empty(PatternNameFaces.Value.Num());
 					for (const int32 Index : PatternNameFaces.Value)
 					{
-						SelectionSet.Emplace(OriginalToNewTriangles[Index]);
+						if (OriginalToNewTriangles.IsValidIndex(Index) && OriginalToNewTriangles[Index] != INDEX_NONE)
+						{
+							SelectionSet.Emplace(OriginalToNewTriangles[Index]);
+						}
 					}
 				}
 			}
