Dataflow : fix static analysis warnings

diff --git a/Engine/Plugins/ChaosClothAssetEditor/Source/ChaosClothAssetEditorTools/Private/ChaosClothAsset/ClothTransferSkinWeightsTool.cpp b/Engine/Plugins/ChaosClothAssetEditor/Source/ChaosClothAssetEditorTools/Private/ChaosClothAsset/ClothTransferSkinWeightsTool.cpp
--- a/Engine/Plugins/ChaosClothAssetEditor/Source/ChaosClothAssetEditorTools/Private/ChaosClothAsset/ClothTransferSkinWeightsTool.cpp
+++ b/Engine/Plugins/ChaosClothAssetEditor/Source/ChaosClothAssetEditorTools/Private/ChaosClothAsset/ClothTransferSkinWeightsTool.cpp
@@ -193,9 +193,12 @@ void UClothTransferSkinWeightsTool::Shutdown(EToolShutdownType ShutdownType)
 		TransferSkinWeightsNode->SkeletalMesh = ToolProperties->SourceMesh;
 		TransferSkinWeightsNode->Transform = TransformFromProperties();
 		TransferSkinWeightsNode->Invalidate();
-		if (UDataflow* const Dataflow = DataflowContextObject? DataflowContextObject->GetDataflowAsset(): nullptr)
+		if (DataflowContextObject)
 		{
-			Dataflow->MarkPackageDirty();
+			if (TObjectPtr<UDataflow> const Dataflow = DataflowContextObject->GetDataflowAsset())
+			{
+				Dataflow->MarkPackageDirty();
+			}
 		}
 	}
 
