Chaos Cloth Asset - Makes the weight map name input showing by default on the MaxDistance config node since it is almost always used when configuring the MaxDistance.

#rnx

diff --git a/Engine/Plugins/ChaosClothAssetEditor/Source/ChaosClothAssetDataflowNodes/Private/ChaosClothAsset/SimulationMaxDistanceConfigNode.cpp b/Engine/Plugins/ChaosClothAssetEditor/Source/ChaosClothAssetDataflowNodes/Private/ChaosClothAsset/SimulationMaxDistanceConfigNode.cpp
--- a/Engine/Plugins/ChaosClothAssetEditor/Source/ChaosClothAssetDataflowNodes/Private/ChaosClothAsset/SimulationMaxDistanceConfigNode.cpp
+++ b/Engine/Plugins/ChaosClothAssetEditor/Source/ChaosClothAssetDataflowNodes/Private/ChaosClothAsset/SimulationMaxDistanceConfigNode.cpp
@@ -15,7 +15,7 @@ FChaosClothAssetSimulationMaxDistanceConfigNode::FChaosClothAssetSimulationMaxDi
 	RegisterCollectionConnections();
 	RegisterInputConnection(&MaxDistance.WeightMap)
 		.SetCanHidePin(true)
-		.SetPinIsHidden(true);
+		.SetPinIsHidden(false);
 	RegisterInputConnection(&InKinematic, GET_MEMBER_NAME_CHECKED(FChaosClothAssetConnectableIStringValue, StringValue));
 	RegisterOutputConnection(&KinematicVertices3D);
 }
