Chaos Cloth Asset:
- Fixed the  ChaosClothAssetTools Content and Resource paths after the move from the  ChaosClothAssetEditor module.
- Removed deprecated code in AssetDefinition_ClothAsset.

#rnx
#virtualized
#Synced-CL 48812260

diff --git a/Engine/Plugins/ChaosClothAsset/Source/ChaosClothAssetTools/Private/ChaosClothAsset/ClothAssetEditorStyle.cpp b/Engine/Plugins/ChaosClothAsset/Source/ChaosClothAssetTools/Private/ChaosClothAsset/ClothAssetEditorStyle.cpp
--- a/Engine/Plugins/ChaosClothAsset/Source/ChaosClothAssetTools/Private/ChaosClothAsset/ClothAssetEditorStyle.cpp
+++ b/Engine/Plugins/ChaosClothAsset/Source/ChaosClothAssetTools/Private/ChaosClothAsset/ClothAssetEditorStyle.cpp
@@ -11,7 +11,7 @@ namespace UE::Chaos::ClothAsset
 	FClothAssetEditorStyle::FClothAssetEditorStyle()
 		: FSlateStyleSet("ClothAssetEditorStyle")
 	{
-		TSharedPtr<IPlugin> ChaosClothAssetPlugin = IPluginManager::Get().FindPlugin("ChaosClothAssetEditor");
+		TSharedPtr<IPlugin> ChaosClothAssetPlugin = IPluginManager::Get().FindPlugin("ChaosClothAsset");
 		if (ChaosClothAssetPlugin.IsValid())
 		{
 			SetContentRoot(ChaosClothAssetPlugin->GetBaseDir() / TEXT("Resources"));
diff --git a/Engine/Plugins/ChaosClothAsset/Source/ChaosClothAssetTools/Private/ChaosClothAsset/ClothAssetFactory.cpp b/Engine/Plugins/ChaosClothAsset/Source/ChaosClothAssetTools/Private/ChaosClothAsset/ClothAssetFactory.cpp
--- a/Engine/Plugins/ChaosClothAsset/Source/ChaosClothAssetTools/Private/ChaosClothAsset/ClothAssetFactory.cpp
+++ b/Engine/Plugins/ChaosClothAsset/Source/ChaosClothAssetTools/Private/ChaosClothAsset/ClothAssetFactory.cpp
@@ -31,7 +31,7 @@ UObject* UChaosClothAssetFactory::FactoryCreateNew(UClass* Class, UObject* Paren
 {
 	using namespace UE::Chaos::ClothAsset;
 
-	static const FString TemplatePath = TEXT("/ChaosClothAssetEditor/ClothAssetTemplate.ClothAssetTemplate");
+	static const FString TemplatePath = TEXT("/ChaosClothAsset/ClothAssetTemplate.ClothAssetTemplate");
 	static const TCHAR* AssetPrefix = TEXT("CA_");
 	const bool bEmbedDataflow = Private::CVarCreateWithEmbeddedDataflow.GetValueOnAnyThread();
 	return UE::DataflowAssetDefinitionHelpers::FactoryCreateNew(Class, Parent, Name, Flags, &TemplatePath, bEmbedDataflow, AssetPrefix);
diff --git a/Engine/Plugins/ChaosClothAssetEditor/Source/ChaosClothAssetEditor/Private/ChaosClothAsset/AssetDefinition_ClothAsset.cpp b/Engine/Plugins/ChaosClothAssetEditor/Source/ChaosClothAssetEditor/Private/ChaosClothAsset/AssetDefinition_ClothAsset.cpp
--- a/Engine/Plugins/ChaosClothAssetEditor/Source/ChaosClothAssetEditor/Private/ChaosClothAsset/AssetDefinition_ClothAsset.cpp
+++ b/Engine/Plugins/ChaosClothAssetEditor/Source/ChaosClothAssetEditor/Private/ChaosClothAsset/AssetDefinition_ClothAsset.cpp
@@ -27,125 +27,11 @@
 
 namespace UE::Chaos::ClothAsset::Private
 {
-static bool bUseClothPanelEditorByDefault = false;
-static FAutoConsoleVariableRef CVarUseClothPanelEditorByDefault(TEXT("p.ChaosCloth.UseClothPanelEditorByDefault"), bUseClothPanelEditorByDefault, TEXT("Enable the use of the Cloth Panel Editor instead of the unified Dataflow one for editing cloth assets"));
-}
-
-namespace ClothAssetDefinitionHelpers
-{
-	// Return true if we should proceed, false if we should re-open the dialog
-	bool CreateNewDataflowAsset(const UChaosClothAsset* ClothAsset, UObject*& OutDataflowAsset)
-	{
-		const UClass* const DataflowClass = UDataflow::StaticClass();
-
-		FSaveAssetDialogConfig NewDataflowAssetDialogConfig;
-		{
-			const FString PackageName = ClothAsset->GetOutermost()->GetName();
-			NewDataflowAssetDialogConfig.DefaultPath = FPackageName::GetLongPackagePath(PackageName);
-			const FString ClothName = ClothAsset->GetName();
-			NewDataflowAssetDialogConfig.DefaultAssetName = "DF_" + ClothName;
-			NewDataflowAssetDialogConfig.AssetClassNames.Add(DataflowClass->GetClassPathName());
-			NewDataflowAssetDialogConfig.ExistingAssetPolicy = ESaveAssetDialogExistingAssetPolicy::AllowButWarn;
-			NewDataflowAssetDialogConfig.DialogTitleOverride = LOCTEXT("NewDataflowAssetDialogTitle", "Save Dataflow Asset As");
-		}
-
-		FContentBrowserModule& ContentBrowserModule = FModuleManager::LoadModuleChecked<FContentBrowserModule>("ContentBrowser");
-
-		FString NewPackageName = FPaths::Combine(NewDataflowAssetDialogConfig.DefaultPath, NewDataflowAssetDialogConfig.DefaultAssetName);
-		FText OutError;
-
-		while (!FFileHelper::IsFilenameValidForSaving(NewPackageName, OutError) ||
-			LoadObject<UObject>(nullptr, *NewPackageName, nullptr, LOAD_NoWarn | LOAD_Quiet))  // Check if an object with this name already exists
-		{
-			const FString AssetSavePath = ContentBrowserModule.Get().CreateModalSaveAssetDialog(NewDataflowAssetDialogConfig);
-			if (AssetSavePath.IsEmpty())
-			{
-				OutDataflowAsset = nullptr;
-				return false;
-			}
-			NewPackageName = FPackageName::ObjectPathToPackageName(AssetSavePath);
-		}
-
-		const FName NewAssetName(FPackageName::GetLongPackageAssetName(NewPackageName));
-		UPackage* const NewPackage = CreatePackage(*NewPackageName);
-
-		UDataflow* const ClothAssetTemplate = LoadObject<UDataflow>(NewPackage, TEXT("/ChaosClothAssetEditor/ClothAssetTemplate.ClothAssetTemplate"));
-		UDataflow* const NewAsset = DuplicateObject(ClothAssetTemplate, NewPackage, NewAssetName);
-
-		NewAsset->MarkPackageDirty();
-
-		// Notify the asset registry
-		FAssetRegistryModule::AssetCreated(NewAsset);
-
-		OutDataflowAsset = NewAsset;
-		return true;
-	}
-
-
-	// Return true if we should proceed, false if we should re-open the dialog
-	bool OpenDataflowAsset(const UChaosClothAsset* ClothAsset, UObject*& OutDataflowAsset)
-	{
-		const UClass* const DataflowClass = UDataflow::StaticClass();
-
-		FOpenAssetDialogConfig NewDataflowAssetDialogConfig;
-		{
-			const FString PackageName = ClothAsset->GetOutermost()->GetName();
-			NewDataflowAssetDialogConfig.DefaultPath = FPackageName::GetLongPackagePath(PackageName);
-			NewDataflowAssetDialogConfig.AssetClassNames.Add(DataflowClass->GetClassPathName());
-			NewDataflowAssetDialogConfig.bAllowMultipleSelection = false;
-			NewDataflowAssetDialogConfig.DialogTitleOverride = LOCTEXT("OpenDataflowAssetDialogTitle", "Open Dataflow Asset");
-		}
-
-		FContentBrowserModule& ContentBrowserModule = FModuleManager::LoadModuleChecked<FContentBrowserModule>("ContentBrowser");
-		TArray<FAssetData> AssetData = ContentBrowserModule.Get().CreateModalOpenAssetDialog(NewDataflowAssetDialogConfig);
-
-		if (AssetData.Num() == 1)
-		{
-			OutDataflowAsset = AssetData[0].GetAsset();
-			return true;
-		}
-
-		return false;
-	}
-
-	// Return true if we should proceed, false if we should re-open the dialog
-	bool NewOrOpenDialog(const UChaosClothAsset* ClothAsset, UObject*& OutDataflowAsset)
-	{
-		TSharedRef<SMessageDialog> ConfirmDialog = SNew(SMessageDialog)
-			.Title(FText(LOCTEXT("ClothDataflow_WindowTitle", "Create or Open Dataflow graph?")))
-			.Message(LOCTEXT("ClothDataflow_WindowText", "This Cloth Asset currently has no Dataflow graph"))
-			.Buttons({
-				SMessageDialog::FButton(LOCTEXT("ClothDataflow_NewText", "Create new Dataflow")),
-				SMessageDialog::FButton(LOCTEXT("ClothDataflow_OpenText", "Open existing Dataflow")),
-				SMessageDialog::FButton(LOCTEXT("ClothDataflow_ContinueText", "Continue without Dataflow")),
-			});
-
-		const int32 ResultButtonIdx = ConfirmDialog->ShowModal();
-		switch(ResultButtonIdx)
-		{
-		case 0:
-			return CreateNewDataflowAsset(ClothAsset, OutDataflowAsset);
-		case 1:
-			return OpenDataflowAsset(ClothAsset, OutDataflowAsset);
-		default:
-			break;
-		}
-
-		return true;
-	}
-}
-
-// Create a new UDataflow if one doesn't already exist for the Cloth Asset
-UObject* UAssetDefinition_ClothAsset::NewOrOpenDataflowAsset(const UChaosClothAsset* ClothAsset)
-{
-	UObject* DataflowAsset = nullptr;
-	bool bDialogDone = false;
-	while (!bDialogDone)
-	{
-		bDialogDone = ClothAssetDefinitionHelpers::CreateNewDataflowAsset(ClothAsset, DataflowAsset);
-	}
-
-	return DataflowAsset;
+	static bool bUseClothPanelEditorByDefault = false;
+	static FAutoConsoleVariableRef CVarUseClothPanelEditorByDefault(
+		TEXT("p.ChaosCloth.UseClothPanelEditorByDefault"),
+		bUseClothPanelEditorByDefault,
+		TEXT("Enable the use of the Cloth Panel Editor instead of the unified Dataflow one for editing cloth assets."));
 }
 
 FText UAssetDefinition_ClothAsset::GetAssetDisplayName() const
diff --git a/Engine/Plugins/ChaosClothAssetEditor/Source/ChaosClothAssetEditor/Private/ChaosClothAsset/AssetDefinition_ClothAsset.h b/Engine/Plugins/ChaosClothAssetEditor/Source/ChaosClothAssetEditor/Private/ChaosClothAsset/AssetDefinition_ClothAsset.h
--- a/Engine/Plugins/ChaosClothAssetEditor/Source/ChaosClothAssetEditor/Private/ChaosClothAsset/AssetDefinition_ClothAsset.h
+++ b/Engine/Plugins/ChaosClothAssetEditor/Source/ChaosClothAssetEditor/Private/ChaosClothAsset/AssetDefinition_ClothAsset.h
@@ -13,22 +13,16 @@ class UAssetDefinition_ClothAsset : public UAssetDefinitionDefault
 	GENERATED_BODY()
 
 public:
-
-	UE_DEPRECATED(5.5, "Dataflow assets are now created through the physics menu.")
-	static UObject* NewOrOpenDataflowAsset(const UChaosClothAsset* ClothAsset);
-
 	static bool LaunchClothPanelAssetEditor(UChaosClothAsset* ClothAsset);
 	static bool LaunchClothDataflowAssetEditor(UChaosClothAsset* ClothAsset);
 	static bool UseClothPanelEditorByDefault();
 	
 private:
-
 	virtual FText GetAssetDisplayName() const override;
 	virtual TSoftClassPtr<UObject> GetAssetClass() const override;
 	virtual FLinearColor GetAssetColor() const override;
 	virtual TConstArrayView<FAssetCategoryPath> GetAssetCategories() const override;
 	virtual UThumbnailInfo* LoadThumbnailInfo(const FAssetData& InAssetData) const override;
-
 	virtual EAssetCommandResult OpenAssets(const FAssetOpenArgs& OpenArgs) const override;
 };
 
