Chaos Cloth Asset - Moved USD Dataflow nodes to their own plugin.

#rnx
#Synced-CL 48812260

diff --git a/Engine/Plugins/ChaosClothAssetEditor/ChaosClothAssetEditor.uplugin b/Engine/Plugins/ChaosClothAssetEditor/ChaosClothAssetEditor.uplugin
--- a/Engine/Plugins/ChaosClothAssetEditor/ChaosClothAssetEditor.uplugin
+++ b/Engine/Plugins/ChaosClothAssetEditor/ChaosClothAssetEditor.uplugin
@@ -25,12 +25,6 @@
 			"Type": "Editor",
 			"LoadingPhase": "Default",
 			"PlatformAllowList": [ "Win64", "Mac", "Linux" ]
-		},
-		{
-			"Name": "ChaosClothAssetDataflowNodes",
-			"Type": "Editor",
-			"LoadingPhase": "Default",
-			"PlatformAllowList": [ "Win64", "Mac", "Linux" ]
 		}
 	],
 	"Plugins": [
@@ -46,6 +40,14 @@
 			"Name": "ChaosClothAsset",
 			"Enabled": true
 		},
+		{
+			"Name": "ChaosClothAssetDataflowNodes",
+			"Enabled": true
+		},
+		{
+			"Name": "ChaosClothAssetUsdDataflowNodes",
+			"Enabled": true
+		},
 		{
 			"Name": "ChaosCloth",
 			"Enabled": true
@@ -64,16 +66,6 @@
 			"Name": "Dataflow",
 			"Enabled": true
 		},
-		{
-			"Name": "USDImporter",
-			"Enabled": true,
-			"SupportedTargetPlatforms": [
-				"Win64",
-				"Mac",
-				"Linux"
-			]
-		},
-
 		{
 			"Name": "MeshResizing",
 			"Enabled": true
