Dataflow : Enable Cloth renderable types
- Fixed some viewmode rules to avoid Render and sim mesh to overlap
- Fixed Sim2D to be of a larger scale to avoid to be too close to the grid
- Nodes won't use it yet because they are still using the legacy rendering system
- added a bInFlip2DSimFaces option to converstion to dynamic mesh function ( only for sim2D cases ) so that the wireframe works in 2D

#Synced-CL 49372252

diff --git a/Engine/Plugins/ChaosClothAsset/Source/ChaosClothAssetTools/Private/ChaosClothAsset/ClothPatternToDynamicMesh.cpp b/Engine/Plugins/ChaosClothAsset/Source/ChaosClothAssetTools/Private/ChaosClothAsset/ClothPatternToDynamicMesh.cpp
--- a/Engine/Plugins/ChaosClothAsset/Source/ChaosClothAssetTools/Private/ChaosClothAsset/ClothPatternToDynamicMesh.cpp
+++ b/Engine/Plugins/ChaosClothAsset/Source/ChaosClothAssetTools/Private/ChaosClothAsset/ClothPatternToDynamicMesh.cpp
@@ -31,9 +31,10 @@ public:
 	typedef int32 NormalIDType;
 	typedef int32 ColorIDType;
 
-	FClothPatternWrapper(const FCollectionClothConstFacade& ClothFacade, int32 PatternIndex, EClothPatternVertexType VertexDataType) :
-		VertexDataType(VertexDataType),
-		Cloth(ClothFacade)
+	FClothPatternWrapper(const FCollectionClothConstFacade& ClothFacade, int32 PatternIndex, EClothPatternVertexType VertexDataType, bool bInFlip2DSimFaces = false)
+		: VertexDataType(VertexDataType)
+		, Cloth(ClothFacade)
+		, bFlip2DSimFaces(bInFlip2DSimFaces)
 	{
 		NumTexCoords = 0;
 		if (PatternIndex == INDEX_NONE)
@@ -340,6 +341,10 @@ public:
 			break;
 		case EClothPatternVertexType::Sim2D:
 			Face = Cloth.GetSimIndices2D()[TriID];
+			if (bFlip2DSimFaces)
+			{
+				Swap(Face[0], Face[1]);
+			}
 			break;
 		case EClothPatternVertexType::Sim3D:
 			Face = Cloth.GetSimIndices3D()[TriID];
@@ -531,13 +536,17 @@ public:
 			ID0 = Face[0];
 			ID1 = Face[1];
 			ID2 = Face[2];
+			if (bFlip2DSimFaces)
+			{
+				Swap(ID0, ID1);
+			}
 
 			return true;
 		}
 		else
 		{
-		return GetTri(TriID, ID0, ID1, ID2); 
-	}
+			return GetTri(TriID, ID0, ID1, ID2); 
+		}
 	}
 	
 	const TArray<int32>& GetUVIDs(int32 LayerID) const 
@@ -614,6 +623,7 @@ private:
 	TArray<VertIDType> VertIDs;		// indices into Positions2D or Positions3D
 	TArray<NormalIDType> NormalIDs; // indices into Normals (and all per vertex data other than positions)
 	int32 NumTexCoords;
+	bool bFlip2DSimFaces = false;
 	
 	TArray<FName> WeightMapNames;
 
@@ -629,13 +639,14 @@ void FClothPatternToDynamicMesh::Convert(
 	EClothPatternVertexType VertexDataType,
 	UE::Geometry::FDynamicMesh3& MeshOut,
 	bool bDisableAttributes,
-	int32 MaterialOffset)
+	int32 MaterialOffset, 
+	bool bFlip2DSimFaces)
 {
 	const FCollectionClothConstFacade ClothFacade(ClothCollection);
 
 	// Actual conversion
 	UE::Geometry::TToDynamicMesh<FClothPatternWrapper> PatternToDynamicMesh;
-	FClothPatternWrapper PatternWrapper(ClothFacade, PatternIndex, VertexDataType);
+	FClothPatternWrapper PatternWrapper(ClothFacade, PatternIndex, VertexDataType, bFlip2DSimFaces);
 
 	auto TriangleToGroupFunction = [](FClothPatternWrapper::TriIDType) { return 0; };
 
@@ -736,13 +747,14 @@ void FClothPatternToDynamicMesh::Convert(
 	int32 PatternIndex,
 	EClothPatternVertexType VertexDataType,
 	UE::Geometry::FDynamicMesh3& MeshOut,
-	int32 MaterialOffset)
+	int32 MaterialOffset,
+	bool bFlip2DSimFaces)
 {
 	const TArray<TSharedRef<const FManagedArrayCollection>>& ClothCollections = ClothAssetMeshIn->GetClothCollections();
 	check(ClothCollections.IsValidIndex(LODIndex));
 
 	constexpr bool bDisableAttributes = false;
-	Convert(ClothCollections[LODIndex], PatternIndex, VertexDataType, MeshOut, bDisableAttributes, MaterialOffset);
+	Convert(ClothCollections[LODIndex], PatternIndex, VertexDataType, MeshOut, bDisableAttributes, MaterialOffset, bFlip2DSimFaces);
 }
 
 }	// namespace UE::Chaos::ClothAsset
diff --git a/Engine/Plugins/ChaosClothAsset/Source/ChaosClothAssetTools/Public/ChaosClothAsset/ClothPatternToDynamicMesh.h b/Engine/Plugins/ChaosClothAsset/Source/ChaosClothAssetTools/Public/ChaosClothAsset/ClothPatternToDynamicMesh.h
--- a/Engine/Plugins/ChaosClothAsset/Source/ChaosClothAssetTools/Public/ChaosClothAsset/ClothPatternToDynamicMesh.h
+++ b/Engine/Plugins/ChaosClothAsset/Source/ChaosClothAssetTools/Public/ChaosClothAsset/ClothPatternToDynamicMesh.h
@@ -33,7 +33,8 @@ public:
 		EClothPatternVertexType VertexDataType,
 		UE::Geometry::FDynamicMesh3& MeshOut,
 		bool bDisableAttributes = false,
-		int32 MaterialOffset = INDEX_NONE);
+		int32 MaterialOffset = INDEX_NONE,
+		bool bFlip2DSimFaces = false);
 
 	UE_API void Convert(
 		const UChaosClothAsset* ClothAssetMeshIn,
@@ -41,7 +42,8 @@ public:
 		int32 PatternIndex,
 		EClothPatternVertexType VertexDataType,
 		UE::Geometry::FDynamicMesh3& MeshOut,
-		int32 MaterialOffset = INDEX_NONE);
+		int32 MaterialOffset = INDEX_NONE, 
+		bool bFlip2DSimFaces = false);
 };
 
 }	// namespace UE::Chaos::ClothAsset
diff --git a/Engine/Plugins/ChaosClothAssetEditor/Source/ChaosClothAssetDataflowNodes/Private/ChaosClothAsset/ChaosClothAssetDataflowNodesModule.cpp b/Engine/Plugins/ChaosClothAssetEditor/Source/ChaosClothAssetDataflowNodes/Private/ChaosClothAsset/ChaosClothAssetDataflowNodesModule.cpp
--- a/Engine/Plugins/ChaosClothAssetEditor/Source/ChaosClothAssetDataflowNodes/Private/ChaosClothAsset/ChaosClothAssetDataflowNodesModule.cpp
+++ b/Engine/Plugins/ChaosClothAssetEditor/Source/ChaosClothAssetDataflowNodes/Private/ChaosClothAsset/ChaosClothAssetDataflowNodesModule.cpp
@@ -11,6 +11,7 @@
 #include "ChaosClothAsset/ClothCollectionGroup.h"
 #include "ChaosClothAsset/ClothCollectionToDynamicMeshNode.h"
 #include "ChaosClothAsset/ClothCollectionQueryNode.h"
+#include "ChaosClothAsset/ClothRenderableTypes.h"
 #include "ChaosClothAsset/ColorScheme.h"
 #include "ChaosClothAsset/CopySimulationToRenderMeshNode.h"
 #include "ChaosClothAsset/CustomRegionResizingNode.h"
@@ -409,6 +410,8 @@ namespace UE::Chaos::ClothAsset
 			UE::Dataflow::FRenderingViewModeFactory::GetInstance().RegisterViewMode(MakeUnique<FClothRenderViewMode>());
 
 			UE::Dataflow::FRenderingFactory::GetInstance()->RegisterCallbacks(MakeUnique<FClothSurfaceRenderCallbacks>());
+
+			UE::Chaos::ClothAsset::RegisterClothRenderableTypes();
 		}
 
 		static void DeregisterRenderingCallbacks()
diff --git a/Engine/Plugins/ChaosClothAssetEditor/Source/ChaosClothAssetDataflowNodes/Private/ChaosClothAsset/ClothRenderableTypes.cpp b/Engine/Plugins/ChaosClothAssetEditor/Source/ChaosClothAssetDataflowNodes/Private/ChaosClothAsset/ClothRenderableTypes.cpp
--- a/Engine/Plugins/ChaosClothAssetEditor/Source/ChaosClothAssetDataflowNodes/Private/ChaosClothAsset/ClothRenderableTypes.cpp
+++ b/Engine/Plugins/ChaosClothAssetEditor/Source/ChaosClothAssetDataflowNodes/Private/ChaosClothAsset/ClothRenderableTypes.cpp
@@ -4,7 +4,6 @@
 
 #include "Components/DynamicMeshComponent.h"
 #include "Dataflow/DataflowRenderingViewMode.h"
-#include "Dataflow/DataflowRenderingViewMode.h"
 #include "DataflowRendering/DataflowRenderableComponents.h"
 #include "DataflowRendering/DataflowRenderableType.h"
 #include "DataflowRendering/DataflowRenderableTypeInstance.h"
@@ -15,6 +14,7 @@
 
 #include "UObject/ObjectPtr.h"
 
+#include "ChaosClothAsset/ClothDataflowViewModes.h"
 #include "ChaosClothAsset/ClothCollectionGroup.h"
 #include "ChaosClothAsset/CollectionClothFacade.h"
 #include "ChaosClothAsset/ClothPatternToDynamicMesh.h"
@@ -36,7 +36,11 @@ namespace UE::Chaos::ClothAsset
 	{
 		UE_DATAFLOW_IRENDERABLE_PRIMARY_TYPE(FManagedArrayCollection, Collection);
 		UE_DATAFLOW_IRENDERABLE_RENDER_GROUP(SimSurface);
-		UE_DATAFLOW_IRENDERABLE_VIEW_MODE(UE::Dataflow::FDataflowConstruction3DViewMode);
+
+		virtual bool IsViewModeSupported(const UE::Dataflow::IDataflowConstructionViewMode& InViewMode) const override
+		{
+			return InViewMode.GetName() == UE::Chaos::ClothAsset::FCloth3DSimViewMode::Name;
+		}
 
 		virtual bool CanRender(const UE::Dataflow::FRenderableTypeInstance& Instance) const override
 		{
@@ -50,6 +54,7 @@ namespace UE::Chaos::ClothAsset
 
 			const TSharedRef<const FManagedArrayCollection> ClothCollection = MakeShared<FManagedArrayCollection>(MoveTemp(Collection));
 			const FCollectionClothConstFacade ClothFacade(ClothCollection);
+
 			if (ClothFacade.HasValidSimulationData())
 			{
 				FDynamicMesh3 DynamicMesh;
@@ -79,7 +84,13 @@ namespace UE::Chaos::ClothAsset
 	{
 		UE_DATAFLOW_IRENDERABLE_PRIMARY_TYPE(FManagedArrayCollection, Collection);
 		UE_DATAFLOW_IRENDERABLE_RENDER_GROUP(RenderSurface);
-		UE_DATAFLOW_IRENDERABLE_VIEW_MODE(UE::Dataflow::FDataflowConstruction3DViewMode);
+
+		virtual bool IsViewModeSupported(const UE::Dataflow::IDataflowConstructionViewMode& InViewMode) const override
+		{
+			return InViewMode.GetName() == UE::Dataflow::FDataflowConstruction3DViewMode::Name
+				|| InViewMode.GetName() == UE::Chaos::ClothAsset::FClothRenderViewMode::Name
+				;
+		}
 
 		virtual bool CanRender(const UE::Dataflow::FRenderableTypeInstance& Instance) const override
 		{
@@ -129,7 +140,13 @@ namespace UE::Chaos::ClothAsset
 	{
 		UE_DATAFLOW_IRENDERABLE_PRIMARY_TYPE(FManagedArrayCollection, Collection);
 		UE_DATAFLOW_IRENDERABLE_RENDER_GROUP(Pattern);
-		UE_DATAFLOW_IRENDERABLE_VIEW_MODE(UE::Dataflow::FDataflowConstruction2DViewMode);
+
+		virtual bool IsViewModeSupported(const UE::Dataflow::IDataflowConstructionViewMode& InViewMode) const override
+		{
+			return InViewMode.GetName() == UE::Dataflow::FDataflowConstruction2DViewMode::Name
+				|| InViewMode.GetName() == UE::Chaos::ClothAsset::FCloth2DSimViewMode::Name
+				;
+		}
 
 		virtual bool CanRender(const UE::Dataflow::FRenderableTypeInstance& Instance) const override
 		{
@@ -147,7 +164,7 @@ namespace UE::Chaos::ClothAsset
 			{
 				FDynamicMesh3 DynamicMesh;
 				FClothPatternToDynamicMesh Converter;
-				Converter.Convert(ClothCollection, INDEX_NONE, EClothPatternVertexType::Sim2D, DynamicMesh);
+				Converter.Convert(ClothCollection, INDEX_NONE, EClothPatternVertexType::Sim2D, DynamicMesh, /*bDisableAttributes*/false, INDEX_NONE, /*bFlip2DSimFaces*/true);
 
 				// only create components for the one that have vertices
 				if (DynamicMesh.VertexCount() > 0)
@@ -160,6 +177,7 @@ namespace UE::Chaos::ClothAsset
 
 						UMaterialInterface* TwoSidedMaterial = Private::GetTwoSidedMaterialForRenderingTypes();
 						Component->SetOverrideRenderMaterial(TwoSidedMaterial);
+						Component->SetWorldScale3D(FVector(100, 100, 100));
 					}
 				}
 			}
