Dataflow : update cloth renderable sim2d and sim3d  types to have render settings
- support colored pattern and seams display settings

#Synced-CL 49450039

diff --git a/Engine/Plugins/ChaosClothAssetEditor/Source/ChaosClothAssetDataflowNodes/ChaosClothAssetDataflowNodes.Build.cs b/Engine/Plugins/ChaosClothAssetEditor/Source/ChaosClothAssetDataflowNodes/ChaosClothAssetDataflowNodes.Build.cs
--- a/Engine/Plugins/ChaosClothAssetEditor/Source/ChaosClothAssetDataflowNodes/ChaosClothAssetDataflowNodes.Build.cs
+++ b/Engine/Plugins/ChaosClothAssetEditor/Source/ChaosClothAssetDataflowNodes/ChaosClothAssetDataflowNodes.Build.cs
@@ -42,6 +42,7 @@ public class ChaosClothAssetDataflowNodes : ModuleRules
 				"MeshDescription",
 				"MeshUtilitiesCommon",
 				"ModelingOperatorsEditorOnly",	// TODO: Someday remove editor dependencies, see UE-206172
+				"ModelingComponents",
 				"ModelingOperators",
 				"MeshConversionEngineTypes",
 				"RenderCore",
diff --git a/Engine/Plugins/ChaosClothAssetEditor/Source/ChaosClothAssetDataflowNodes/Private/ChaosClothAsset/ClothRenderableTypes.cpp b/Engine/Plugins/ChaosClothAssetEditor/Source/ChaosClothAssetDataflowNodes/Private/ChaosClothAsset/ClothRenderableTypes.cpp
--- a/Engine/Plugins/ChaosClothAssetEditor/Source/ChaosClothAssetDataflowNodes/Private/ChaosClothAsset/ClothRenderableTypes.cpp
+++ b/Engine/Plugins/ChaosClothAssetEditor/Source/ChaosClothAssetDataflowNodes/Private/ChaosClothAsset/ClothRenderableTypes.cpp
@@ -3,6 +3,9 @@
 #include "ChaosClothAsset/ClothRenderableTypes.h"
 
 #include "Components/DynamicMeshComponent.h"
+#include "Drawing/LineSetComponent.h"
+#include "Drawing/PointSetComponent.h"
+
 #include "Dataflow/DataflowRenderingViewMode.h"
 #include "DataflowRendering/DataflowRenderableComponents.h"
 #include "DataflowRendering/DataflowRenderableType.h"
@@ -14,10 +17,11 @@
 
 #include "UObject/ObjectPtr.h"
 
-#include "ChaosClothAsset/ClothDataflowViewModes.h"
 #include "ChaosClothAsset/ClothCollectionGroup.h"
-#include "ChaosClothAsset/CollectionClothFacade.h"
+#include "ChaosClothAsset/ClothDataflowViewModes.h"
+#include "ChaosClothAsset/ClothGeometryTools.h"
 #include "ChaosClothAsset/ClothPatternToDynamicMesh.h"
+#include "ChaosClothAsset/CollectionClothFacade.h"
 
 #include "Materials/Material.h"
 
@@ -29,6 +33,207 @@ namespace UE::Chaos::ClothAsset
 		{
 			return Cast<UMaterial>(StaticLoadObject(UMaterial::StaticClass(), NULL, TEXT("/Engine/EditorMaterials/Dataflow/DataflowTwoSidedVertexMaterial")));
 		}
+
+		static UMaterialInterface* GetTwoSidedMaterialForColoredPatterns()
+		{
+			return Cast<UMaterial>(StaticLoadObject(UMaterial::StaticClass(), NULL, TEXT("/MeshModelingToolsetExp/Materials/MeshVertexColorMaterialTwoSided")));
+		}
+
+		static FLinearColor MakePseudoRandomColor(int32 NumColorRotations)
+		{
+			constexpr uint8 Spread = 157;  // Prime number that gives a good spread of colors without getting too similar as a rand might do.
+			uint8 Seed = Spread;
+			for (int32 Rotation = 0; Rotation < NumColorRotations; ++Rotation)
+			{
+				Seed += Spread;
+			}
+			return FLinearColor::MakeFromHSV8(Seed, 180, 140);
+		}
+
+		static void AddPatternsColorsMesh(const FCollectionClothConstFacade& ClothFacade, FDynamicMesh3& InOutMesh)
+		{
+			using namespace UE::Geometry;
+
+			checkf(InOutMesh.TriangleCount() == ClothFacade.GetNumSimFaces(), TEXT("Expected to have the same number of faces in LodMesh and ClothCollection sim mesh"));
+
+			InOutMesh.Attributes()->EnablePrimaryColors();
+			InOutMesh.Attributes()->PrimaryColors()->CreateFromPredicate([&ClothFacade](int ParentVID, int TriIDA, int TriIDB)
+				{
+					return ClothFacade.FindSimPatternByFaceIndex(TriIDA) == ClothFacade.FindSimPatternByFaceIndex(TriIDB);
+				}
+				, 0.0f);
+
+			if (FDynamicMeshColorOverlay* const ColorAttributeLayer = InOutMesh.Attributes()->PrimaryColors())
+			{
+				for (int32 PatternID = 0; PatternID < ClothFacade.GetNumSimPatterns(); ++PatternID)
+				{
+					const FCollectionClothSimPatternConstFacade Pattern = ClothFacade.GetSimPattern(PatternID);
+					const FLinearColor PatternColor = MakePseudoRandomColor(PatternID);
+
+					for (int32 TriID = 0; TriID < Pattern.GetNumSimFaces(); ++TriID)
+					{
+						const int32 GlobalTriID = Pattern.GetSimFacesOffset() + TriID;
+						FIndex3i AttrTri;
+						if (ColorAttributeLayer->GetTriangleIfValid(GlobalTriID, AttrTri))
+						{
+							ColorAttributeLayer->SetElement(AttrTri[0], (FVector4f)PatternColor);
+							ColorAttributeLayer->SetElement(AttrTri[1], (FVector4f)PatternColor);
+							ColorAttributeLayer->SetElement(AttrTri[2], (FVector4f)PatternColor);
+						}
+					}
+				}
+			}
+		}
+
+		static void GenerateSeamsForSim2D(
+			const TSharedRef<const FManagedArrayCollection>& Collection,
+			const FCollectionClothConstFacade& ClothFacade,
+			const UE::Geometry::FDynamicMesh3& Mesh,
+			UPointSetComponent* PointComponent, const float PointSize,
+			ULineSetComponent* LineComponent, const float LineThickness,
+			bool bSeamsCollapse
+		)
+		{
+			if (PointComponent == nullptr && LineComponent == nullptr)
+			{
+				return;
+			}
+
+			if (LineComponent)
+			{
+				LineComponent->SetLineMaterial(LoadObject<UMaterialInterface>(nullptr, TEXT("/MeshModelingToolsetExp/Materials/LineSetOverlaidComponentMaterial")));
+			}
+			if (PointComponent)
+			{
+				PointComponent->SetPointMaterial(LoadObject<UMaterial>(nullptr, TEXT("/MeshModelingToolsetExp/Materials/PointSetOverlaidComponentMaterial")));
+			}
+
+			int32 ConnectedSeamIndex = 0;		// Used to generate different colors for each connected seam, if multiple connected seams are found per input seam
+
+			for (int32 SeamIndex = 0; SeamIndex < ClothFacade.GetNumSeams(); ++SeamIndex)
+			{
+				// Stitches are given in random order, so first construct paths of connected stitches
+				// Note one SeamFacade can contain multiple disjoint paths
+				TArray<TArray<FIntVector2>> ConnectedSeams;
+				UE::Chaos::ClothAsset::FClothGeometryTools::BuildConnectedSeams2D(Collection, SeamIndex, Mesh, ConnectedSeams);
+
+				for (const TArray<FIntVector2>& ConnectedSeam : ConnectedSeams)
+				{
+					const FColor SeamColor = UE::Chaos::ClothAsset::Private::MakePseudoRandomColor(ConnectedSeamIndex++).ToFColor(true);
+
+					// draw connected edge on each side of the seam
+					if (LineComponent)
+					{
+						for (int32 StitchID = 0; StitchID < ConnectedSeam.Num() - 1; ++StitchID)
+						{
+							const FVector3d PointA = Mesh.GetVertex(ConnectedSeam[StitchID][0]);
+							const FVector3d PointB = Mesh.GetVertex(ConnectedSeam[StitchID][1]);
+							const FVector3d PointC = Mesh.GetVertex(ConnectedSeam[StitchID + 1][0]);
+							const FVector3d PointD = Mesh.GetVertex(ConnectedSeam[StitchID + 1][1]);
+							LineComponent->AddLine(PointA, PointC, SeamColor, LineThickness);
+							LineComponent->AddLine(PointB, PointD, SeamColor, LineThickness);
+						}
+					}
+
+					// draw connection between stitch points
+					if (bSeamsCollapse)
+					{
+						const int32 StitchID = ConnectedSeam.Num() / 2;
+						const int32 VertexA = ConnectedSeam[StitchID][0];
+						const int32 VertexB = ConnectedSeam[StitchID][1];
+						const FVector PtA = Mesh.GetVertex(VertexA);
+						const FVector PtB = Mesh.GetVertex(VertexB);
+						if (LineComponent)
+						{
+							LineComponent->AddLine(PtA, PtB, SeamColor, LineThickness);
+						}
+						if (PointComponent)
+						{
+							PointComponent->AddPoint(PtA, SeamColor, PointSize);
+							PointComponent->AddPoint(PtB, SeamColor, PointSize);
+						}
+					}
+					else
+					{
+						for (int32 StitchID = 0; StitchID < ConnectedSeam.Num(); ++StitchID)
+						{
+							const int32 VertexA = ConnectedSeam[StitchID][0];
+							const int32 VertexB = ConnectedSeam[StitchID][1];
+							const FVector PtA = Mesh.GetVertex(VertexA);
+							const FVector PtB = Mesh.GetVertex(VertexB);
+							if (LineComponent)
+							{
+								LineComponent->AddLine(PtA, PtB, SeamColor, 2.0f);
+							}
+							if (PointComponent)
+							{
+								PointComponent->AddPoint(PtA, SeamColor, PointSize);
+								PointComponent->AddPoint(PtB, SeamColor, PointSize);
+							}
+						}
+					}
+				}
+			}
+		}
+
+		static void GenerateSeamsForSim3D(
+			const TSharedRef<const FManagedArrayCollection>& Collection, 
+			const FCollectionClothConstFacade& ClothFacade, 
+			const UE::Geometry::FDynamicMesh3& Mesh, 
+			UPointSetComponent* PointComponent, const float PointSize,
+			ULineSetComponent* LineComponent, const float LineThickness,
+			bool bSeamsCollapse
+		)
+		{
+			if (PointComponent == nullptr && LineComponent == nullptr)
+			{
+				return;
+			}
+
+			if (LineComponent)
+			{
+				LineComponent->SetLineMaterial(LoadObject<UMaterialInterface>(nullptr, TEXT("/Engine/EngineMaterials/LineSetComponentMaterial")));
+			}
+			if (PointComponent)
+			{
+				PointComponent->SetPointMaterial(LoadObject<UMaterial>(nullptr, TEXT("/MeshModelingToolsetExp/Materials/PointSetComponentMaterial")));
+			}
+
+			int32 ConnectedSeamIndex = 0;		// Used to generate different colors for each connected seam, if multiple connected seams are found per input seam
+
+			for (int32 SeamIndex = 0; SeamIndex < ClothFacade.GetNumSeams(); ++SeamIndex)
+			{
+				const UE::Chaos::ClothAsset::FCollectionClothSeamConstFacade SeamFacade = ClothFacade.GetSeam(SeamIndex);
+
+				const TArray<int32> SeamStitches(SeamFacade.GetSeamStitch3DIndex());
+				const FColor SeamColor = UE::Chaos::ClothAsset::Private::MakePseudoRandomColor(SeamIndex).ToFColor(true);
+
+				// In 3D we should be able to draw the seam edges in any order, doesn't need to be in connected paths
+				for (int32 StitchIndexI = 0; StitchIndexI < SeamStitches.Num(); ++StitchIndexI)
+				{
+					const int32 VertexIndexI = SeamStitches[StitchIndexI];
+					for (int32 StitchIndexJ = StitchIndexI + 1; StitchIndexJ < SeamStitches.Num(); ++StitchIndexJ)
+					{
+						const int32 VertexIndexJ = SeamStitches[StitchIndexJ];
+
+						if (Mesh.FindEdge(VertexIndexI, VertexIndexJ) != UE::Geometry::FDynamicMesh3::InvalidID)
+						{
+							const FVector PointI = Mesh.GetVertex(VertexIndexI);
+							const FVector PointJ = Mesh.GetVertex(VertexIndexJ);
+							if (LineComponent)
+							{
+								LineComponent->AddLine(PointI, PointJ, SeamColor, LineThickness);
+							}
+							if (PointComponent)
+							{
+								PointComponent->AddPoint(PointI, SeamColor, PointSize);
+								PointComponent->AddPoint(PointJ, SeamColor, PointSize);
+							}
+						}
+					}
+				}
+			}
+		}
 	}
 	////////////////////////////////////////////////////////////////////////////////////////////////////////////////////////////
 
@@ -36,6 +241,7 @@ namespace UE::Chaos::ClothAsset
 	{
 		UE_DATAFLOW_IRENDERABLE_PRIMARY_TYPE(FManagedArrayCollection, Collection);
 		UE_DATAFLOW_IRENDERABLE_RENDER_GROUP(SimSurface);
+		UE_DATAFLOW_IRENDERABLE_SETTINGS(UDataflowClothSimRenderSettings);
 
 		virtual bool IsViewModeSupported(const UE::Dataflow::IDataflowConstructionViewMode& InViewMode) const override
 		{
@@ -50,6 +256,8 @@ namespace UE::Chaos::ClothAsset
 
 		virtual void GetPrimitiveComponents(const UE::Dataflow::FRenderableTypeInstance& Instance, UE::Dataflow::FRenderableComponents& OutComponents) const override
 		{
+			const UDataflowClothSimRenderSettings* Settings = Instance.GetTypedRenderSettings<UDataflowClothSimRenderSettings>();
+
 			FManagedArrayCollection Collection = GetCollection(Instance);
 
 			const TSharedRef<const FManagedArrayCollection> ClothCollection = MakeShared<FManagedArrayCollection>(MoveTemp(Collection));
@@ -61,17 +269,50 @@ namespace UE::Chaos::ClothAsset
 				FClothPatternToDynamicMesh Converter;
 				Converter.Convert(ClothCollection, INDEX_NONE, EClothPatternVertexType::Sim3D, DynamicMesh);
 
+				const bool bShowPatternColors = (Settings && Settings->bShowPatternColors);
+				if (bShowPatternColors)
+				{
+					Private::AddPatternsColorsMesh(ClothFacade, DynamicMesh);
+				}
+
 				// only create components for the one that have vertices
 				if (DynamicMesh.VertexCount() > 0)
 				{
-					static const FName ComponentName = TEXT("ClothSim3DMesh");
+					static const FName ComponentName = TEXT("Cloth_Sim3D_Mesh");
 					if (UDynamicMeshComponent* Component = OutComponents.AddNewComponent<UDynamicMeshComponent>(ComponentName))
 					{
 						Component->SetCastShadow(false);
 						Component->SetMesh(MoveTemp(DynamicMesh));
 
-						UMaterialInterface* TwoSidedMaterial = Private::GetTwoSidedMaterialForRenderingTypes();
+						UMaterialInterface* TwoSidedMaterial = bShowPatternColors
+							? Private::GetTwoSidedMaterialForColoredPatterns()
+							: Private::GetTwoSidedMaterialForRenderingTypes();
 						Component->SetOverrideRenderMaterial(TwoSidedMaterial);
+
+						const FDynamicMesh3* Mesh = Component->GetMesh();
+						if (Mesh && Settings && Settings->bShowSeams)
+						{
+							static const FName PointComponentName = TEXT("Cloth_Sim3D_SeamPoints");
+							UPointSetComponent* PointComponent = (Settings && Settings->bShowSeamPoints)
+								? OutComponents.AddNewComponent<UPointSetComponent>(PointComponentName)
+								: nullptr;
+
+							static const FName LineComponentName = TEXT("Cloth_Sim3D_SeamLines");
+							ULineSetComponent* LineComponent = (Settings && Settings->bShowSeamLines)
+								? OutComponents.AddNewComponent<ULineSetComponent>(LineComponentName)
+								: nullptr;
+
+							const float SeamPointsSize = Settings ? Settings->SeamPointSize : 4.0f;
+							const float SeamLinesThickness = Settings ? Settings->SeamLineThickness : 2.0f;
+							const bool bCollapseSeamLines = Settings && Settings->bCollapseSeamLines;
+
+							Private::GenerateSeamsForSim3D(
+								ClothCollection, ClothFacade, *Mesh,
+								PointComponent, SeamPointsSize,
+								LineComponent, SeamLinesThickness,
+								bCollapseSeamLines
+							);
+						}
 					}
 				}
 			}
@@ -113,7 +354,7 @@ namespace UE::Chaos::ClothAsset
 				// only create components for the one that have vertices
 				if (DynamicMesh.VertexCount() > 0)
 				{
-					static const FName ComponentName = TEXT("ClothRender3DMesh");
+					static const FName ComponentName = TEXT("Cloth_Render_Mesh");
 					if (UDynamicMeshComponent* Component = OutComponents.AddNewComponent<UDynamicMeshComponent>(ComponentName))
 					{
 						Component->SetCastShadow(false);
@@ -140,6 +381,7 @@ namespace UE::Chaos::ClothAsset
 	{
 		UE_DATAFLOW_IRENDERABLE_PRIMARY_TYPE(FManagedArrayCollection, Collection);
 		UE_DATAFLOW_IRENDERABLE_RENDER_GROUP(Pattern);
+		UE_DATAFLOW_IRENDERABLE_SETTINGS(UDataflowClothSimRenderSettings);
 
 		virtual bool IsViewModeSupported(const UE::Dataflow::IDataflowConstructionViewMode& InViewMode) const override
 		{
@@ -156,6 +398,8 @@ namespace UE::Chaos::ClothAsset
 
 		virtual void GetPrimitiveComponents(const UE::Dataflow::FRenderableTypeInstance& Instance, UE::Dataflow::FRenderableComponents& OutComponents) const override
 		{
+			const UDataflowClothSimRenderSettings* Settings = Instance.GetTypedRenderSettings<UDataflowClothSimRenderSettings>();
+
 			FManagedArrayCollection Collection = GetCollection(Instance);
 
 			const TSharedRef<const FManagedArrayCollection> ClothCollection = MakeShared<FManagedArrayCollection>(MoveTemp(Collection));
@@ -166,18 +410,62 @@ namespace UE::Chaos::ClothAsset
 				FClothPatternToDynamicMesh Converter;
 				Converter.Convert(ClothCollection, INDEX_NONE, EClothPatternVertexType::Sim2D, DynamicMesh, /*bDisableAttributes*/false, INDEX_NONE, /*bFlip2DSimFaces*/true);
 
+				const bool bShowPatternColors = (Settings && Settings->bShowPatternColors);
+				if (bShowPatternColors)
+				{
+					Private::AddPatternsColorsMesh(ClothFacade, DynamicMesh);
+				}
+
+				const FVector ScaleFactor(100, 100, 1);
+
 				// only create components for the one that have vertices
 				if (DynamicMesh.VertexCount() > 0)
 				{
-					static const FName ComponentName = TEXT("ClothSim2DMesh");
+					static const FName ComponentName = TEXT("Cloth_Sim2D_Mesh");
 					if (UDynamicMeshComponent* Component = OutComponents.AddNewComponent<UDynamicMeshComponent>(ComponentName))
 					{
 						Component->SetCastShadow(false);
 						Component->SetMesh(MoveTemp(DynamicMesh));
-
-						UMaterialInterface* TwoSidedMaterial = Private::GetTwoSidedMaterialForRenderingTypes();
+						UMaterialInterface* TwoSidedMaterial = bShowPatternColors
+							? Private::GetTwoSidedMaterialForColoredPatterns()
+							: Private::GetTwoSidedMaterialForRenderingTypes();
 						Component->SetOverrideRenderMaterial(TwoSidedMaterial);
-						Component->SetWorldScale3D(FVector(100, 100, 100));
+						Component->SetWorldScale3D(ScaleFactor);
+
+						const FDynamicMesh3* Mesh = Component->GetMesh();
+						if (Mesh && Settings && Settings->bShowSeams)
+						{
+
+							static const FName PointComponentName = TEXT("Cloth_Sim2D_SeamPoints");
+							UPointSetComponent* PointComponent = (Settings && Settings->bShowSeamPoints)
+								? OutComponents.AddNewComponent<UPointSetComponent>(PointComponentName)
+								: nullptr;
+
+							static const FName LineComponentName = TEXT("Cloth_Sim2D_SeamLines");
+							ULineSetComponent* LineComponent = (Settings && Settings->bShowSeamLines)
+								? OutComponents.AddNewComponent<ULineSetComponent>(LineComponentName)
+								: nullptr;
+
+							const float SeamPointsSize = Settings ? Settings->SeamPointSize : 4.0f;
+							const float SeamLinesThickness = Settings ? Settings->SeamLineThickness : 2.0f;
+							const bool bCollapseSeamLines = Settings && Settings->bCollapseSeamLines;
+
+							Private::GenerateSeamsForSim2D(
+								ClothCollection, ClothFacade, *Mesh,
+								PointComponent, SeamPointsSize,
+								LineComponent, SeamLinesThickness,
+								bCollapseSeamLines
+							);
+
+							if (PointComponent)
+							{
+								PointComponent->SetWorldScale3D(ScaleFactor);
+							}
+							if (LineComponent)
+							{
+								LineComponent->SetWorldScale3D(ScaleFactor);
+							}
+						}
 					}
 				}
 			}
diff --git a/Engine/Plugins/ChaosClothAssetEditor/Source/ChaosClothAssetDataflowNodes/Private/ChaosClothAsset/ClothRenderableTypes.h b/Engine/Plugins/ChaosClothAssetEditor/Source/ChaosClothAssetDataflowNodes/Private/ChaosClothAsset/ClothRenderableTypes.h
--- a/Engine/Plugins/ChaosClothAssetEditor/Source/ChaosClothAssetDataflowNodes/Private/ChaosClothAsset/ClothRenderableTypes.h
+++ b/Engine/Plugins/ChaosClothAssetEditor/Source/ChaosClothAssetDataflowNodes/Private/ChaosClothAsset/ClothRenderableTypes.h
@@ -2,6 +2,45 @@
 
 #pragma once
 
+#include "DataflowRendering/DataflowRenderableTypeSettings.h"
+
+#include "ClothRenderableTypes.generated.h"
+
+UCLASS(MinimalAPI, AutoExpandCategories = "Cloth")
+class UDataflowClothSimRenderSettings : public UDataflowRenderableTypeSettings
+{
+	GENERATED_BODY()
+
+public:
+	/** display the seams between patterns */
+	UPROPERTY(EditAnywhere, Category = "Cloth")
+	bool bShowSeams = false;
+
+	/** when seams are displayed, collapse the multiple seams lines into one */
+	UPROPERTY(EditAnywhere, Category = "Cloth", meta = (EditCondition = bShowSeams))
+	bool bCollapseSeamLines = false;
+
+	/** Whether to show the seam points or not */
+	UPROPERTY(EditAnywhere, Category = "Cloth", meta = (EditCondition = bShowSeams))
+	bool bShowSeamPoints = true;
+
+	/** Size of the seam points */
+	UPROPERTY(EditAnywhere, Category = "Cloth", meta = (UIMin = 0.01, ClampMin = 0.01, EditCondition = "bShowSeams && bShowSeamPoints"))
+	float SeamPointSize = 4.0f;
+
+	/** Whether to show the seam lines or not */
+	UPROPERTY(EditAnywhere, Category = "Cloth", meta = (EditCondition = bShowSeams))
+	bool bShowSeamLines = true;
+
+	/** Thickness of the seam lines */
+	UPROPERTY(EditAnywhere, Category = "Cloth", meta = (UIMin = 0.01, ClampMin = 0.01, EditCondition = "bShowSeams && bShowSeamLines"))
+	float SeamLineThickness = 2.0f;
+
+	/** display the patterns colors */
+	UPROPERTY(EditAnywhere, Category = "Cloth")
+	bool bShowPatternColors = false;
+};
+
 namespace UE::Chaos::ClothAsset
 {
 	void RegisterClothRenderableTypes();
diff --git a/Engine/Plugins/Experimental/Dataflow/Source/DataflowEditor/Private/DataflowRendering/DataflowRenderableTypeSettings.h b/Engine/Plugins/Experimental/Dataflow/Source/DataflowEditor/Private/DataflowRendering/DataflowRenderableTypeSettings.h
deleted file mode 100644
diff --git a/Engine/Plugins/Experimental/Dataflow/Source/DataflowEditor/Public/DataflowRendering/DataflowRenderableTypeInstance.h b/Engine/Plugins/Experimental/Dataflow/Source/DataflowEditor/Public/DataflowRendering/DataflowRenderableTypeInstance.h
--- a/Engine/Plugins/Experimental/Dataflow/Source/DataflowEditor/Public/DataflowRendering/DataflowRenderableTypeInstance.h
+++ b/Engine/Plugins/Experimental/Dataflow/Source/DataflowEditor/Public/DataflowRendering/DataflowRenderableTypeInstance.h
@@ -85,7 +85,7 @@ namespace UE::Dataflow
 
 		const void GetSelectionToVisualize(FDataflowSelection& OutSelection) const;
 
-		UDataflowRenderableTypeSettings* GetRenderSettings() const;
+		UE_API UDataflowRenderableTypeSettings* GetRenderSettings() const;
 
 		template <typename T UE_REQUIRES(std::is_base_of_v<UDataflowRenderableTypeSettings, T>)>
 		const T* GetTypedRenderSettings() const
diff --git a/Engine/Plugins/Experimental/Dataflow/Source/DataflowEditor/Public/DataflowRendering/DataflowRenderableTypeSettings.h b/Engine/Plugins/Experimental/Dataflow/Source/DataflowEditor/Public/DataflowRendering/DataflowRenderableTypeSettings.h
new file mode 100644
--- /dev/null
+++ b/Engine/Plugins/Experimental/Dataflow/Source/DataflowEditor/Public/DataflowRendering/DataflowRenderableTypeSettings.h
@@ -0,0 +1,36 @@
+// Copyright Epic Games, Inc. All Rights Reserved.
+
+#pragma once
+
+#include "UObject/Object.h"
+
+#include "DataflowRenderableTypeSettings.generated.h"
+
+UCLASS(MinimalAPI)
+class UDataflowRenderableTypeSettings : public UObject
+{
+	GENERATED_BODY()
+};
+
+// Common types
+
+UCLASS(MinimalAPI, AutoExpandCategories = "UVs")
+class UDataflowUVsRenderSettings : public UDataflowRenderableTypeSettings
+{
+	GENERATED_BODY()
+
+public:
+	/** return the overriden channel if enabled or the default passed to this method */
+	int32 GetUVChannel(int32 DefaultChannel) const
+	{
+		return bOverrideUVchannel ? UVChannel : DefaultChannel;
+	}
+
+private:
+	UPROPERTY(EditAnywhere, Category = "UVs", meta = (InlineEditConditionToggle))
+	bool bOverrideUVchannel = false;
+
+	/** UV channel to display if oiverride is enabled */
+	UPROPERTY(EditAnywhere, Category = "UVs", meta = (UIMIn = 0, UIMax = 7, ClampMin = 0, ClampMax = 7, EditCondition = bOverrideUVchannel))
+	int32 UVChannel = 0;
+};
