Update WeightMap to selection to show the resulting selection

#virtualized
#Synced-CL 49450039

diff --git a/Engine/Plugins/ChaosClothAssetEditor/Source/ChaosClothAssetDataflowNodes/Private/ChaosClothAsset/ClothRenderableTypes.cpp b/Engine/Plugins/ChaosClothAssetEditor/Source/ChaosClothAssetDataflowNodes/Private/ChaosClothAsset/ClothRenderableTypes.cpp
--- a/Engine/Plugins/ChaosClothAssetEditor/Source/ChaosClothAssetDataflowNodes/Private/ChaosClothAsset/ClothRenderableTypes.cpp
+++ b/Engine/Plugins/ChaosClothAssetEditor/Source/ChaosClothAssetDataflowNodes/Private/ChaosClothAsset/ClothRenderableTypes.cpp
@@ -22,6 +22,7 @@
 #include "ChaosClothAsset/ClothGeometryTools.h"
 #include "ChaosClothAsset/ClothPatternToDynamicMesh.h"
 #include "ChaosClothAsset/CollectionClothFacade.h"
+#include "ChaosClothAsset/CollectionClothSelectionFacade.h"
 
 #include "Materials/Material.h"
 
@@ -50,18 +51,22 @@ namespace UE::Chaos::ClothAsset
 			return FLinearColor::MakeFromHSV8(Seed, 180, 140);
 		}
 
-		static void AddPatternsColorsMesh(const FCollectionClothConstFacade& ClothFacade, FDynamicMesh3& InOutMesh)
+		static void AddPatternsColorsToMesh(const FCollectionClothConstFacade& ClothFacade, FDynamicMesh3& InOutMesh)
 		{
 			using namespace UE::Geometry;
 
 			checkf(InOutMesh.TriangleCount() == ClothFacade.GetNumSimFaces(), TEXT("Expected to have the same number of faces in LodMesh and ClothCollection sim mesh"));
 
-			InOutMesh.Attributes()->EnablePrimaryColors();
-			InOutMesh.Attributes()->PrimaryColors()->CreateFromPredicate([&ClothFacade](int ParentVID, int TriIDA, int TriIDB)
-				{
-					return ClothFacade.FindSimPatternByFaceIndex(TriIDA) == ClothFacade.FindSimPatternByFaceIndex(TriIDB);
-				}
+			if (!InOutMesh.Attributes()->HasPrimaryColors())
+			{
+				InOutMesh.Attributes()->EnablePrimaryColors();
+				InOutMesh.Attributes()->PrimaryColors()->CreateFromPredicate([](int ParentVID, int TriIDA, int TriIDB)
+					{
+						return false; // make sure we have a value per corner
+					}
 				, 0.0f);
+			}
+
 
 			if (FDynamicMeshColorOverlay* const ColorAttributeLayer = InOutMesh.Attributes()->PrimaryColors())
 			{
@@ -85,6 +90,176 @@ namespace UE::Chaos::ClothAsset
 			}
 		}
 
+		static bool AddFaceSelectionToMesh(const TSharedRef<const FManagedArrayCollection>& Collection, FName SelectionName, FDynamicMesh3& InOutMesh)
+		{
+			using namespace UE::Geometry;
+
+			if (SelectionName.IsNone())
+			{
+				return false;
+			}
+
+			const FCollectionClothSelectionConstFacade SelectionFacade(Collection);
+			if (!SelectionFacade.IsValid())
+			{
+				return false;
+			}
+
+			const TSet<int32>* SelectionSet = SelectionFacade.FindSelectionSet(SelectionName);
+			if (!SelectionSet)
+			{
+				return false;
+			}
+			
+			const FName& SelectionGroup = SelectionFacade.GetSelectionGroup(SelectionName);
+			if (SelectionGroup != ClothCollectionGroup::SimFaces)
+			{
+				return false;
+			}
+
+			if (!InOutMesh.Attributes()->HasPrimaryColors())
+			{
+				InOutMesh.Attributes()->EnablePrimaryColors();
+				InOutMesh.Attributes()->PrimaryColors()->CreateFromPredicate([](int ParentVID, int TriIDA, int TriIDB)
+					{
+						return false; // make sure we have a value per corner
+					}
+				, 0.0f);
+			}
+
+			const FLinearColor SelectionColor = FLinearColor::Yellow;
+
+			if (FDynamicMeshColorOverlay* const ColorAttributeLayer = InOutMesh.Attributes()->PrimaryColors())
+			{
+				for (const int32 FaceID: *SelectionSet)
+				{
+					FIndex3i AttrTri;
+					if (ColorAttributeLayer->GetTriangleIfValid(FaceID, AttrTri))
+					{
+						ColorAttributeLayer->SetElement(AttrTri[0], (FVector4f)SelectionColor);
+						ColorAttributeLayer->SetElement(AttrTri[1], (FVector4f)SelectionColor);
+						ColorAttributeLayer->SetElement(AttrTri[2], (FVector4f)SelectionColor);
+					}
+				}
+			}
+
+			return true;
+		}
+
+		static void GetVertexSelectionForSim3D(const TSharedRef<const FManagedArrayCollection>& Collection, const FCollectionClothConstFacade& ClothFacade, FName SelectionName, TArray<FVector>& OutPoints)
+		{
+			using namespace UE::Geometry;
+
+			OutPoints.Reset();
+
+			if (SelectionName.IsNone())
+			{
+				return;
+			}
+
+			const FCollectionClothSelectionConstFacade SelectionFacade(Collection);
+			if (!SelectionFacade.IsValid())
+			{
+				return;
+			}
+
+			const TSet<int32>* SelectionSet = SelectionFacade.FindSelectionSet(SelectionName);
+			if (!SelectionSet)
+			{
+				return;
+			}
+
+			OutPoints.Reserve(SelectionSet->Num());
+
+			const FName& SelectionGroup = SelectionFacade.GetSelectionGroup(SelectionName);
+			if (SelectionGroup == ClothCollectionGroup::SimVertices3D)
+			{
+				TConstArrayView<FVector3f> Sim3DPositions = ClothFacade.GetSimPosition3D();
+				for (const int32 VertexIndex : *SelectionSet)
+				{
+					if (Sim3DPositions.IsValidIndex(VertexIndex))
+					{
+						OutPoints.Add(FVector(Sim3DPositions[VertexIndex]));
+					}
+				}
+			}
+			else if (SelectionGroup == ClothCollectionGroup::SimVertices2D)
+			{
+				TConstArrayView<FVector3f> Sim3DPositions = ClothFacade.GetSimPosition3D();
+				const TConstArrayView<int32> Vertex2DTo3D = ClothFacade.GetSimVertex3DLookup();
+				for (const int32 VertexIndex : *SelectionSet)
+				{
+					if (Vertex2DTo3D.IsValidIndex(VertexIndex))
+					{
+						const int32 VertexIndexIn3D = Vertex2DTo3D[VertexIndex];
+						if (Sim3DPositions.IsValidIndex(VertexIndexIn3D))
+						{
+							OutPoints.Add(FVector(Sim3DPositions[VertexIndexIn3D]));
+						}
+					}
+				}
+			}
+		}
+
+		static void GetVertexSelectionForSim2D(const TSharedRef<const FManagedArrayCollection>& Collection, const FCollectionClothConstFacade& ClothFacade, FName SelectionName, TArray<FVector>& OutPoints)
+		{
+			using namespace UE::Geometry;
+
+			OutPoints.Reset();
+
+			if (SelectionName.IsNone())
+			{
+				return;
+			}
+
+			const FCollectionClothSelectionConstFacade SelectionFacade(Collection);
+			if (!SelectionFacade.IsValid())
+			{
+				return;
+			}
+
+			const TSet<int32>* SelectionSet = SelectionFacade.FindSelectionSet(SelectionName);
+			if (!SelectionSet)
+			{
+				return;
+			}
+
+			OutPoints.Reserve(SelectionSet->Num());
+
+			const FName& SelectionGroup = SelectionFacade.GetSelectionGroup(SelectionName);
+			if (SelectionGroup == ClothCollectionGroup::SimVertices3D)
+			{
+				TConstArrayView<FVector2f> Sim2DPositions = ClothFacade.GetSimPosition2D();
+				const TConstArrayView<TArray<int32>> Vertex3DTo2D = ClothFacade.GetSimVertex2DLookup();
+				for (const int32 VertexIndex : *SelectionSet)
+				{
+					if (Vertex3DTo2D.IsValidIndex(VertexIndex))
+					{
+						for (const int32 VertexIndexIn2D : Vertex3DTo2D[VertexIndex])
+						{
+							if (Sim2DPositions.IsValidIndex(VertexIndexIn2D))
+							{
+								const FVector2f PositionIn2D = Sim2DPositions[VertexIndexIn2D];
+								OutPoints.Add(FVector(PositionIn2D.X, PositionIn2D.Y, 0.f));
+							}
+						}
+					}
+				}
+			}
+			else if (SelectionGroup == ClothCollectionGroup::SimVertices2D)
+			{
+				TConstArrayView<FVector2f> Sim2DPositions = ClothFacade.GetSimPosition2D();
+				for (const int32 VertexIndex : *SelectionSet)
+				{
+					if (Sim2DPositions.IsValidIndex(VertexIndex))
+					{
+						const FVector2f PositionIn2D = Sim2DPositions[VertexIndex];
+						OutPoints.Add(FVector(PositionIn2D.X, PositionIn2D.Y, 0.f));
+					}
+				}
+			}
+		}
+
 		static void GenerateSeamsForSim2D(
 			const TSharedRef<const FManagedArrayCollection>& Collection,
 			const FCollectionClothConstFacade& ClothFacade,
@@ -272,8 +447,13 @@ namespace UE::Chaos::ClothAsset
 				const bool bShowPatternColors = (Settings && Settings->bShowPatternColors);
 				if (bShowPatternColors)
 				{
-					Private::AddPatternsColorsMesh(ClothFacade, DynamicMesh);
+					Private::AddPatternsColorsToMesh(ClothFacade, DynamicMesh);
 				}
+				const FString SelectionName = Instance.GetOutputValueByType<FString>({}, TEXT("DataflowSelection"));
+				const bool bHasFaceSelection = Private::AddFaceSelectionToMesh(ClothCollection, FName(SelectionName), DynamicMesh);
+
+				TArray<FVector> SelectedPoints;
+				Private::GetVertexSelectionForSim3D(ClothCollection, ClothFacade, FName(SelectionName), SelectedPoints);
 
 				// only create components for the one that have vertices
 				if (DynamicMesh.VertexCount() > 0)
@@ -284,7 +464,7 @@ namespace UE::Chaos::ClothAsset
 						Component->SetCastShadow(false);
 						Component->SetMesh(MoveTemp(DynamicMesh));
 
-						UMaterialInterface* TwoSidedMaterial = bShowPatternColors
+						UMaterialInterface* TwoSidedMaterial = bShowPatternColors || bHasFaceSelection
 							? Private::GetTwoSidedMaterialForColoredPatterns()
 							: Private::GetTwoSidedMaterialForRenderingTypes();
 						Component->SetOverrideRenderMaterial(TwoSidedMaterial);
@@ -313,6 +493,16 @@ namespace UE::Chaos::ClothAsset
 								bCollapseSeamLines
 							);
 						}
+
+						if (SelectedPoints.Num() > 0)
+						{
+							static const FName PointComponentName = TEXT("Cloth_Sim3D_SelectedPoints");
+							if (UPointSetComponent* PointComponent = OutComponents.AddNewComponent<UPointSetComponent>(PointComponentName))
+							{
+								PointComponent->AddPoints(SelectedPoints, FColor::Yellow, 10.f);
+								PointComponent->SetPointMaterial(LoadObject<UMaterial>(nullptr, TEXT("/MeshModelingToolsetExp/Materials/PointSetComponentMaterial")));
+							}
+						}
 					}
 				}
 			}
@@ -413,8 +603,13 @@ namespace UE::Chaos::ClothAsset
 				const bool bShowPatternColors = (Settings && Settings->bShowPatternColors);
 				if (bShowPatternColors)
 				{
-					Private::AddPatternsColorsMesh(ClothFacade, DynamicMesh);
+					Private::AddPatternsColorsToMesh(ClothFacade, DynamicMesh);
 				}
+				const FString SelectionName = Instance.GetOutputValueByType<FString>({}, TEXT("DataflowSelection"));
+				const bool bHasFaceSelection = Private::AddFaceSelectionToMesh(ClothCollection, FName(SelectionName), DynamicMesh);
+
+				TArray<FVector> SelectedPoints;
+				Private::GetVertexSelectionForSim2D(ClothCollection, ClothFacade, FName(SelectionName), SelectedPoints);
 
 				const FVector ScaleFactor(100, 100, 1);
 
@@ -426,7 +621,7 @@ namespace UE::Chaos::ClothAsset
 					{
 						Component->SetCastShadow(false);
 						Component->SetMesh(MoveTemp(DynamicMesh));
-						UMaterialInterface* TwoSidedMaterial = bShowPatternColors
+						UMaterialInterface* TwoSidedMaterial = bShowPatternColors || bHasFaceSelection
 							? Private::GetTwoSidedMaterialForColoredPatterns()
 							: Private::GetTwoSidedMaterialForRenderingTypes();
 						Component->SetOverrideRenderMaterial(TwoSidedMaterial);
@@ -466,6 +661,17 @@ namespace UE::Chaos::ClothAsset
 								LineComponent->SetWorldScale3D(ScaleFactor);
 							}
 						}
+
+						if (SelectedPoints.Num() > 0)
+						{
+							static const FName PointComponentName = TEXT("Cloth_Sim2D_SelectedPoints");
+							if (UPointSetComponent* PointComponent = OutComponents.AddNewComponent<UPointSetComponent>(PointComponentName))
+							{
+								PointComponent->AddPoints(SelectedPoints, FColor::Yellow, 10.f);
+								PointComponent->SetPointMaterial(LoadObject<UMaterial>(nullptr, TEXT("/MeshModelingToolsetExp/Materials/PointSetOverlaidComponentMaterial")));
+								PointComponent->SetWorldScale3D(ScaleFactor);
+							}
+						}
 					}
 				}
 			}
diff --git a/Engine/Plugins/ChaosClothAssetEditor/Source/ChaosClothAssetDataflowNodes/Public/ChaosClothAsset/WeightMapToSelectionNode.h b/Engine/Plugins/ChaosClothAssetEditor/Source/ChaosClothAssetDataflowNodes/Public/ChaosClothAsset/WeightMapToSelectionNode.h
--- a/Engine/Plugins/ChaosClothAssetEditor/Source/ChaosClothAssetDataflowNodes/Public/ChaosClothAsset/WeightMapToSelectionNode.h
+++ b/Engine/Plugins/ChaosClothAssetEditor/Source/ChaosClothAssetDataflowNodes/Public/ChaosClothAsset/WeightMapToSelectionNode.h
@@ -32,7 +32,6 @@ struct FChaosClothAssetWeightMapToSelectionNode : public FDataflowNode
 {
 	GENERATED_USTRUCT_BODY()
 	DATAFLOW_NODE_DEFINE_INTERNAL(FChaosClothAssetWeightMapToSelectionNode, "WeightMapToSelection", "Cloth", "Cloth Weight Map To Selection")
-	DATAFLOW_NODE_RENDER_TYPE("SurfaceRender", FName("FClothCollection"), "Collection")
 
 public:
 
@@ -47,7 +46,7 @@ public:
 	 * The name of the select attribute that will be added to the collection.
 	 * If left empty the same name as the Weight Map  name will be used instead.
 	 */
-	UPROPERTY(EditAnywhere, Category = "Weight Map To Selection", Meta = (DataflowOutput))
+	UPROPERTY(EditAnywhere, Category = "Weight Map To Selection", Meta = (DataflowOutput, DataflowSelection))
 	FString SelectionName;
 
 	/** The type of element the selection refers to */
