Cloth dataflow visualization : change the default normal length from 20 to 7

#Synced-CL 49572843

diff --git a/Engine/Plugins/ChaosClothAssetEditor/Source/ChaosClothAssetEditor/Private/ChaosClothAsset/ClothEditorSimulationVisualization.h b/Engine/Plugins/ChaosClothAssetEditor/Source/ChaosClothAssetEditor/Private/ChaosClothAsset/ClothEditorSimulationVisualization.h
--- a/Engine/Plugins/ChaosClothAssetEditor/Source/ChaosClothAssetEditor/Private/ChaosClothAsset/ClothEditorSimulationVisualization.h
+++ b/Engine/Plugins/ChaosClothAssetEditor/Source/ChaosClothAssetEditor/Private/ChaosClothAsset/ClothEditorSimulationVisualization.h
@@ -115,9 +115,9 @@ private:
 	FNameSelectionData WeightMapSelection;
 	FNameSelectionData MorphTargetSelection;
 	FNameSelectionData AccessoryMeshSelection;
-	float PointNormalLength = 20.f;
-	float AnimatedNormalLength = 20.f;
-	float AccessoryMeshNormalLength = 20.f;
+	float PointNormalLength = 7.f;
+	float AnimatedNormalLength = 7.f;
+	float AccessoryMeshNormalLength = 7.f;
 	float AerodynamicsLengthScale = 10.f;
 	float WindVelocityLengthScale = .1f;
 	float VelocityScaleLength = .1f;
