Cloth editor : add debug draw to Max Distance to display value per vertices
- will need to extend that to other node and tools as well

#Synced-CL 49572843

diff --git a/Engine/Plugins/ChaosClothAssetEditor/Source/ChaosClothAssetDataflowNodes/Private/ChaosClothAsset/SimulationMaxDistanceConfigNode.cpp b/Engine/Plugins/ChaosClothAssetEditor/Source/ChaosClothAssetDataflowNodes/Private/ChaosClothAsset/SimulationMaxDistanceConfigNode.cpp
--- a/Engine/Plugins/ChaosClothAssetEditor/Source/ChaosClothAssetDataflowNodes/Private/ChaosClothAsset/SimulationMaxDistanceConfigNode.cpp
+++ b/Engine/Plugins/ChaosClothAssetEditor/Source/ChaosClothAssetDataflowNodes/Private/ChaosClothAsset/SimulationMaxDistanceConfigNode.cpp
@@ -6,6 +6,12 @@
 #include "ChaosClothAsset/CollectionClothSelectionFacade.h"
 #include "Chaos/CollectionPropertyFacade.h"
 
+#if WITH_EDITOR
+#include "ChaosClothAsset/ClothDataflowViewModes.h"
+#include "Dataflow/DataflowDebugDrawInterface.h"
+#include "Dataflow/DataflowRenderingViewMode.h"
+#endif
+
 #include UE_INLINE_GENERATED_CPP_BY_NAME(SimulationMaxDistanceConfigNode)
 
 FChaosClothAssetSimulationMaxDistanceConfigNode::FChaosClothAssetSimulationMaxDistanceConfigNode(const UE::Dataflow::FNodeParameters& InParam, FGuid InGuid)
@@ -47,4 +53,60 @@ void FChaosClothAssetSimulationMaxDistanceConfigNode::EvaluateClothCollection(UE
 	SelectionFacade.FindOrAddSelectionSet(FName(KinematicVertices3D), ClothCollectionGroup::SimVertices3D) =
 		FClothGeometryTools::GenerateKinematicVertices3D(ClothCollection, MaxDistanceString, FVector2f(MaxDistance.Low, MaxDistance.High), InputKinematicString);
 
-}
\ No newline at end of file
+}
+
+#if WITH_EDITOR
+
+bool FChaosClothAssetSimulationMaxDistanceConfigNode::CanDebugDraw() const
+{
+	return true;
+}
+
+bool FChaosClothAssetSimulationMaxDistanceConfigNode::CanDebugDrawViewMode(const FName& ViewModeName) const
+{
+	return ViewModeName == UE::Dataflow::FDataflowConstruction3DViewMode::Name || ViewModeName == UE::Chaos::ClothAsset::FCloth3DSimViewMode::Name;
+}
+
+void FChaosClothAssetSimulationMaxDistanceConfigNode::DebugDraw(UE::Dataflow::FContext& Context, IDataflowDebugDrawInterface& DataflowRenderingInterface, const FDebugDrawParameters& DebugDrawParameters) const
+{
+	using namespace UE::Chaos::ClothAsset;
+	
+	const FLinearColor DynamicColor(FColor::White);
+	const FLinearColor KinematicColor(FColor::Purple);
+
+	FNumberFormattingOptions NumberFormattingOptions;
+	NumberFormattingOptions.AlwaysSign = false;
+	NumberFormattingOptions.UseGrouping = false;
+	NumberFormattingOptions.RoundingMode = ERoundingMode::HalfFromZero;
+	NumberFormattingOptions.MinimumIntegralDigits = 1;
+	NumberFormattingOptions.MaximumIntegralDigits = 6;
+	NumberFormattingOptions.MinimumFractionalDigits = 2;
+	NumberFormattingOptions.MaximumFractionalDigits = 2;
+
+	const FManagedArrayCollection OutCollection = GetOutputValue(Context, &Collection, Collection);
+	const FName MaxDistanceMapName(GetValue<FString>(Context, &MaxDistance.WeightMap));
+
+	const TSharedRef<FManagedArrayCollection> ClothCollection = MakeShared<FManagedArrayCollection>(OutCollection);
+	const FCollectionClothConstFacade ClothFacade(ClothCollection);
+	const FCollectionClothSelectionConstFacade SelectionFacade(ClothCollection);
+	if (ClothFacade.IsValid() && ClothFacade.HasWeightMap(MaxDistanceMapName))
+	{
+		TConstArrayView<float> MaxDistanceMap = ClothFacade.GetWeightMap(MaxDistanceMapName);
+		const TSet<int32>* KinematicSet = SelectionFacade.FindSelectionSet(FName(KinematicVertices3D));
+		const float MinValue = MaxDistance.Low;
+		const float ValueRange = MaxDistance.High - MaxDistance.Low;
+		for (int32 Index = 0; Index < MaxDistanceMap.Num(); ++Index)
+		{
+			const TConstArrayView<FVector3f> Positions = ClothFacade.GetSimPosition3D();
+			if (Positions.IsValidIndex(Index))
+			{
+				const float MaxDistanceValue = (MinValue + (MaxDistanceMap[Index] * ValueRange));
+				const FText Text = FText::AsNumber(MaxDistanceValue, &NumberFormattingOptions);
+				const bool bIsKinematic = KinematicSet && KinematicSet->Contains(Index);
+				DataflowRenderingInterface.SetColor(bIsKinematic ? KinematicColor: DynamicColor);
+				DataflowRenderingInterface.DrawText3d(Text.ToString(), FVector(Positions[Index]));
+			}
+		}
+	}
+}
+#endif // WITH_EDITOR
\ No newline at end of file
diff --git a/Engine/Plugins/ChaosClothAssetEditor/Source/ChaosClothAssetDataflowNodes/Public/ChaosClothAsset/SimulationMaxDistanceConfigNode.h b/Engine/Plugins/ChaosClothAssetEditor/Source/ChaosClothAssetDataflowNodes/Public/ChaosClothAsset/SimulationMaxDistanceConfigNode.h
--- a/Engine/Plugins/ChaosClothAssetEditor/Source/ChaosClothAssetDataflowNodes/Public/ChaosClothAsset/SimulationMaxDistanceConfigNode.h
+++ b/Engine/Plugins/ChaosClothAssetEditor/Source/ChaosClothAssetDataflowNodes/Public/ChaosClothAsset/SimulationMaxDistanceConfigNode.h
@@ -45,4 +45,10 @@ private:
 	virtual void EvaluateClothCollection(UE::Dataflow::FContext& Context, const TSharedRef<FManagedArrayCollection>& ClothCollection) const override;
 
 	virtual void Evaluate(UE::Dataflow::FContext& Context, const FDataflowOutput* Out) const override;
+
+#if WITH_EDITOR
+	virtual bool CanDebugDraw() const override;
+	virtual bool CanDebugDrawViewMode(const FName& ViewModeName) const override;
+	virtual void DebugDraw(UE::Dataflow::FContext& Context, IDataflowDebugDrawInterface& DataflowRenderingInterface, const FDebugDrawParameters& DebugDrawParameters) const override;
+#endif
 };
