Static and skeletal mesh import nodes : unify the properties order

#Synced-CL 49661615

diff --git a/Engine/Plugins/ChaosClothAssetEditor/Source/ChaosClothAssetDataflowNodes/Public/ChaosClothAsset/StaticMeshImportNode.h b/Engine/Plugins/ChaosClothAssetEditor/Source/ChaosClothAssetDataflowNodes/Public/ChaosClothAsset/StaticMeshImportNode.h
--- a/Engine/Plugins/ChaosClothAssetEditor/Source/ChaosClothAssetDataflowNodes/Public/ChaosClothAsset/StaticMeshImportNode.h
+++ b/Engine/Plugins/ChaosClothAssetEditor/Source/ChaosClothAssetDataflowNodes/Public/ChaosClothAsset/StaticMeshImportNode.h
@@ -34,18 +34,26 @@ public:
 	UPROPERTY(EditAnywhere, Category = "Static Mesh Import", Meta = (ButtonImage = "Persona.ReimportAsset", DisplayName = "Reimport Asset", EditCondition = "!bIsLocked"))
 	FDataflowFunctionProperty Reimport;
 
-	/* Which static mesh Lod to import. */
-	UPROPERTY(EditAnywhere, Category = "Static Mesh Import", Meta = (ClampMin = "0", DisplayName = "LOD Index"))
-	int32 LODIndex = 0;
-
 	/* Import static mesh data as a simulation mesh data. */
 	UPROPERTY(EditAnywhere, Category = "Static Mesh Import")
 	bool bImportSimMesh = true;
 
+	/* Import static mesh data as render mesh data. */
+	UPROPERTY(EditAnywhere, Category = "Static Mesh Import")
+	bool bImportRenderMesh = true;
+
+	/* Which static mesh Lod to import. */
+	UPROPERTY(EditAnywhere, Category = "Static Mesh Import", Meta = (ClampMin = "0", DisplayName = "LOD Index"))
+	int32 LODIndex = 0;
+
 	/* Material section to import as sim mesh data. Use -1 to import all sections. */
 	UPROPERTY(EditAnywhere, Category = "Static Mesh Import", Meta = (ClampMin = "-1", EditCondition = "bImportSimMesh"))
 	int32 SimMeshSection = INDEX_NONE;
 
+	/* Material section to import as render mesh data. Use -1 to import all sections. */
+	UPROPERTY(EditAnywhere, Category = "Static Mesh Import", Meta = (ClampMin = "-1", EditCondition = "bImportRenderMesh"))
+	int32 RenderMeshSection = INDEX_NONE;
+
 	/**
 	 * UV channel of the static mesh to import the 2D simulation mesh patterns from.
 	 * If set to -1, or the specified UVChannel doesn't exist then the import will unwrap the 3D simulation mesh into 2D simulation mesh patterns.
@@ -57,14 +65,6 @@ public:
 	UPROPERTY(EditAnywhere, Category = "Static Mesh Import", Meta = (AllowPreserveRatio, EditCondition = "bImportSimMesh && (UVChannel >= 0)"))
 	FVector2f UVScale = { 1.f, 1.f };
 
-	/* Import static mesh data as render mesh data. */
-	UPROPERTY(EditAnywhere, Category = "Static Mesh Import")
-	bool bImportRenderMesh = true;
-
-	/* Material section to import as render mesh data. Use -1 to import all sections. */
-	UPROPERTY(EditAnywhere, Category = "Static Mesh Import", Meta = (ClampMin = "-1", EditCondition = "bImportRenderMesh"))
-	int32 RenderMeshSection = INDEX_NONE;
-
 	FChaosClothAssetStaticMeshImportNode_v2(const UE::Dataflow::FNodeParameters& InParam, FGuid InGuid = FGuid::NewGuid());
 
 private:
diff --git a/Engine/Plugins/Experimental/Dataflow/Source/DataflowNodes/Public/Dataflow/DataflowStaticMeshNodes.h b/Engine/Plugins/Experimental/Dataflow/Source/DataflowNodes/Public/Dataflow/DataflowStaticMeshNodes.h
--- a/Engine/Plugins/Experimental/Dataflow/Source/DataflowNodes/Public/Dataflow/DataflowStaticMeshNodes.h
+++ b/Engine/Plugins/Experimental/Dataflow/Source/DataflowNodes/Public/Dataflow/DataflowStaticMeshNodes.h
@@ -25,7 +25,7 @@ public:
 	UPROPERTY(EditAnywhere, Category = "Dataflow", meta = (DataflowOutput, DisplayName = "StaticMesh"))
 	TObjectPtr<const UStaticMesh> StaticMesh = nullptr;
 
-	UPROPERTY(EditAnywhere, Category = "Overrides")
+	UPROPERTY(EditAnywhere, Category = "Dataflow")
 	FName PropertyName = "StaticMesh";
 
 	FGetStaticMeshDataflowNode(const UE::Dataflow::FNodeParameters& InParam, FGuid InGuid = FGuid::NewGuid());
