Cloth Dataflow : update nodes so that they output warning through the context

#Synced-CL 49754362

diff --git a/Engine/Plugins/ChaosClothAssetEditor/Source/ChaosClothAssetDataflowNodes/Private/ChaosClothAsset/AddStitchNode.cpp b/Engine/Plugins/ChaosClothAssetEditor/Source/ChaosClothAssetDataflowNodes/Private/ChaosClothAsset/AddStitchNode.cpp
--- a/Engine/Plugins/ChaosClothAssetEditor/Source/ChaosClothAssetDataflowNodes/Private/ChaosClothAsset/AddStitchNode.cpp
+++ b/Engine/Plugins/ChaosClothAssetEditor/Source/ChaosClothAssetDataflowNodes/Private/ChaosClothAsset/AddStitchNode.cpp
@@ -77,13 +77,13 @@ void FChaosClothAssetAddStitchNode::Evaluate(UE::Dataflow::FContext& Context, co
 				}
 				else
 				{
-					FClothDataflowTools::LogAndToastWarning(*this,
-						LOCTEXT("SelectionTypeNotCorrectHeadline", "Invalid selection group."),
+					Context.Warning(
 						FText::Format(LOCTEXT("SelectionTypeNotCorrectDetails", "Selection \"{0}\" does not have its index dependency group set to \"{1}\" or \"{2}\"."),
 							FText::FromName(InSelectionName),
 							FText::FromName(ClothCollectionGroup::SimVertices2D),
-							FText::FromName(ClothCollectionGroup::SimVertices3D)));
-
+							FText::FromName(ClothCollectionGroup::SimVertices3D))
+						,this, Out
+					);
 				}
 
 				if (!StitchPairs.IsEmpty())
diff --git a/Engine/Plugins/ChaosClothAssetEditor/Source/ChaosClothAssetDataflowNodes/Private/ChaosClothAsset/AddWeightMapNode.cpp b/Engine/Plugins/ChaosClothAssetEditor/Source/ChaosClothAssetDataflowNodes/Private/ChaosClothAsset/AddWeightMapNode.cpp
--- a/Engine/Plugins/ChaosClothAssetEditor/Source/ChaosClothAssetDataflowNodes/Private/ChaosClothAsset/AddWeightMapNode.cpp
+++ b/Engine/Plugins/ChaosClothAssetEditor/Source/ChaosClothAssetDataflowNodes/Private/ChaosClothAsset/AddWeightMapNode.cpp
@@ -247,16 +247,17 @@ void FChaosClothAssetAddWeightMapNode::Evaluate(UE::Dataflow::FContext& Context,
 {
 	using namespace UE::Chaos::ClothAsset;
 
-	auto CheckSourceVertexWeights = [this](TArrayView<float>& ClothWeights, const TArray<float>& SourceVertexWeights, bool bIsSim)
+	auto CheckSourceVertexWeights = [this, &Context, Out](TArrayView<float>& ClothWeights, const TArray<float>& SourceVertexWeights, bool bIsSim)
 	{
 		if (SourceVertexWeights.Num() > 0 && SourceVertexWeights.Num() != ClothWeights.Num())
 		{
-			FClothDataflowTools::LogAndToastWarning(*this,
-				LOCTEXT("VertexCountMismatchHeadline", "Vertex count mismatch."),
+			Context.Warning(
 				FText::Format(LOCTEXT("VertexCountMismatchDetails", "{0} vertex weights in the node: {1}\n{0} vertices in the cloth: {2}"),
 					bIsSim ? FText::FromString("Sim") : FText::FromString("Render"),
 					SourceVertexWeights.Num(),
-					ClothWeights.Num()));
+					ClothWeights.Num())
+				, this, Out
+			);
 		}
 	};
 
@@ -283,10 +284,11 @@ void FChaosClothAssetAddWeightMapNode::Evaluate(UE::Dataflow::FContext& Context,
 				if (ClothSimWeights.Num() != ClothFacade.GetNumSimVertices3D())
 				{
 					check(ClothSimWeights.Num() == 0);
-					FClothDataflowTools::LogAndToastWarning(*this,
-						LOCTEXT("InvalidSimWeightMapNameHeadline", "Invalid weight map name."),
+					Context.Warning(
 						FText::Format(LOCTEXT("InvalidSimWeightMapNameDetails", "Could not create a sim weight map with name \"{0}\" (reserved name? wrong type?)."),
-							FText::FromName(InName)));
+							FText::FromName(InName))
+						, this, Out
+					);
 				}
 				else
 				{
@@ -306,10 +308,11 @@ void FChaosClothAssetAddWeightMapNode::Evaluate(UE::Dataflow::FContext& Context,
 				if (ClothRenderWeights.Num() != ClothFacade.GetNumRenderVertices())
 				{
 					check(ClothRenderWeights.Num() == 0);
-					FClothDataflowTools::LogAndToastWarning(*this,
-						LOCTEXT("InvalidRenderWeightMapNameHeadline", "Invalid weight map name."),
+					Context.Warning(
 						FText::Format(LOCTEXT("InvalidRenderWeightMapNameDetails", "Could not create a render weight map with name \"{0}\" (reserved name? wrong type?)."),
-							FText::FromName(InName)));
+							FText::FromName(InName))
+						,this, Out
+					);
 				}
 				else
 				{
diff --git a/Engine/Plugins/ChaosClothAssetEditor/Source/ChaosClothAssetDataflowNodes/Private/ChaosClothAsset/AttributeNode.cpp b/Engine/Plugins/ChaosClothAssetEditor/Source/ChaosClothAssetDataflowNodes/Private/ChaosClothAsset/AttributeNode.cpp
--- a/Engine/Plugins/ChaosClothAssetEditor/Source/ChaosClothAssetDataflowNodes/Private/ChaosClothAsset/AttributeNode.cpp
+++ b/Engine/Plugins/ChaosClothAssetEditor/Source/ChaosClothAssetDataflowNodes/Private/ChaosClothAsset/AttributeNode.cpp
@@ -75,10 +75,10 @@ void FChaosClothAssetAttributeNode_v2::Evaluate(UE::Dataflow::FContext& Context,
 			}
 			else if (!GroupName.IsNone())
 			{
-				FClothDataflowTools::LogAndToastWarning(
-					*this,
-					LOCTEXT("CreateAttributeHeadline", "Invalid Group"),
-					FText::Format(LOCTEXT("CreateAttributeDetail", "No group \"{0}\" currently exists on the input collection"), FText::FromName(GroupName)));
+				Context.Warning(
+					FText::Format(LOCTEXT("CreateAttributeDetail", "No group \"{0}\" currently exists on the input collection"), FText::FromName(GroupName))
+					, this, Out
+				);
 			}
 		}
 		SetValue(Context, MoveTemp(*ClothCollection), &Collection);
@@ -151,10 +151,10 @@ void FChaosClothAssetAttributeNode::Evaluate(UE::Dataflow::FContext& Context, co
 			}
 			else if (!GroupName.IsNone())
 			{
-				FClothDataflowTools::LogAndToastWarning(
-					*this,
-					LOCTEXT("CreateAttributeHeadline", "Invalid Group"),
-					FText::Format(LOCTEXT("CreateAttributeDetail", "No group \"{0}\" currently exists on the input collection"), FText::FromName(GroupName)));
+				Context.Warning(
+					FText::Format(LOCTEXT("CreateAttributeDetail", "No group \"{0}\" currently exists on the input collection"), FText::FromName(GroupName))
+					, this, Out
+				);
 			}
 		}
 		SetValue(Context, MoveTemp(*ClothCollection), &Collection);
diff --git a/Engine/Plugins/ChaosClothAssetEditor/Source/ChaosClothAssetDataflowNodes/Private/ChaosClothAsset/BlendVerticesNode.cpp b/Engine/Plugins/ChaosClothAssetEditor/Source/ChaosClothAssetDataflowNodes/Private/ChaosClothAsset/BlendVerticesNode.cpp
--- a/Engine/Plugins/ChaosClothAssetEditor/Source/ChaosClothAssetDataflowNodes/Private/ChaosClothAsset/BlendVerticesNode.cpp
+++ b/Engine/Plugins/ChaosClothAssetEditor/Source/ChaosClothAssetDataflowNodes/Private/ChaosClothAsset/BlendVerticesNode.cpp
@@ -74,12 +74,13 @@ void FChaosClothAssetBlendVerticesNode::Evaluate(UE::Dataflow::FContext& Context
 		FCollectionClothFacade ClothFacade(ClothCollection);
 		const FCollectionClothConstFacade BlendClothFacade(BlendClothCollection);
 
-		auto LogVertexCountMismatch = [this](const FText& VertexType, int32 CollectionCount, int32 BlendCollectionCount)
+		auto LogVertexCountMismatch = [this, &Context, Out](const FText& VertexType, int32 CollectionCount, int32 BlendCollectionCount)
 			{
-				UE::Chaos::ClothAsset::FClothDataflowTools::LogAndToastWarning(*this,
-					FText::Format(LOCTEXT("VertexCountMismatchHeadline", "Failed to blend {0} attributes."), VertexType),
-					FText::Format(LOCTEXT("VertexCountMismatchDetails", "Vertex count mismatch {0} != {1}. Set 'Require Same Vertex Counts' to false to disable this check."),
-						CollectionCount, BlendCollectionCount));
+				Context.Warning(
+					FText::Format(LOCTEXT("VertexCountMismatchDetails", "{0} Vertex count mismatch {1} != {2}. Set 'Require Same Vertex Counts' to false to disable this check."),
+						VertexType, CollectionCount, BlendCollectionCount)
+					, this, Out
+				);
 			};
 
 		if (ClothFacade.IsValid())
@@ -196,10 +197,7 @@ void FChaosClothAssetBlendVerticesNode::Evaluate(UE::Dataflow::FContext& Context
 			}
 			else
 			{
-				UE::Chaos::ClothAsset::FClothDataflowTools::LogAndToastWarning(*this,
-					LOCTEXT("InvalidBlendCollectionHeadline", "Invalid Blend Collection"),
-					LOCTEXT("InvalidBlendCollectionDetails", "Input Blend Collection is not a valid Cloth Collection.")
-				);
+				Context.Warning(LOCTEXT("InvalidBlendCollectionDetails", "Input Blend Collection is not a valid Cloth Collection."), this, Out);
 			}
 		}
 
diff --git a/Engine/Plugins/ChaosClothAssetEditor/Source/ChaosClothAssetDataflowNodes/Private/ChaosClothAsset/ClothDataflowTools.cpp b/Engine/Plugins/ChaosClothAssetEditor/Source/ChaosClothAssetDataflowNodes/Private/ChaosClothAsset/ClothDataflowTools.cpp
--- a/Engine/Plugins/ChaosClothAssetEditor/Source/ChaosClothAssetDataflowNodes/Private/ChaosClothAsset/ClothDataflowTools.cpp
+++ b/Engine/Plugins/ChaosClothAssetEditor/Source/ChaosClothAssetDataflowNodes/Private/ChaosClothAsset/ClothDataflowTools.cpp
@@ -5,6 +5,7 @@
 #include "ChaosClothAsset/CollectionClothFacade.h"
 #include "ChaosClothAsset/ClothCollectionGroup.h"
 #include "Dataflow/DataflowNode.h"
+#include "Dataflow/DataflowNodeParameters.h"
 #include "DynamicMesh/NonManifoldMappingSupport.h"
 #include "Framework/Notifications/NotificationManager.h"
 #include "Interfaces/ITargetPlatformManagerModule.h"
diff --git a/Engine/Plugins/ChaosClothAssetEditor/Source/ChaosClothAssetDataflowNodes/Private/ChaosClothAsset/ClothDataflowTools.h b/Engine/Plugins/ChaosClothAssetEditor/Source/ChaosClothAssetDataflowNodes/Private/ChaosClothAsset/ClothDataflowTools.h
--- a/Engine/Plugins/ChaosClothAssetEditor/Source/ChaosClothAssetDataflowNodes/Private/ChaosClothAsset/ClothDataflowTools.h
+++ b/Engine/Plugins/ChaosClothAssetEditor/Source/ChaosClothAssetDataflowNodes/Private/ChaosClothAsset/ClothDataflowTools.h
@@ -16,6 +16,11 @@ class IPropertyHandle;
 struct FMeshResizingRBFInterpolationData;
 struct FMeshResizingCustomRegion;
 
+namespace UE::Dataflow
+{
+	class FContext;
+}
+
 namespace UE::Chaos::ClothAsset
 {
 	/**
@@ -28,6 +33,7 @@ namespace UE::Chaos::ClothAsset
 		static void AddSimPatternsFromSkeletalMeshSection(const TSharedRef<FManagedArrayCollection>& ClothCollection, const FSkeletalMeshLODModel& SkeletalMeshModel, const int32 SectionIndex, const int32 UVChannelIndex, const FVector2f& UVScale, bool bImportNormals = false, TArray<int32>* OutSim2DToSourceVertex = nullptr);
 
 		static void LogAndToastWarning(const FDataflowNode& DataflowNode, const FText& Headline, const FText& Details);
+		static void LogAndToastWarning(UE::Dataflow::FContext& Context, const FDataflowNode& DataflowNode, const FText& Headline, const FText& Details);
 
 		/**
 		 * Turn a string into a valid collection group or attribute name.
diff --git a/Engine/Plugins/ChaosClothAssetEditor/Source/ChaosClothAssetDataflowNodes/Private/ChaosClothAsset/DeleteElementNode.cpp b/Engine/Plugins/ChaosClothAssetEditor/Source/ChaosClothAssetDataflowNodes/Private/ChaosClothAsset/DeleteElementNode.cpp
--- a/Engine/Plugins/ChaosClothAssetEditor/Source/ChaosClothAssetDataflowNodes/Private/ChaosClothAsset/DeleteElementNode.cpp
+++ b/Engine/Plugins/ChaosClothAssetEditor/Source/ChaosClothAssetDataflowNodes/Private/ChaosClothAsset/DeleteElementNode.cpp
@@ -315,10 +315,11 @@ void FChaosClothAssetDeleteElementNode::Evaluate(UE::Dataflow::FContext& Context
 					}
 					else
 					{
-						FClothDataflowTools::LogAndToastWarning(*this,
-							LOCTEXT("SelectionTypeNotCorrectHeadline", "Invalid selection group."),
+						Context.Warning(
 							FText::Format(LOCTEXT("SelectionTypeNotCorrectDetails", "Selection \"{0}\" does not have its index dependency group set to a type known to the DeleteElement node."),
-								FText::FromName(InSelectionName)));
+								FText::FromName(InSelectionName))
+							,this, Out
+						);
 
 					}
 
diff --git a/Engine/Plugins/ChaosClothAssetEditor/Source/ChaosClothAssetDataflowNodes/Private/ChaosClothAsset/GenerateSimMorphTargetNode.cpp b/Engine/Plugins/ChaosClothAssetEditor/Source/ChaosClothAssetDataflowNodes/Private/ChaosClothAsset/GenerateSimMorphTargetNode.cpp
--- a/Engine/Plugins/ChaosClothAssetEditor/Source/ChaosClothAssetDataflowNodes/Private/ChaosClothAsset/GenerateSimMorphTargetNode.cpp
+++ b/Engine/Plugins/ChaosClothAssetEditor/Source/ChaosClothAssetDataflowNodes/Private/ChaosClothAsset/GenerateSimMorphTargetNode.cpp
@@ -75,10 +75,10 @@ void FChaosClothAssetGenerateSimMorphTargetNode::Evaluate(UE::Dataflow::FContext
 							const int32 ExistingMorphTarget = ClothFacade.FindSimMorphTargetIndexByName(MorphTargetName);
 							if (ExistingMorphTarget != INDEX_NONE)
 							{
-								UE::Chaos::ClothAsset::FClothDataflowTools::LogAndToastWarning(*this,
-									LOCTEXT("DuplicateMorphTargetHeadline", "Duplicate Sim Morph Target"),
+								Context.Warning(
 									FText::Format(LOCTEXT("DuplicateMorphTargetDetails", "Existing Sim Morph target with name '{0}' will be replaced."),
-										FText::FromString(MorphTargetName)));
+										FText::FromString(MorphTargetName))
+									,this, Out);
 							}
 
 							FCollectionClothSimMorphTargetFacade MorphTargetFacade = ExistingMorphTarget == INDEX_NONE ? ClothFacade.AddGetSimMorphTarget() : ClothFacade.GetSimMorphTarget(ExistingMorphTarget);
@@ -86,27 +86,21 @@ void FChaosClothAssetGenerateSimMorphTargetNode::Evaluate(UE::Dataflow::FContext
 						}
 						else
 						{
-							UE::Chaos::ClothAsset::FClothDataflowTools::LogAndToastWarning(*this,
-								LOCTEXT("IdenticalCollectionsHeadline", "Identical sim collections"),
-								LOCTEXT("IdenticalCollectionsDetails", "No morph targets generated because the sim collections are identical.")
-							);
+							Context.Warning(LOCTEXT("IdenticalCollectionsDetails", "No morph targets generated because the sim collections are identical."), this, Out);
 						}
 					}
 					else
 					{
-						UE::Chaos::ClothAsset::FClothDataflowTools::LogAndToastWarning(*this,
-							LOCTEXT("VertexCountMismatchHeadline", "Failed to generate morph target"),
-							FText::Format(LOCTEXT("VertexCountMismatchDetail", "Vertex count mismatch {0} != {1}."),
-								NumSimVertices, NumMorphSimVertices));
+						Context.Warning(
+							FText::Format(LOCTEXT("VertexCountMismatchDetail", "Vertex count mismatch {0} != {1}."), NumSimVertices, NumMorphSimVertices)
+							, this, Out
+						);
 					}
 				}
 			}
 			else
 			{
-				UE::Chaos::ClothAsset::FClothDataflowTools::LogAndToastWarning(*this,
-					LOCTEXT("InvalidMorphTargetCollectionHeadline", "Invalid Morph Target Collection"),
-					LOCTEXT("InvalidMorphTargetCollectionDetails", "Input Morph Target Collection is not a valid Cloth Collection.")
-				);
+				Context.Warning(LOCTEXT("InvalidMorphTargetCollectionDetails", "Input Morph Target Collection is not a valid Cloth Collection."), this, Out);
 			}
 		}
 
diff --git a/Engine/Plugins/ChaosClothAssetEditor/Source/ChaosClothAssetDataflowNodes/Private/ChaosClothAsset/ImportNode.cpp b/Engine/Plugins/ChaosClothAssetEditor/Source/ChaosClothAssetDataflowNodes/Private/ChaosClothAsset/ImportNode.cpp
--- a/Engine/Plugins/ChaosClothAssetEditor/Source/ChaosClothAssetDataflowNodes/Private/ChaosClothAsset/ImportNode.cpp
+++ b/Engine/Plugins/ChaosClothAssetEditor/Source/ChaosClothAssetDataflowNodes/Private/ChaosClothAsset/ImportNode.cpp
@@ -47,9 +47,7 @@ void FChaosClothAssetImportNode::Evaluate(UE::Dataflow::FContext& Context, const
 					{
 						// Can't create a loop
 						bSetValue = false;
-						FClothDataflowTools::LogAndToastWarning(*this,
-							LOCTEXT("RecursiveAssetLoopHeadline", "Recursive asset loop."),
-							LOCTEXT("RecursiveAssetLoopDetails", "The source asset cannot be the same as the terminal asset."));
+						Context.Warning(LOCTEXT("RecursiveAssetLoopDetails", "The source asset cannot be the same as the terminal asset."), this, Out);
 					}
 				}
 			}
diff --git a/Engine/Plugins/ChaosClothAssetEditor/Source/ChaosClothAssetDataflowNodes/Private/ChaosClothAsset/MakeClothAssetNode.cpp b/Engine/Plugins/ChaosClothAssetEditor/Source/ChaosClothAssetDataflowNodes/Private/ChaosClothAsset/MakeClothAssetNode.cpp
--- a/Engine/Plugins/ChaosClothAssetEditor/Source/ChaosClothAssetDataflowNodes/Private/ChaosClothAsset/MakeClothAssetNode.cpp
+++ b/Engine/Plugins/ChaosClothAssetEditor/Source/ChaosClothAssetDataflowNodes/Private/ChaosClothAsset/MakeClothAssetNode.cpp
@@ -91,7 +91,7 @@ TArray<TSharedRef<const FManagedArrayCollection>> FChaosClothAssetMakeClothAsset
 
 	TArray<TSharedRef<const FManagedArrayCollection>> CollectionLodValues;
 	int32 LastValidLodIndex = INDEX_NONE;
-	auto AddCleanedCollectionLod = [this, &LastValidLodIndex, &CollectionLodValues](int32 LodIndex, TSharedRef<FManagedArrayCollection> CollectionLodValue)
+	auto AddCleanedCollectionLod = [this, &Context, &LastValidLodIndex, &CollectionLodValues](int32 LodIndex, TSharedRef<FManagedArrayCollection> CollectionLodValue)
 		{
 			FCollectionClothFacade ClothFacade(CollectionLodValue);
 			if (ClothFacade.HasValidRenderData())  // The cloth collection must at least have a render mesh
@@ -104,22 +104,24 @@ TArray<TSharedRef<const FManagedArrayCollection>> FChaosClothAssetMakeClothAsset
 				ClothFacade.DefineSchema();
 				ClothFacade.Initialize(FCollectionClothConstFacade(CollectionLodValues[LastValidLodIndex]));
 
-				FClothDataflowTools::LogAndToastWarning(*this,
-					LOCTEXT("InvalidInputLodNHeadline", "Invalid input LOD."),
+				Context.Warning(
 					FText::Format(
 						LOCTEXT("InvalidInputLodNDetails",
 							"Invalid or empty input LOD for LOD {0}.\n"
 							"Using the previous valid LOD {1} instead."),
 						LodIndex,
-						LastValidLodIndex));
+						LastValidLodIndex)
+					, this
+				);
 			}
 			else
 			{
-				FClothDataflowTools::LogAndToastWarning(*this,
-					LOCTEXT("InvalidInputLod0Headline", "Invalid input LOD 0."),
+				Context.Warning(
 					LOCTEXT("InvalidInputLod0Details",
 						"Invalid or empty input LOD for LOD 0.\n"
-						"LOD 0 cannot be empty in order to construct a valid Cloth Asset."));
+						"LOD 0 cannot be empty in order to construct a valid Cloth Asset.")
+					, this
+				);
 				return;
 			}
 			CollectionLodValues.Emplace(MoveTemp(CollectionLodValue));
diff --git a/Engine/Plugins/ChaosClothAssetEditor/Source/ChaosClothAssetDataflowNodes/Private/ChaosClothAsset/MergeClothCollectionsNode.cpp b/Engine/Plugins/ChaosClothAssetEditor/Source/ChaosClothAssetDataflowNodes/Private/ChaosClothAsset/MergeClothCollectionsNode.cpp
--- a/Engine/Plugins/ChaosClothAssetEditor/Source/ChaosClothAssetDataflowNodes/Private/ChaosClothAsset/MergeClothCollectionsNode.cpp
+++ b/Engine/Plugins/ChaosClothAssetEditor/Source/ChaosClothAssetDataflowNodes/Private/ChaosClothAsset/MergeClothCollectionsNode.cpp
@@ -572,8 +572,7 @@ void FChaosClothAssetMergeClothCollectionsNode_v2::Evaluate(UE::Dataflow::FConte
 				}
 				else
 				{
-					static const FText ErrorHeadline = LOCTEXT("IncompatibleSkeletalMeshesHeadline", "Incompatible Skeletal Meshes.");
-					FClothDataflowTools::LogAndToastWarning(*this, ErrorHeadline, IncompatibleErrorDetails);
+					Context.Warning(IncompatibleErrorDetails, this, Out);
 				}
 			}
 
@@ -922,8 +921,7 @@ void FChaosClothAssetMergeClothCollectionsNode::Evaluate(UE::Dataflow::FContext&
 				}
 				else
 				{
-					static const FText ErrorHeadline = LOCTEXT("IncompatibleSkeletalMeshesHeadline", "Incompatible Skeletal Meshes.");
-					FClothDataflowTools::LogAndToastWarning(*this, ErrorHeadline, IncompatibleErrorDetails);
+					Context.Warning(IncompatibleErrorDetails, this, Out);
 				}
 			}
 
diff --git a/Engine/Plugins/ChaosClothAssetEditor/Source/ChaosClothAssetDataflowNodes/Private/ChaosClothAsset/ProceduralSelectionNode.cpp b/Engine/Plugins/ChaosClothAssetEditor/Source/ChaosClothAssetDataflowNodes/Private/ChaosClothAsset/ProceduralSelectionNode.cpp
--- a/Engine/Plugins/ChaosClothAssetEditor/Source/ChaosClothAssetDataflowNodes/Private/ChaosClothAsset/ProceduralSelectionNode.cpp
+++ b/Engine/Plugins/ChaosClothAssetEditor/Source/ChaosClothAssetDataflowNodes/Private/ChaosClothAsset/ProceduralSelectionNode.cpp
@@ -68,9 +68,11 @@ void FChaosClothAssetProceduralSelectionNode::Evaluate(UE::Dataflow::FContext& C
 					}
 					else
 					{
-						FClothDataflowTools::LogAndToastWarning(*this, LOCTEXT("ConversionFailureHeadline", "Conversion Failure"),
+						Context.Warning(
 							FText::Format(LOCTEXT("ConversionFailureDetails", "Failed to convert selection '{0}' to group type '{1}' either because '{0}' does not exist or the conversion type is unsupported."),
-								FText::FromName(InConversionName), FText::FromName(SelectionGroupName)));
+								FText::FromName(InConversionName), FText::FromName(SelectionGroupName))
+							, this, Out
+						);
 					}
 				}
 			}
diff --git a/Engine/Plugins/ChaosClothAssetEditor/Source/ChaosClothAssetDataflowNodes/Private/ChaosClothAsset/ProxyDeformerNode.cpp b/Engine/Plugins/ChaosClothAssetEditor/Source/ChaosClothAssetDataflowNodes/Private/ChaosClothAsset/ProxyDeformerNode.cpp
--- a/Engine/Plugins/ChaosClothAssetEditor/Source/ChaosClothAssetDataflowNodes/Private/ChaosClothAsset/ProxyDeformerNode.cpp
+++ b/Engine/Plugins/ChaosClothAssetEditor/Source/ChaosClothAssetDataflowNodes/Private/ChaosClothAsset/ProxyDeformerNode.cpp
@@ -467,9 +467,10 @@ void FChaosClothAssetProxyDeformerNode::Evaluate(UE::Dataflow::FContext& Context
 			FName SimVertexSelectionName = FName(*GetValue<FString>(Context, &SimVertexSelection.StringValue));
 			if (SimVertexSelectionName != NAME_None && (!SelectionFacade.IsValid() || !SelectionFacade.FindSelectionSet(SimVertexSelectionName)))
 			{
-				FClothDataflowTools::LogAndToastWarning(*this,
-					LOCTEXT("HasSimVertexSelectionHeadline", "Unknown SimVertexSelection."),
-					LOCTEXT("HasSimVertexSelectionDetails", "The specified SimVertexSelection does't exist within the input Cloth Collection."));
+				Context.Warning(
+					LOCTEXT("HasSimVertexSelectionDetails", "The specified SimVertexSelection does't exist within the input Cloth Collection.")
+					, this, Out
+				);
 				SimVertexSelectionName = NAME_None;
 			}
 
diff --git a/Engine/Plugins/ChaosClothAssetEditor/Source/ChaosClothAssetDataflowNodes/Private/ChaosClothAsset/RemeshNode.cpp b/Engine/Plugins/ChaosClothAssetEditor/Source/ChaosClothAssetDataflowNodes/Private/ChaosClothAsset/RemeshNode.cpp
--- a/Engine/Plugins/ChaosClothAssetEditor/Source/ChaosClothAssetDataflowNodes/Private/ChaosClothAsset/RemeshNode.cpp
+++ b/Engine/Plugins/ChaosClothAssetEditor/Source/ChaosClothAssetDataflowNodes/Private/ChaosClothAsset/RemeshNode.cpp
@@ -1850,9 +1850,10 @@ void FChaosClothAssetRemeshNode_v2::Evaluate(UE::Dataflow::FContext& Context, co
 
 		if ((bRemeshSim || bRemeshRender) && ClothFacade.IsValid(EClothCollectionExtendedSchemas::RenderDeformer))
 		{
-			FClothDataflowTools::LogAndToastWarning(*this,
-				LOCTEXT("InputHasDeformerDataHeadline", "Proxy Deformer Data Found"),
-				LOCTEXT("InputHasDeformerDataDetails", "The input Cloth Collection has Proxy Deformer data that will be removed by the Remesh node. Default deformer bindings will be computed in the final asset. Consider placing ProxyDeformer Node after the Remesh Node."));
+			Context.Warning(
+				LOCTEXT("InputHasDeformerDataDetails", "The input Cloth Collection has Proxy Deformer data that will be removed by the Remesh node. Default deformer bindings will be computed in the final asset. Consider placing ProxyDeformer Node after the Remesh Node.")
+				, this, Out
+			);
 
 			// Don't copy proxy deformer data
 			const TArray<FName> GroupsToSkip = TArray<FName>();
@@ -1885,9 +1886,10 @@ void FChaosClothAssetRemeshNode_v2::Evaluate(UE::Dataflow::FContext& Context, co
 
 				if (TotalArea2D < UE_SMALL_NUMBER)
 				{
-					FClothDataflowTools::LogAndToastWarning(*this,
-						LOCTEXT("SimMeshAreaZeroHeadline", "2D Sim Mesh Has Zero Area"),
-						LOCTEXT("SimMeshAreaZeroDetails", "The input Cloth Collection has a 2D Sim mesh with zero area. Sim remeshing will not be performed"));
+					Context.Warning(
+						LOCTEXT("SimMeshAreaZeroDetails", "The input Cloth Collection has a 2D Sim mesh with zero area. Sim remeshing will not be performed")
+						, this, Out
+					);
 				}
 				else
 				{
@@ -2951,9 +2953,10 @@ void FChaosClothAssetRemeshNode::Evaluate(UE::Dataflow::FContext& Context, const
 
 		if ((bRemeshSim || bRemeshRender) && ClothFacade.IsValid(EClothCollectionExtendedSchemas::RenderDeformer))
 		{
-			FClothDataflowTools::LogAndToastWarning(*this, 
-				LOCTEXT("InputHasDeformerDataHeadline", "Proxy Deformer Data Found"),
-				LOCTEXT("InputHasDeformerDataDetails", "The input Cloth Collection has Proxy Deformer data that will be removed by the Remesh node. Default deformer bindings will be computed in the final asset. Consider placing ProxyDeformer Node after the Remesh Node."));
+			Context.Warning(
+				LOCTEXT("InputHasDeformerDataDetails", "The input Cloth Collection has Proxy Deformer data that will be removed by the Remesh node. Default deformer bindings will be computed in the final asset. Consider placing ProxyDeformer Node after the Remesh Node.")
+				, this, Out
+			);
 
 			// Don't copy proxy deformer data
 			const TArray<FName> GroupsToSkip = TArray<FName>();
@@ -2988,9 +2991,10 @@ void FChaosClothAssetRemeshNode::Evaluate(UE::Dataflow::FContext& Context, const
 
 				if (TotalArea2D < UE_SMALL_NUMBER)
 				{
-					FClothDataflowTools::LogAndToastWarning(*this,
-						LOCTEXT("SimMeshAreaZeroHeadline", "2D Sim Mesh Has Zero Area"),
-						LOCTEXT("SimMeshAreaZeroDetails", "The input Cloth Collection has a 2D Sim mesh with zero area. Sim remeshing will not be performed"));
+					Context.Warning(
+						LOCTEXT("SimMeshAreaZeroDetails", "The input Cloth Collection has a 2D Sim mesh with zero area. Sim remeshing will not be performed")
+						, this, Out
+					);
 				}
 				else
 				{
diff --git a/Engine/Plugins/ChaosClothAssetEditor/Source/ChaosClothAssetDataflowNodes/Private/ChaosClothAsset/SelectionNode.cpp b/Engine/Plugins/ChaosClothAssetEditor/Source/ChaosClothAssetDataflowNodes/Private/ChaosClothAsset/SelectionNode.cpp
--- a/Engine/Plugins/ChaosClothAssetEditor/Source/ChaosClothAssetDataflowNodes/Private/ChaosClothAsset/SelectionNode.cpp
+++ b/Engine/Plugins/ChaosClothAssetEditor/Source/ChaosClothAssetDataflowNodes/Private/ChaosClothAsset/SelectionNode.cpp
@@ -308,10 +308,11 @@ void FChaosClothAssetSelectionNode_v2::OnImport(UE::Dataflow::FContext& Context)
 	}
 	else
 	{
-		FClothDataflowTools::LogAndToastWarning(*this,
-			LOCTEXT("NoMatchingSelectionFoundHeadline", "No matching selection found"),
+		Context.Warning(
 			FText::Format(LOCTEXT("NoMatchingSelectionFoundDetails", "No matching selection with the name \"{0}\" has been found to import."),
-				FText::FromName(InInputName)));
+				FText::FromName(InInputName))
+			, this
+		);
 	}
 }
 
@@ -335,10 +336,11 @@ PRAGMA_DISABLE_DEPRECATION_WARNINGS
 PRAGMA_ENABLE_DEPRECATION_WARNINGS
 	else
 	{
-		FClothDataflowTools::LogAndToastWarning(*this,
-			LOCTEXT("NoMatchingSecondarySelectionFoundHeadline", "No matching secondary selection found"),
+		Context.Warning(
 			FText::Format(LOCTEXT("NoMatchingSecondarySelectionFoundDetails", "No matching secondary selection with the name \"{0}\" has been found to import."),
-				FText::FromName(InInputName)));
+				FText::FromName(InInputName))
+			, this
+		);
 	}
 }
 
@@ -375,7 +377,7 @@ void FChaosClothAssetSelectionNode_v2::Evaluate(UE::Dataflow::FContext& Context,
 {
 	using namespace UE::Chaos::ClothAsset;
 
-	auto CopyIntoSelection = [this](const TSharedRef<FManagedArrayCollection> SelectionCollection,
+	auto CopyIntoSelection = [this, &Context, Out](const TSharedRef<FManagedArrayCollection> SelectionCollection,
 		const FName& SelectionGroupName,
 		const TSet<int32>& SourceIndices,
 		TSet<int32>& DestSelectionSet)
@@ -407,9 +409,9 @@ void FChaosClothAssetSelectionNode_v2::Evaluate(UE::Dataflow::FContext& Context,
 			if (bFoundAnyInvalidIndex)
 			{
 				// Toast once
-				const FText ToastErrorMessage = FText::Format(LOCTEXT("AnySelectionIndexOutOfBoundsDetails", "Found invalid selection indices for group \"{0}.\" See log for details"),
+				const FText WarningMessage = FText::Format(LOCTEXT("AnySelectionIndexOutOfBoundsDetails", "Found invalid selection indices for group \"{0}.\" See log for details"),
 					FText::FromName(SelectionGroupName));
-				FClothDataflowTools::LogAndToastWarning(*this, LOCTEXT("AnySelectionIndexOutOfBoundsHeadline", "Invalid selection"), ToastErrorMessage);
+				Context.Warning(WarningMessage, this, Out);
 			}
 		};
 
@@ -583,7 +585,7 @@ void FChaosClothAssetSelectionNode::Evaluate(UE::Dataflow::FContext& Context, co
 {
 	using namespace UE::Chaos::ClothAsset;
 
-	auto CopyIntoSelection = [this](const TSharedRef<FManagedArrayCollection> SelectionCollection, 
+	auto CopyIntoSelection = [this, &Context, Out](const TSharedRef<FManagedArrayCollection> SelectionCollection, 
 		const FName& SelectionGroupName, 
 		const TSet<int32>& SourceIndices,
 		TSet<int32>& DestSelectionSet)
@@ -615,9 +617,9 @@ void FChaosClothAssetSelectionNode::Evaluate(UE::Dataflow::FContext& Context, co
 		if (bFoundAnyInvalidIndex)
 		{
 			// Toast once
-			const FText ToastErrorMessage = FText::Format(LOCTEXT("AnySelectionIndexOutOfBoundsDetails", "Found invalid selection indices for group \"{0}.\" See log for details"),
+			const FText WarningMessage = FText::Format(LOCTEXT("AnySelectionIndexOutOfBoundsDetails", "Found invalid selection indices for group \"{0}.\" See log for details"),
 				FText::FromName(SelectionGroupName));
-			FClothDataflowTools::LogAndToastWarning(*this, LOCTEXT("AnySelectionIndexOutOfBoundsHeadline", "Invalid selection"), ToastErrorMessage);
+			Context.Warning(WarningMessage, this, Out);
 		}
 	};
 
diff --git a/Engine/Plugins/ChaosClothAssetEditor/Source/ChaosClothAssetDataflowNodes/Private/ChaosClothAsset/SelectionToIntMapNode.cpp b/Engine/Plugins/ChaosClothAssetEditor/Source/ChaosClothAssetDataflowNodes/Private/ChaosClothAsset/SelectionToIntMapNode.cpp
--- a/Engine/Plugins/ChaosClothAssetEditor/Source/ChaosClothAssetDataflowNodes/Private/ChaosClothAsset/SelectionToIntMapNode.cpp
+++ b/Engine/Plugins/ChaosClothAssetEditor/Source/ChaosClothAssetDataflowNodes/Private/ChaosClothAsset/SelectionToIntMapNode.cpp
@@ -41,6 +41,7 @@ void FChaosClothAssetSelectionToIntMapNode::Evaluate(UE::Dataflow::FContext& Con
 			{
 				const FName& SelectionGroup = SelectionFacade.GetSelectionGroup(InSelectionName);
 
+				bool bValidMap = true;
 				bool bIsNewMap = false;
 				if (!ClothFacade.HasUserDefinedAttribute<int32>(InIntMapName, SelectionGroup))
 				{
@@ -50,37 +51,42 @@ void FChaosClothAssetSelectionToIntMapNode::Evaluate(UE::Dataflow::FContext& Con
 					}
 					else
 					{
-						FClothDataflowTools::LogAndToastWarning(*this,
-							LOCTEXT("IntMapNameInvalidHeadline", "Invalid Int Map Name."),
+						bValidMap = false;
+						Context.Warning(
 							FText::Format(LOCTEXT("IntMapNameInvalidDetails", "Could not create (or find) int map with name \"{0}\" in group \"{1}\""),
 								FText::FromName(InIntMapName),
-								FText::FromName(SelectionGroup)));
+								FText::FromName(SelectionGroup))
+							, this, Out
+						);
 					}
 				}
-
-				TArrayView<int32> IntMap = ClothFacade.GetUserDefinedAttribute<int32>(InIntMapName, SelectionGroup);
-
-				if (bIsNewMap || !bKeepExistingUnselectedValues)
+				if (bValidMap)
 				{
-					for (int32& MapValue : IntMap)
+					TArrayView<int32> IntMap = ClothFacade.GetUserDefinedAttribute<int32>(InIntMapName, SelectionGroup);
+
+					if (bIsNewMap || !bKeepExistingUnselectedValues)
 					{
-						MapValue = UnselectedValue;
+						for (int32& MapValue : IntMap)
+						{
+							MapValue = UnselectedValue;
+						}
 					}
-				}
-				for (const int32 Index : *SelectionSet)
-				{
-					if (IntMap.IsValidIndex(Index))
+					for (const int32 Index : *SelectionSet)
 					{
-						IntMap[Index] = SelectedValue;
+						if (IntMap.IsValidIndex(Index))
+						{
+							IntMap[Index] = SelectedValue;
+						}
 					}
 				}
 			}
 			else
 			{
-				FClothDataflowTools::LogAndToastWarning(*this,
-					LOCTEXT("SelectionNameNotFoundHeadline", "Invalid selection name."),
+				Context.Warning(
 					FText::Format(LOCTEXT("SelectionNameNotFoundDetails", "Selection \"{0}\" was not found in the collection."),
-						FText::FromName(InSelectionName)));
+						FText::FromName(InSelectionName))
+					, this, Out
+				);
 			}
 		}
 
diff --git a/Engine/Plugins/ChaosClothAssetEditor/Source/ChaosClothAssetDataflowNodes/Private/ChaosClothAsset/SelectionToWeightMapNode.cpp b/Engine/Plugins/ChaosClothAssetEditor/Source/ChaosClothAssetDataflowNodes/Private/ChaosClothAsset/SelectionToWeightMapNode.cpp
--- a/Engine/Plugins/ChaosClothAssetEditor/Source/ChaosClothAssetDataflowNodes/Private/ChaosClothAsset/SelectionToWeightMapNode.cpp
+++ b/Engine/Plugins/ChaosClothAssetEditor/Source/ChaosClothAssetDataflowNodes/Private/ChaosClothAsset/SelectionToWeightMapNode.cpp
@@ -86,38 +86,45 @@ void FChaosClothAssetSelectionToWeightMapNode::Evaluate(UE::Dataflow::FContext&
 							const TConstArrayView<FIntVector3> SimIndices3D = ClothFacade.GetSimIndices3D();
 							for (const int32 FaceIndex : *SelectionSet)
 							{
-								OutClothWeights[SimIndices3D[FaceIndex][0]] = SelectedValue;
-								OutClothWeights[SimIndices3D[FaceIndex][1]] = SelectedValue;
-								OutClothWeights[SimIndices3D[FaceIndex][2]] = SelectedValue;
+								if (SimIndices3D.IsValidIndex(FaceIndex))
+								{
+									// Indices in SimIndices3D are assumed to be valid to index OutClothWeights
+									OutClothWeights[SimIndices3D[FaceIndex][0]] = SelectedValue;
+									OutClothWeights[SimIndices3D[FaceIndex][1]] = SelectedValue;
+									OutClothWeights[SimIndices3D[FaceIndex][2]] = SelectedValue;
+								}
 							}
 						}
 					}
 					else
 					{
 						check(OutClothWeights.Num() == 0);
-						FClothDataflowTools::LogAndToastWarning(*this,
-							LOCTEXT("InvalidWeightMapNameHeadline", "Invalid weight map name."),
+						Context.Warning(
 							FText::Format(LOCTEXT("InvalidWeightMapNameDetails", "Could not create a weight map with name \"{0}\" (reserved name? wrong type?)."),
-								FText::FromName(InMapName)));
+								FText::FromName(InMapName))
+							, this, Out
+						);
 					}
 				}
 				else
 				{
-					FClothDataflowTools::LogAndToastWarning(*this,
-						LOCTEXT("SelectionTypeNotCorrectHeadline", "Invalid selection group."),
+					Context.Warning(
 						FText::Format(LOCTEXT("SelectionTypeNotCorrectDetails", "Selection \"{0}\" does not have its index dependency group set to \"{1}\", \"{2}\", or \"{3}\"."),
 							FText::FromName(InSelectionName),
 							FText::FromName(ClothCollectionGroup::SimVertices2D),
 							FText::FromName(ClothCollectionGroup::SimVertices3D),
-							FText::FromName(ClothCollectionGroup::SimFaces)));
+							FText::FromName(ClothCollectionGroup::SimFaces))
+						, this, Out
+					);
 				}
 			}
 			else
 			{
-				FClothDataflowTools::LogAndToastWarning(*this,
-					LOCTEXT("SelectionNameNotFoundHeadline", "Invalid selection name."),
+				Context.Warning(
 					FText::Format(LOCTEXT("SelectionNameNotFoundDetails", "Selection \"{0}\" was not found in the collection."),
-						FText::FromName(InSelectionName)));
+						FText::FromName(InSelectionName))
+					, this, Out
+				);
 			}
 		}
 
diff --git a/Engine/Plugins/ChaosClothAssetEditor/Source/ChaosClothAssetDataflowNodes/Private/ChaosClothAsset/SimulationDefaultConfigNode.cpp b/Engine/Plugins/ChaosClothAssetEditor/Source/ChaosClothAssetDataflowNodes/Private/ChaosClothAsset/SimulationDefaultConfigNode.cpp
--- a/Engine/Plugins/ChaosClothAssetEditor/Source/ChaosClothAssetDataflowNodes/Private/ChaosClothAsset/SimulationDefaultConfigNode.cpp
+++ b/Engine/Plugins/ChaosClothAssetEditor/Source/ChaosClothAssetDataflowNodes/Private/ChaosClothAsset/SimulationDefaultConfigNode.cpp
@@ -78,9 +78,10 @@ void FChaosClothAssetSimulationDefaultConfigNode::Evaluate(UE::Dataflow::FContex
 			}
 			else
 			{
-				FClothDataflowTools::LogAndToastWarning(*this,
-					LOCTEXT("CannotOverridePreExistingConfigHeadline", "Cannot override pre-existing config."),
-					LOCTEXT("CannotOverridePreExistingConfigDetails", "The Simulation Default Config node cannot be used to override any existing properties, and must be placed ahead of any other config nodes."));
+				Context.Warning(
+					LOCTEXT("CannotOverridePreExistingConfigDetails", "The Simulation Default Config node cannot be used to override any existing properties, and must be placed ahead of any other config nodes.")
+					, this, Out
+				);
 			}
 		}
 
diff --git a/Engine/Plugins/ChaosClothAssetEditor/Source/ChaosClothAssetDataflowNodes/Private/ChaosClothAsset/SkeletalMeshImportNode.cpp b/Engine/Plugins/ChaosClothAssetEditor/Source/ChaosClothAssetDataflowNodes/Private/ChaosClothAsset/SkeletalMeshImportNode.cpp
--- a/Engine/Plugins/ChaosClothAssetEditor/Source/ChaosClothAssetDataflowNodes/Private/ChaosClothAsset/SkeletalMeshImportNode.cpp
+++ b/Engine/Plugins/ChaosClothAssetEditor/Source/ChaosClothAssetDataflowNodes/Private/ChaosClothAsset/SkeletalMeshImportNode.cpp
@@ -59,11 +59,13 @@ void FChaosClothAssetSkeletalMeshImportNode_v2::Evaluate(UE::Dataflow::FContext&
 			const bool bIsValidLOD = ImportedModel && ImportedModel->LODModels.IsValidIndex(LODIndex);
 			if (!bIsValidLOD)
 			{
-				FClothDataflowTools::LogAndToastWarning(*this, LOCTEXT("InvalidLODHeadline", "Invalid LOD."),
+				Context.Warning(
 					FText::Format(
 						LOCTEXT("InvalidLODDetails", "No valid LOD {0} found for skeletal mesh {1}."),
 						LODIndex,
-						FText::FromString(InSkeletalMesh->GetName())));
+						FText::FromString(InSkeletalMesh->GetName()))
+					, this, Out
+				);
 
 				SetValue(Context, MoveTemp(*ClothCollection), &Collection);
 				return;
@@ -81,11 +83,13 @@ void FChaosClothAssetSkeletalMeshImportNode_v2::Evaluate(UE::Dataflow::FContext&
 				const bool bIsValidSection = LODModel.Sections.IsValidIndex(Section);;
 				if (!bIsValidSection)
 				{
-					FClothDataflowTools::LogAndToastWarning(*this, LOCTEXT("InvalidSectionHeadline", "Invalid section."),
+					Context.Warning(
 						FText::Format(
 							LOCTEXT("InvalidSectionDetails", "No valid section {0} found for skeletal mesh {1}."),
 							Section,
-							FText::FromString(InSkeletalMesh->GetName())));
+							FText::FromString(InSkeletalMesh->GetName()))
+						, this, Out
+					);
 
 					continue;
 				}
@@ -120,7 +124,7 @@ void FChaosClothAssetSkeletalMeshImportNode_v2::Evaluate(UE::Dataflow::FContext&
 				TArrayView<int32> SimImportVertexID = ClothFacade.GetSimImportVertexID();
 				for (int32 Vertex2DIndex = 0; Vertex2DIndex < SimImportVertexID.Num(); ++Vertex2DIndex)
 				{
-					if (LODModel.MeshToImportVertexMap.IsValidIndex(Sim2DToSKMIndex[Vertex2DIndex]))
+					if (Sim2DToSKMIndex.IsValidIndex(Vertex2DIndex) && LODModel.MeshToImportVertexMap.IsValidIndex(Sim2DToSKMIndex[Vertex2DIndex]))
 					{
 						SimImportVertexID[Vertex2DIndex] = LODModel.MeshToImportVertexMap[Sim2DToSKMIndex[Vertex2DIndex]];
 					}
@@ -231,11 +235,13 @@ void FChaosClothAssetSkeletalMeshImportNode::Evaluate(UE::Dataflow::FContext& Co
 			const bool bIsValidLOD = ImportedModel && ImportedModel->LODModels.IsValidIndex(LODIndex);
 			if (!bIsValidLOD)
 			{
-				FClothDataflowTools::LogAndToastWarning(*this, LOCTEXT("InvalidLODHeadline", "Invalid LOD."),
+				Context.Warning(
 					FText::Format(
 						LOCTEXT("InvalidLODDetails", "No valid LOD {0} found for skeletal mesh {1}."),
 						LODIndex,
-						FText::FromString(InSkeletalMesh->GetName())));
+						FText::FromString(InSkeletalMesh->GetName()))
+					, this, Out
+				);
 
 				SetValue(Context, MoveTemp(*ClothCollection), &Collection);
 				return;
@@ -251,11 +257,13 @@ void FChaosClothAssetSkeletalMeshImportNode::Evaluate(UE::Dataflow::FContext& Co
 				const bool bIsValidSection = LODModel.Sections.IsValidIndex(Section);;
 				if (!bIsValidSection)
 				{
-					FClothDataflowTools::LogAndToastWarning(*this, LOCTEXT("InvalidSectionHeadline", "Invalid section."),
+					Context.Warning(
 						FText::Format(
 							LOCTEXT("InvalidSectionDetails", "No valid section {0} found for skeletal mesh {1}."),
 							Section,
-							FText::FromString(InSkeletalMesh->GetName())));
+							FText::FromString(InSkeletalMesh->GetName()))
+						, this, Out
+					);
 
 					continue;
 				}
diff --git a/Engine/Plugins/ChaosClothAssetEditor/Source/ChaosClothAssetDataflowNodes/Private/ChaosClothAsset/SkinningBlendNode.cpp b/Engine/Plugins/ChaosClothAssetEditor/Source/ChaosClothAssetDataflowNodes/Private/ChaosClothAsset/SkinningBlendNode.cpp
--- a/Engine/Plugins/ChaosClothAssetEditor/Source/ChaosClothAssetDataflowNodes/Private/ChaosClothAsset/SkinningBlendNode.cpp
+++ b/Engine/Plugins/ChaosClothAssetEditor/Source/ChaosClothAssetDataflowNodes/Private/ChaosClothAsset/SkinningBlendNode.cpp
@@ -77,7 +77,10 @@ namespace UE::Chaos::ClothAsset::Private
 			{
 				for (const int32 VertexIndex : *KinematicSelectionSet)
 				{
-					PointWeightMap[VertexIndex] = UnSelectedValue;
+					if (VertexIndex >= 0 && VertexIndex < PointWeightMap.Num())
+					{
+						PointWeightMap[VertexIndex] = UnSelectedValue;
+					}
 				}
 				return PointWeightMap;
 			}
@@ -86,7 +89,11 @@ namespace UE::Chaos::ClothAsset::Private
 				const TConstArrayView<int32> Vertex2DTo3D = ClothFacade.GetSimVertex3DLookup();
 				for (const int32 VertexIndex : *KinematicSelectionSet)
 				{
-					PointWeightMap[Vertex2DTo3D[VertexIndex]] = UnSelectedValue;
+					if (Vertex2DTo3D.IsValidIndex(VertexIndex))
+					{
+						// Indices in Vertex2DTo3D  are assumed to be valid for accessing PointWeightMap
+						PointWeightMap[Vertex2DTo3D[VertexIndex]] = UnSelectedValue;
+					}
 				}
 				return PointWeightMap;
 			}
@@ -95,9 +102,13 @@ namespace UE::Chaos::ClothAsset::Private
 				const TConstArrayView<FIntVector3> SimIndices3D = ClothFacade.GetSimIndices3D();
 				for (const int32 FaceIndex : *KinematicSelectionSet)
 				{
-					PointWeightMap[SimIndices3D[FaceIndex][0]] = UnSelectedValue;
-					PointWeightMap[SimIndices3D[FaceIndex][1]] = UnSelectedValue;
-					PointWeightMap[SimIndices3D[FaceIndex][2]] = UnSelectedValue;
+					if (SimIndices3D.IsValidIndex(FaceIndex))
+					{
+						// Indices in SimIndices3D are assumed to be valid to index the PointWeightMap array
+						PointWeightMap[SimIndices3D[FaceIndex][0]] = UnSelectedValue;
+						PointWeightMap[SimIndices3D[FaceIndex][1]] = UnSelectedValue;
+						PointWeightMap[SimIndices3D[FaceIndex][2]] = UnSelectedValue;
+					}
 				}
 				return PointWeightMap;
 			}
@@ -133,9 +144,10 @@ void FChaosClothAssetSkinningBlendNode::Evaluate(UE::Dataflow::FContext& Context
 			// Check for the optional render deformer schema
 			if (!ClothFacade.IsValid(EClothCollectionExtendedSchemas::RenderDeformer))
 			{
-				FClothDataflowTools::LogAndToastWarning(*this,
-					LOCTEXT("NoProxyDefromerHeadline", "No Proxy Deformer data."),
-					LOCTEXT("NoProxyDefromerDetails", "There isn't any Proxy Deformer mapping data on the input Cloth Collection to generate the Skinning Blend weight map."));
+				Context.Warning(
+					LOCTEXT("NoProxyDefromerDetails", "There isn't any Proxy Deformer mapping data on the input Cloth Collection to generate the Skinning Blend weight map.")
+					, this, Out
+				);
 			}
 			else
 			{
@@ -143,9 +155,10 @@ void FChaosClothAssetSkinningBlendNode::Evaluate(UE::Dataflow::FContext& Context
 				FName KinematicSelectionName = FName(*GetValue<FString>(Context, &KinematicVertices3D.StringValue));
 				if (KinematicSelectionName != NAME_None && (!SelectionFacade.IsValid() || !SelectionFacade.FindSelectionSet(KinematicSelectionName)))
 				{
-					FClothDataflowTools::LogAndToastWarning(*this,
-						LOCTEXT("HasSimVertexSelectionHeadline", "Unknown KinematicVertices3D selection."),
-						LOCTEXT("HasSimVertexSelectionDetails", "The specified KinematicVertices3D selection does't exist within the input Cloth Collection. An empty selection of kinematic vertices will be used instead."));
+					Context.Warning(
+						LOCTEXT("HasSimVertexSelectionDetails", "The specified KinematicVertices3D selection does't exist within the input Cloth Collection. An empty selection of kinematic vertices will be used instead.")
+						, this, Out
+					);
 					KinematicSelectionName = NAME_None;
 				}
 				// Create the render weight map for storing the skinning blend weights pattern per pattern, as the number of influences could vary
diff --git a/Engine/Plugins/ChaosClothAssetEditor/Source/ChaosClothAssetDataflowNodes/Private/ChaosClothAsset/StaticMeshImportNode.cpp b/Engine/Plugins/ChaosClothAssetEditor/Source/ChaosClothAssetDataflowNodes/Private/ChaosClothAsset/StaticMeshImportNode.cpp
--- a/Engine/Plugins/ChaosClothAssetEditor/Source/ChaosClothAssetDataflowNodes/Private/ChaosClothAsset/StaticMeshImportNode.cpp
+++ b/Engine/Plugins/ChaosClothAssetEditor/Source/ChaosClothAssetDataflowNodes/Private/ChaosClothAsset/StaticMeshImportNode.cpp
@@ -113,11 +113,12 @@ void FChaosClothAssetStaticMeshImportNode_v2::Evaluate(UE::Dataflow::FContext& C
 					{
 						if (!InitializeSimMeshFromSkeletalMeshModel(SkeletalMeshLODModel, ClothCollection, SimMeshSection, UVChannel, UVScale, bImportSimMeshNormals))
 						{
-							FClothDataflowTools::LogAndToastWarning(*this,
-								LOCTEXT("InvalidRenderMeshHeadline", "Invalid render mesh."),
+							Context.Warning(
 								FText::Format(
 									LOCTEXT("InvalidRenderMeshDetails", "The input static mesh {0} failed to convert to a valid render mesh."),
-									FText::FromString(InStaticMesh->GetName())));
+									FText::FromString(InStaticMesh->GetName()))
+								, this, Out
+							);
 						}
 					}
 					if (bImportRenderMesh)
@@ -139,11 +140,12 @@ void FChaosClothAssetStaticMeshImportNode_v2::Evaluate(UE::Dataflow::FContext& C
 
 						if (!InitializeRenderMeshFromSkeletalMeshModel(SkeletalMeshLODModel, ClothCollection, RenderMeshSection, GetMaterialPathName))
 						{
-							FClothDataflowTools::LogAndToastWarning(*this,
-								LOCTEXT("InvalidRenderMeshHeadline", "Invalid render mesh."),
+							Context.Warning(
 								FText::Format(
 									LOCTEXT("InvalidRenderMeshDetails", "The input static mesh {0} failed to convert to a valid render mesh."),
-									FText::FromString(InStaticMesh->GetName())));
+									FText::FromString(InStaticMesh->GetName()))
+								, this, Out
+							);
 						}
 					}
 				}
@@ -155,14 +157,15 @@ void FChaosClothAssetStaticMeshImportNode_v2::Evaluate(UE::Dataflow::FContext& C
 			}
 			else
 			{
-				FClothDataflowTools::LogAndToastWarning(*this,
-					LOCTEXT("InvalidLODIndexHeadline", "Invalid LOD index."),
+				Context.Warning(
 					FText::Format(LOCTEXT("InvalidLODIndexDetails",
 						"{0} is not a valid LOD index for static mesh {1}.\n"
 						"This static mesh has {2} LOD(s)."),
 						FText::FromString(InStaticMesh->GetName()),
 						LODIndex,
-						NumLods));
+						NumLods)
+					, this, Out
+				);
 			}
 		}
 
@@ -195,7 +198,7 @@ void FChaosClothAssetStaticMeshImportNode::Evaluate(UE::Dataflow::FContext& Cont
 		if (InStaticMesh && (bImportSimMesh || bImportRenderMesh))
 		{
 			const int32 NumLods = InStaticMesh->GetNumSourceModels();
-			if (LODIndex < NumLods)
+			if (LODIndex >= 0 && LODIndex < NumLods)
 			{
 				const FMeshDescription* const MeshDescription = InStaticMesh->GetMeshDescription(LODIndex);
 				check(MeshDescription);
@@ -224,11 +227,12 @@ void FChaosClothAssetStaticMeshImportNode::Evaluate(UE::Dataflow::FContext& Cont
 					{
 						if (!InitializeSimMeshFromSkeletalMeshModel(SkeletalMeshLODModel, ClothCollection, SimMeshSection, UVChannel, UVScale, bImportSimMeshNormals))
 						{
-							FClothDataflowTools::LogAndToastWarning(*this,
-								LOCTEXT("InvalidRenderMeshHeadline", "Invalid render mesh."),
+							Context.Warning(
 								FText::Format(
 									LOCTEXT("InvalidRenderMeshDetails", "The input static mesh {0} failed to convert to a valid render mesh."),
-									FText::FromString(InStaticMesh->GetName())));
+									FText::FromString(InStaticMesh->GetName()))
+								, this, Out
+							);
 						}
 					}
 					if (bImportRenderMesh)
@@ -250,11 +254,12 @@ void FChaosClothAssetStaticMeshImportNode::Evaluate(UE::Dataflow::FContext& Cont
 
 						if (!InitializeRenderMeshFromSkeletalMeshModel(SkeletalMeshLODModel, ClothCollection, RenderMeshSection, GetMaterialPathName))
 						{
-							FClothDataflowTools::LogAndToastWarning(*this,
-								LOCTEXT("InvalidRenderMeshHeadline", "Invalid render mesh."),
+							Context.Warning(
 								FText::Format(
 									LOCTEXT("InvalidRenderMeshDetails", "The input static mesh {0} failed to convert to a valid render mesh."),
-									FText::FromString(InStaticMesh->GetName())));
+									FText::FromString(InStaticMesh->GetName()))
+								, this, Out
+							);
 						}
 					}
 				}
@@ -269,14 +274,15 @@ void FChaosClothAssetStaticMeshImportNode::Evaluate(UE::Dataflow::FContext& Cont
 			}
 			else
 			{
-				FClothDataflowTools::LogAndToastWarning(*this,
-					LOCTEXT("InvalidLODIndexHeadline", "Invalid LOD index."),
+				Context.Warning(
 					FText::Format(LOCTEXT("InvalidLODIndexDetails",
 						"{0} is not a valid LOD index for static mesh {1}.\n"
 						"This static mesh has {2} LOD(s)."),
 						FText::FromString(InStaticMesh->GetName()),
 						LODIndex,
-						NumLods));
+						NumLods)
+					, this, Out
+				);
 			}
 		}
 
diff --git a/Engine/Plugins/ChaosClothAssetEditor/Source/ChaosClothAssetDataflowNodes/Private/ChaosClothAsset/TerminalNode.cpp b/Engine/Plugins/ChaosClothAssetEditor/Source/ChaosClothAssetDataflowNodes/Private/ChaosClothAsset/TerminalNode.cpp
--- a/Engine/Plugins/ChaosClothAssetEditor/Source/ChaosClothAssetDataflowNodes/Private/ChaosClothAsset/TerminalNode.cpp
+++ b/Engine/Plugins/ChaosClothAssetEditor/Source/ChaosClothAssetDataflowNodes/Private/ChaosClothAsset/TerminalNode.cpp
@@ -150,9 +150,9 @@ void FChaosClothAssetTerminalNode_v2::SetAssetValue(TObjectPtr<UObject> Asset, U
 		FText VerboseText;
 		ClothAsset->Build(InClothCollections, &LODTransitionDataCache, &ErrorText, &VerboseText);
 
-		if (!ErrorText.IsEmpty())
+		if (!VerboseText.IsEmpty())
 		{
-			FClothDataflowTools::LogAndToastWarning(*this, ErrorText, VerboseText);
+			Context.Warning(VerboseText, this);
 		}
 
 		// Asset must be resaved
@@ -216,22 +216,24 @@ TArray<TSharedRef<const FManagedArrayCollection>> FChaosClothAssetTerminalNode_v
 			ClothFacade.DefineSchema();
 			ClothFacade.Initialize(FCollectionClothConstFacade(CollectionLodValues[LastValidLodIndex]));
 
-			FClothDataflowTools::LogAndToastWarning(*this,
-				LOCTEXT("InvalidInputLodNHeadline", "Invalid input LOD."),
+			Context.Warning(
 				FText::Format(
 					LOCTEXT("InvalidInputLodNDetails",
 						"Invalid or empty input LOD for LOD {0}.\n"
 						"Using the previous valid LOD {1} instead."),
 					LodIndex,
-					LastValidLodIndex));
+					LastValidLodIndex)
+				, this
+			);
 		}
 		else
 		{
-			FClothDataflowTools::LogAndToastWarning(*this,
-				LOCTEXT("InvalidInputLod0Headline", "Invalid input LOD 0."),
+			Context.Warning(
 				LOCTEXT("InvalidInputLod0Details",
 					"Invalid or empty input LOD for LOD 0.\n"
-					"LOD 0 cannot be empty in order to construct a valid Cloth Asset."));
+					"LOD 0 cannot be empty in order to construct a valid Cloth Asset.")
+				, this
+			);
 			break;
 		}
 		CollectionLodValues.Emplace(MoveTemp(CollectionLodValue));
@@ -347,9 +349,9 @@ PRAGMA_ENABLE_DEPRECATION_WARNINGS
 		FText VerboseText;
 		ClothAsset->Build(InClothCollections, &LODTransitionDataCache, &ErrorText, &VerboseText);
 
-		if (!ErrorText.IsEmpty())
+		if (!VerboseText.IsEmpty())
 		{
-			FClothDataflowTools::LogAndToastWarning(*this, ErrorText, VerboseText);
+			Context.Warning(VerboseText, this);
 		}
 
 		// Asset must be resaved
@@ -504,22 +506,24 @@ TArray<TSharedRef<const FManagedArrayCollection>> FChaosClothAssetTerminalNode::
 			ClothFacade.DefineSchema();
 			ClothFacade.Initialize(FCollectionClothConstFacade(CollectionLodValues[LastValidLodIndex]));
 
-			FClothDataflowTools::LogAndToastWarning(*this,
-				LOCTEXT("InvalidInputLodNHeadline", "Invalid input LOD."),
+			Context.Warning(
 				FText::Format(
 					LOCTEXT("InvalidInputLodNDetails",
 						"Invalid or empty input LOD for LOD {0}.\n"
 						"Using the previous valid LOD {1} instead."),
 					LodIndex,
-					LastValidLodIndex));
+					LastValidLodIndex)
+				, this
+			);
 		}
 		else
 		{
-			FClothDataflowTools::LogAndToastWarning(*this,
-				LOCTEXT("InvalidInputLod0Headline", "Invalid input LOD 0."),
+			Context.Warning(
 				LOCTEXT("InvalidInputLod0Details",
 					"Invalid or empty input LOD for LOD 0.\n"
-					"LOD 0 cannot be empty in order to construct a valid Cloth Asset."));
+					"LOD 0 cannot be empty in order to construct a valid Cloth Asset.")
+				, this
+			);
 			break;
 		}
 		CollectionLodValues.Emplace(MoveTemp(CollectionLodValue));
diff --git a/Engine/Plugins/ChaosClothAssetEditor/Source/ChaosClothAssetDataflowNodes/Private/ChaosClothAsset/TransferSkinWeightsNode.cpp b/Engine/Plugins/ChaosClothAssetEditor/Source/ChaosClothAssetDataflowNodes/Private/ChaosClothAsset/TransferSkinWeightsNode.cpp
--- a/Engine/Plugins/ChaosClothAssetEditor/Source/ChaosClothAssetDataflowNodes/Private/ChaosClothAsset/TransferSkinWeightsNode.cpp
+++ b/Engine/Plugins/ChaosClothAssetEditor/Source/ChaosClothAssetDataflowNodes/Private/ChaosClothAsset/TransferSkinWeightsNode.cpp
@@ -863,11 +863,12 @@ void FChaosClothAssetTransferSkinWeightsNode::Evaluate(UE::Dataflow::FContext& C
 
 				if (!InSkeletalMesh)
 				{
-					FClothDataflowTools::LogAndToastWarning(*this,
-						LOCTEXT("InvalidSkeletalMeshHeadline", "Invalid Skeletal Mesh."),
+					Context.Warning(
 						FText::Format(
 							LOCTEXT("InvalidSkeletalMeshDetails", "No skeletal mesh has been specified and one is required for this type of transfer."),
-							LodIndex));
+							LodIndex)
+						, this, Out
+					);
 					SetValue(Context, MoveTemp(*ClothCollection), &Collection);
 					return;
 				}
@@ -878,24 +879,25 @@ void FChaosClothAssetTransferSkinWeightsNode::Evaluate(UE::Dataflow::FContext& C
 					const int32 LastLODIndex = InSkeletalMesh->GetLODNum() - 1;
 					if (InSkeletalMesh->IsValidLODIndex(LastLODIndex))
 					{
-						FClothDataflowTools::LogAndToastWarning(*this,
-							LOCTEXT("InvalidLodIndexHeadline", "Invalid LOD Index."),
+						Context.Warning(
 							FText::Format(
 								LOCTEXT("InvalidLodIndexSubstituteDetails", "LOD index {0} is not a valid LOD for skeletal mesh {1}. Using LOD index {2} instead."),
 								InLodIndex,
 								FText::FromString(InSkeletalMesh->GetName()),
-								LastLODIndex));
+								LastLODIndex)
+						, this, Out);
 
 						InLodIndex = LastLODIndex;
 					}
 					else
 					{
-						FClothDataflowTools::LogAndToastWarning(*this,
-							LOCTEXT("InvalidLodIndexHeadline", "Invalid LOD Index."),
+						Context.Warning(
 							FText::Format(
 								LOCTEXT("InvalidLodIndexDetails", "LOD index {0} is not a valid LOD for skeletal mesh {1}."),
 								InLodIndex,
-								FText::FromString(InSkeletalMesh->GetName())));
+								FText::FromString(InSkeletalMesh->GetName()))
+							, this, Out
+						);
 						SetValue(Context, MoveTemp(*ClothCollection), &Collection);
 						return;
 					}
@@ -903,12 +905,13 @@ void FChaosClothAssetTransferSkinWeightsNode::Evaluate(UE::Dataflow::FContext& C
 
 				if (!SkeletalMeshToDynamicMesh(InSkeletalMesh.Get(), InLodIndex, SourceDynamicMesh))
 				{
-					FClothDataflowTools::LogAndToastWarning(*this,
-						LOCTEXT("InvalidLodHeadline", "Could not convert LOD to Dynamic Mesh."),
+					Context.Warning(
 						FText::Format(
 							LOCTEXT("InvalidLodDetails", "Could not convert LOD index {0} of the skeletal mesh {1} to dyanamic mesh."),
 							InLodIndex,
-							FText::FromString(InSkeletalMesh->GetName())));
+							FText::FromString(InSkeletalMesh->GetName()))
+						, this, Out
+					);
 					SetValue(Context, MoveTemp(*ClothCollection), &Collection);
 					return;
 				}
@@ -937,9 +940,7 @@ void FChaosClothAssetTransferSkinWeightsNode::Evaluate(UE::Dataflow::FContext& C
 					}
 					if (!TargetRefSkeleton)
 					{
-						FClothDataflowTools::LogAndToastWarning(*this,
-							LOCTEXT("InvalidSimRefSkeletonHeadline", "Invalid Reference Skeleton."),
-							LOCTEXT("InvalidSimRefSkeletonDetails", "Couldn't find a valid reference skeleton from the input sim collection."));
+						Context.Warning(LOCTEXT("InvalidSimRefSkeletonDetails", "Couldn't find a valid reference skeleton from the input sim collection."), this, Out);
 						SetValue(Context, MoveTemp(*ClothCollection), &Collection);
 						return;
 					}
@@ -950,9 +951,7 @@ void FChaosClothAssetTransferSkinWeightsNode::Evaluate(UE::Dataflow::FContext& C
 
 					if (!TargetRefSkeleton)
 					{
-						FClothDataflowTools::LogAndToastWarning(*this,
-							LOCTEXT("InvalidRefSkeletonHeadline", "Invalid Reference Skeleton."),
-							LOCTEXT("InvalidRefSkeletonDetails", "Couldn't find a valid reference skeleton from the input collection."));
+						Context.Warning(LOCTEXT("InvalidRefSkeletonDetails", "Couldn't find a valid reference skeleton from the input collection."), this, Out);
 						SetValue(Context, MoveTemp(*ClothCollection), &Collection);
 						return;
 					}
@@ -1004,9 +1003,7 @@ void FChaosClothAssetTransferSkinWeightsNode::Evaluate(UE::Dataflow::FContext& C
 
 			if (!bTransferResult)
 			{
-				FClothDataflowTools::LogAndToastWarning(*this,
-					LOCTEXT("TransferWeightsFailedHeadline", "Transfer Weights Failed."),
-					LOCTEXT("TransferWeightsDetails", "Failed to transfer skinning weights from the source."));
+				Context.Warning(LOCTEXT("TransferWeightsDetails", "Failed to transfer skinning weights from the source."), this, Out);
 				SetValue(Context, MoveTemp(*ClothCollection), &Collection);
 				return;
 			}
diff --git a/Engine/Plugins/ChaosClothAssetEditor/Source/ChaosClothAssetDataflowNodes/Private/ChaosClothAsset/USDImportNode.cpp b/Engine/Plugins/ChaosClothAssetEditor/Source/ChaosClothAssetDataflowNodes/Private/ChaosClothAsset/USDImportNode.cpp
--- a/Engine/Plugins/ChaosClothAssetEditor/Source/ChaosClothAssetDataflowNodes/Private/ChaosClothAsset/USDImportNode.cpp
+++ b/Engine/Plugins/ChaosClothAssetEditor/Source/ChaosClothAssetDataflowNodes/Private/ChaosClothAsset/USDImportNode.cpp
@@ -114,9 +114,10 @@ PRAGMA_ENABLE_DEPRECATION_WARNINGS
 
 					if (!ImportFromFile(FilePath, AssetPath, bImportSimMesh, ClothCollection, MutableThis->PackagePath, ErrorText))
 					{
-						FClothDataflowTools::LogAndToastWarning(*this,
-							LOCTEXT("FailedToImportUsdFileHeadline", "Failed to import USD file from file."),
-							FText::Format(LOCTEXT("FailedToImportUsdDetails", "Error while importing USD cloth from file '{0}':\n{1}"), FText::FromString(FilePath), ErrorText));
+						Context.Error(
+							FText::Format(LOCTEXT("FailedToImportUsdDetails", "Error while importing USD cloth from file '{0}':\n{1}"), FText::FromString(FilePath), ErrorText)
+							, this
+						);
 					}
 					MutableThis->CollectionCache = MoveTemp(*ClothCollection);
 
@@ -141,9 +142,10 @@ void FChaosClothAssetUSDImportNode::Evaluate(UE::Dataflow::FContext& Context, co
 		if (!ImportFromCache(ClothCollection, ErrorText))
 		{
 PRAGMA_DISABLE_DEPRECATION_WARNINGS
-			FClothDataflowTools::LogAndToastWarning(*this,
-				LOCTEXT("FailedToExportUsdCacheHeadline", "Failed to import USD file from cache."),
-				FText::Format(LOCTEXT("FailedToExportUsdCacheDetails", "Error while importing USD cloth from cache '{0}':\n{1}"), FText::FromString(UsdFile.FilePath), ErrorText));
+			Context.Error(
+				FText::Format(LOCTEXT("FailedToExportUsdCacheDetails", "Error while importing USD cloth from cache '{0}':\n{1}"), FText::FromString(UsdFile.FilePath), ErrorText)
+				, this, Out
+			);
 PRAGMA_ENABLE_DEPRECATION_WARNINGS
 		}
 
diff --git a/Engine/Plugins/ChaosClothAssetEditor/Source/ChaosClothAssetDataflowNodes/Private/ChaosClothAsset/USDImportNode_v2.cpp b/Engine/Plugins/ChaosClothAssetEditor/Source/ChaosClothAssetDataflowNodes/Private/ChaosClothAsset/USDImportNode_v2.cpp
--- a/Engine/Plugins/ChaosClothAssetEditor/Source/ChaosClothAssetDataflowNodes/Private/ChaosClothAsset/USDImportNode_v2.cpp
+++ b/Engine/Plugins/ChaosClothAssetEditor/Source/ChaosClothAssetDataflowNodes/Private/ChaosClothAsset/USDImportNode_v2.cpp
@@ -969,7 +969,7 @@ bool FChaosClothAssetUsdClothData::Serialize(FArchive& Ar)
 FChaosClothAssetUSDImportNode_v2::FChaosClothAssetUSDImportNode_v2(const UE::Dataflow::FNodeParameters& InParam, FGuid InGuid)
 	: FDataflowNode(InParam, InGuid)
 	, UsdFile(
-		FDataflowFunctionProperty::FDelegate::CreateLambda([this, OwningObject = InParam.OwningObject](UE::Dataflow::FContext& /*Context*/)
+		FDataflowFunctionProperty::FDelegate::CreateLambda([this, OwningObject = InParam.OwningObject](UE::Dataflow::FContext& Context)
 			{
 				const FString AssetPath = OwningObject ? OwningObject->GetPackage()->GetPathName() : FString();
 				FText ErrorText;
@@ -979,9 +979,10 @@ FChaosClothAssetUSDImportNode_v2::FChaosClothAssetUSDImportNode_v2(const UE::Dat
 
 					if (!ImportUsdFile_Schemaless(UsdFile.FilePath, AssetPath, ErrorText))
 					{
-						UE::Chaos::ClothAsset::FClothDataflowTools::LogAndToastWarning(*this,
-							LOCTEXT("FailedToImportUsdFileHeadline", "Failed to import USD file from file."),
-							FText::Format(LOCTEXT("FailedToImportUsdDetails", "Error while importing USD cloth from file '{0}':\n{1}"), FText::FromString(UsdFile.FilePath), ErrorText));
+						Context.Error(
+							FText::Format(LOCTEXT("FailedToImportUsdDetails", "Error while importing USD cloth from file '{0}':\n{1}"), FText::FromString(UsdFile.FilePath), ErrorText)
+							, this
+						);
 					}
 				}
 			}))
@@ -991,28 +992,30 @@ FChaosClothAssetUSDImportNode_v2::FChaosClothAssetUSDImportNode_v2(const UE::Dat
 				UsdFile.Execute(Context);
 			}))
 	, ReloadSimStaticMesh(
-		FDataflowFunctionProperty::FDelegate::CreateLambda([this](UE::Dataflow::FContext& /*Context*/)
+		FDataflowFunctionProperty::FDelegate::CreateLambda([this](UE::Dataflow::FContext& Context)
 			{
 				const TSharedRef<FManagedArrayCollection> ClothCollection = MakeShared<FManagedArrayCollection>(MoveTemp(Collection));
 				FText ErrorText;
 				if (!ImportSimStaticMesh(ClothCollection, ErrorText))
 				{
-					UE::Chaos::ClothAsset::FClothDataflowTools::LogAndToastWarning(*this,
-						LOCTEXT("FailedToImportSimMeshHeadline", "Failed to reload the simulation static mesh."),
-						FText::Format(LOCTEXT("FailedToImportSimMeshDetails", "Error while re-importing the simulation mesh from static mesh '{0}':\n{1}"), FText::FromString(ImportedSimStaticMesh->GetName()), ErrorText));
+					Context.Error(
+						FText::Format(LOCTEXT("FailedToImportSimMeshDetails", "Error while re-importing the simulation mesh from static mesh '{0}':\n{1}"), FText::FromString(ImportedSimStaticMesh->GetName()), ErrorText)
+						, this
+					);
 				}
 				Collection = MoveTemp(*ClothCollection);
 			}))
 	, ReloadRenderStaticMesh(
-		FDataflowFunctionProperty::FDelegate::CreateLambda([this](UE::Dataflow::FContext& /*Context*/)
+		FDataflowFunctionProperty::FDelegate::CreateLambda([this](UE::Dataflow::FContext& Context)
 			{
 				const TSharedRef<FManagedArrayCollection> ClothCollection = MakeShared<FManagedArrayCollection>(MoveTemp(Collection));
 				FText ErrorText;
 				if (!ImportRenderStaticMesh(ClothCollection, ErrorText))
 				{
-					UE::Chaos::ClothAsset::FClothDataflowTools::LogAndToastWarning(*this,
-						LOCTEXT("FailedToImportRenderMeshHeadline", "Failed to reload the render static mesh."),
-						FText::Format(LOCTEXT("FailedToImportRenderMeshDetails", "Error while re-importing the render mesh from static mesh '{0}':\n{1}"), FText::FromString(ImportedRenderStaticMesh->GetName()), ErrorText));
+					Context.Error(
+						FText::Format(LOCTEXT("FailedToImportRenderMeshDetails", "Error while re-importing the render mesh from static mesh '{0}':\n{1}"), FText::FromString(ImportedRenderStaticMesh->GetName()), ErrorText)
+						, this
+					);
 				}
 				Collection = MoveTemp(*ClothCollection);
 			}))
diff --git a/Engine/Plugins/ChaosClothAssetEditor/Source/ChaosClothAssetDataflowNodes/Private/ChaosClothAsset/WeightMapNode.cpp b/Engine/Plugins/ChaosClothAssetEditor/Source/ChaosClothAssetDataflowNodes/Private/ChaosClothAsset/WeightMapNode.cpp
--- a/Engine/Plugins/ChaosClothAssetEditor/Source/ChaosClothAssetDataflowNodes/Private/ChaosClothAsset/WeightMapNode.cpp
+++ b/Engine/Plugins/ChaosClothAssetEditor/Source/ChaosClothAssetDataflowNodes/Private/ChaosClothAsset/WeightMapNode.cpp
@@ -135,16 +135,17 @@ void FChaosClothAssetWeightMapNode::Evaluate(UE::Dataflow::FContext& Context, co
 {
 	using namespace UE::Chaos::ClothAsset;
 
-	auto CheckSourceVertexWeights = [this](TArrayView<float>& ClothWeights, const TArray<float>& SourceVertexWeights, bool bIsSim)
+	auto CheckSourceVertexWeights = [this, &Context, Out](TArrayView<float>& ClothWeights, const TArray<float>& SourceVertexWeights, bool bIsSim)
 	{
 		if (SourceVertexWeights.Num() > 0 && SourceVertexWeights.Num() != ClothWeights.Num())
 		{
-			FClothDataflowTools::LogAndToastWarning(*this,
-				LOCTEXT("VertexCountMismatchHeadline", "Vertex count mismatch."),
+			Context.Warning(
 				FText::Format(LOCTEXT("VertexCountMismatchDetails", "{0} vertex weights in the node: {1}\n{0} vertices in the cloth: {2}"),
 					bIsSim ? FText::FromString("Sim") : FText::FromString("Render"),
 					SourceVertexWeights.Num(),
-					ClothWeights.Num()));
+					ClothWeights.Num())
+				, this, Out
+			);
 		}
 	};
 
@@ -171,10 +172,11 @@ void FChaosClothAssetWeightMapNode::Evaluate(UE::Dataflow::FContext& Context, co
 				if (ClothSimWeights.Num() != ClothFacade.GetNumSimVertices3D())
 				{
 					check(ClothSimWeights.Num() == 0);
-					FClothDataflowTools::LogAndToastWarning(*this,
-						LOCTEXT("InvalidSimWeightMapNameHeadline", "Invalid weight map name."),
+					Context.Warning(
 						FText::Format(LOCTEXT("InvalidSimWeightMapNameDetails", "Could not create a sim weight map with name \"{0}\" (reserved name? wrong type?)."),
-							FText::FromName(InName)));
+							FText::FromName(InName))
+						, this, Out
+					);
 				}
 				else
 				{
@@ -193,10 +195,11 @@ void FChaosClothAssetWeightMapNode::Evaluate(UE::Dataflow::FContext& Context, co
 				if (ClothRenderWeights.Num() != ClothFacade.GetNumRenderVertices())
 				{
 					check(ClothRenderWeights.Num() == 0);
-					FClothDataflowTools::LogAndToastWarning(*this,
-						LOCTEXT("InvalidRenderWeightMapNameHeadline", "Invalid weight map name."),
+					Context.Warning(
 						FText::Format(LOCTEXT("InvalidRenderWeightMapNameDetails", "Could not create a render weight map with name \"{0}\" (reserved name? wrong type?)."),
-							FText::FromName(InName)));
+							FText::FromName(InName))
+						, this, Out
+					);
 				}
 				else
 				{
diff --git a/Engine/Plugins/ChaosClothAssetEditor/Source/ChaosClothAssetDataflowNodes/Private/ChaosClothAsset/WeightMapToSelectionNode.cpp b/Engine/Plugins/ChaosClothAssetEditor/Source/ChaosClothAssetDataflowNodes/Private/ChaosClothAsset/WeightMapToSelectionNode.cpp
--- a/Engine/Plugins/ChaosClothAssetEditor/Source/ChaosClothAssetDataflowNodes/Private/ChaosClothAsset/WeightMapToSelectionNode.cpp
+++ b/Engine/Plugins/ChaosClothAssetEditor/Source/ChaosClothAssetDataflowNodes/Private/ChaosClothAsset/WeightMapToSelectionNode.cpp
@@ -105,10 +105,11 @@ void FChaosClothAssetWeightMapToSelectionNode::Evaluate(UE::Dataflow::FContext&
 			}
 			else
 			{
-				FClothDataflowTools::LogAndToastWarning(*this,
-					LOCTEXT("WeightMapNameNotFoundHeadline", "Invalid weight map name."),
+				Context.Warning(
 					FText::Format(LOCTEXT("WeightMapeNotFoundDetails", "Weight map \"{0}\" was not found in the collection."),
-						FText::FromName(InWeightMapName)));
+						FText::FromName(InWeightMapName))
+					, this, Out
+				);
 
 			}
 		}
