Chaos Cloth Asset - Fixed icon paths broken after the ChaosClothAssetTools module got moved to the ClothAsset plugin.

#rnx
#tests Checked that all editor ClothAsset icons were no longer broken.

#Synced-CL 49759196

diff --git a/Engine/Plugins/ChaosClothAsset/Source/ChaosClothAssetTools/Private/ChaosClothAsset/ClothAssetEditorStyle.cpp b/Engine/Plugins/ChaosClothAsset/Source/ChaosClothAssetTools/Private/ChaosClothAsset/ClothAssetEditorStyle.cpp
--- a/Engine/Plugins/ChaosClothAsset/Source/ChaosClothAssetTools/Private/ChaosClothAsset/ClothAssetEditorStyle.cpp
+++ b/Engine/Plugins/ChaosClothAsset/Source/ChaosClothAssetTools/Private/ChaosClothAsset/ClothAssetEditorStyle.cpp
@@ -11,8 +11,7 @@ namespace UE::Chaos::ClothAsset
 	FClothAssetEditorStyle::FClothAssetEditorStyle()
 		: FSlateStyleSet("ClothAssetEditorStyle")
 	{
-		TSharedPtr<IPlugin> ChaosClothAssetPlugin = IPluginManager::Get().FindPlugin("ChaosClothAsset");
-		if (ChaosClothAssetPlugin.IsValid())
+		if (const TSharedPtr<IPlugin> ChaosClothAssetPlugin = IPluginManager::Get().GetModuleOwnerPlugin(UE_MODULE_NAME))
 		{
 			SetContentRoot(ChaosClothAssetPlugin->GetBaseDir() / TEXT("Resources"));
 
@@ -22,7 +21,6 @@ namespace UE::Chaos::ClothAsset
 			Set("ClassIcon.ChaosClothPreset", new FSlateVectorImageBrush(RootToContentDir(TEXT("ClothPreset_16.svg")), FVector2D(16)));
 			Set("ClassThumbnail.ChaosClothPreset", new FSlateVectorImageBrush(RootToContentDir(TEXT("ClothPreset_64.svg")), FVector2D(64)));
 		}
-
 		FSlateStyleRegistry::RegisterSlateStyle(*this);
 	}
 
diff --git a/Engine/Plugins/ChaosClothAsset/Source/ChaosClothAssetTools/Private/ChaosClothAsset/ClothComponentEditorStyle.cpp b/Engine/Plugins/ChaosClothAsset/Source/ChaosClothAssetTools/Private/ChaosClothAsset/ClothComponentEditorStyle.cpp
--- a/Engine/Plugins/ChaosClothAsset/Source/ChaosClothAssetTools/Private/ChaosClothAsset/ClothComponentEditorStyle.cpp
+++ b/Engine/Plugins/ChaosClothAsset/Source/ChaosClothAssetTools/Private/ChaosClothAsset/ClothComponentEditorStyle.cpp
@@ -11,16 +11,15 @@ namespace UE::Chaos::ClothAsset
 	FClothComponentEditorStyle::FClothComponentEditorStyle()
 		: FSlateStyleSet("ClothComponentEditorStyle")
 	{
-		TSharedPtr<IPlugin> ChaosClothAssetPlugin = IPluginManager::Get().FindPlugin("ChaosClothAssetEditor");
-		if (ChaosClothAssetPlugin.IsValid())
+		if (const TSharedPtr<IPlugin> ChaosClothAssetPlugin = IPluginManager::Get().GetModuleOwnerPlugin(UE_MODULE_NAME))
 		{
 			SetContentRoot(ChaosClothAssetPlugin->GetBaseDir() / TEXT("Resources"));
 
 			Set("ClassIcon.ChaosClothComponent", new FSlateVectorImageBrush(RootToContentDir(TEXT("ClothAsset_16.svg")), FVector2D(16)));
 			Set("ClassThumbnail.ChaosClothComponent", new FSlateVectorImageBrush(RootToContentDir(TEXT("ClothAsset_64.svg")), FVector2D(64)));
 
-			FSlateStyleRegistry::RegisterSlateStyle(*this);
 		}
+		FSlateStyleRegistry::RegisterSlateStyle(*this);
 	}
 
 	FClothComponentEditorStyle::~FClothComponentEditorStyle()
diff --git a/Engine/Plugins/ChaosClothAssetEditor/Source/ChaosClothAssetEditor/Private/ChaosClothAsset/ClothEditorStyle.cpp b/Engine/Plugins/ChaosClothAssetEditor/Source/ChaosClothAssetEditor/Private/ChaosClothAsset/ClothEditorStyle.cpp
--- a/Engine/Plugins/ChaosClothAssetEditor/Source/ChaosClothAssetEditor/Private/ChaosClothAsset/ClothEditorStyle.cpp
+++ b/Engine/Plugins/ChaosClothAssetEditor/Source/ChaosClothAssetEditor/Private/ChaosClothAsset/ClothEditorStyle.cpp
@@ -12,65 +12,59 @@ namespace UE::Chaos::ClothAsset
 {
 const FName FChaosClothAssetEditorStyle::StyleName("ClothStyle");
 
-FString FChaosClothAssetEditorStyle::InContent(const FString& RelativePath, const ANSICHAR* Extension)
-{
-	static FString ContentDir = IPluginManager::Get().FindPlugin(TEXT("ChaosClothAssetEditor"))->GetContentDir();
-	FString Path = (ContentDir / RelativePath) + Extension;
-	return Path;
-}
-
 FChaosClothAssetEditorStyle::FChaosClothAssetEditorStyle()
 	: FSlateStyleSet(StyleName)
 {
 	SetCoreContentRoot(FPaths::EngineContentDir() / TEXT("Editor/Slate"));
 
-	FString ContentDir = IPluginManager::Get().FindPlugin(TEXT("ChaosClothAssetEditor"))->GetContentDir();
-	SetContentRoot(ContentDir);
+	if (const TSharedPtr<IPlugin> ChaosClothAssetEditorPlugin = IPluginManager::Get().GetModuleOwnerPlugin(UE_MODULE_NAME))
+	{
+		SetContentRoot(ChaosClothAssetEditorPlugin->GetContentDir());
 
-	// Some standard icon sizes used elsewhere in the editor
-	const FVector2D Icon8x8(8.0f, 8.0f);
-	const FVector2D Icon16x16(16.0f, 16.0f);
-	const FVector2D Icon20x20(20.0f, 20.0f);
-	const FVector2D Icon28x28(28.0f, 28.0f);
-	const FVector2D Icon40x40(40.0f, 40.0f);
-	const FVector2D Icon120(120.0f, 120.0f);
+		// Some standard icon sizes used elsewhere in the editor
+		const FVector2D Icon8x8(8.0f, 8.0f);
+		const FVector2D Icon16x16(16.0f, 16.0f);
+		const FVector2D Icon20x20(20.0f, 20.0f);
+		const FVector2D Icon28x28(28.0f, 28.0f);
+		const FVector2D Icon40x40(40.0f, 40.0f);
+		const FVector2D Icon120(120.0f, 120.0f);
 
-	// Icon sizes used in this style set
-	const FVector2D ViewportToolbarIconSize = Icon16x16;
-	const FVector2D ToolbarIconSize = Icon20x20;
+		// Icon sizes used in this style set
+		const FVector2D ViewportToolbarIconSize = Icon16x16;
+		const FVector2D ToolbarIconSize = Icon20x20;
 
-	FString PropertyNameString = "ChaosClothAssetEditor." + FChaosClothAssetEditorCommands::AddWeightMapNodeIdentifier;
-	Set(*PropertyNameString, new IMAGE_BRUSH_SVG("Icons/PaintMaps", ToolbarIconSize));
+		FString PropertyNameString = "ChaosClothAssetEditor." + FChaosClothAssetEditorCommands::AddWeightMapNodeIdentifier;
+		Set(*PropertyNameString, new IMAGE_BRUSH_SVG("Icons/PaintMaps", ToolbarIconSize));
 
-	PropertyNameString = "ChaosClothAssetEditor." + FChaosClothAssetEditorCommands::AddMeshSelectionNodeIdentifier;
-	Set(*PropertyNameString, new IMAGE_BRUSH_SVG("Icons/MeshSelect", ToolbarIconSize));
+		PropertyNameString = "ChaosClothAssetEditor." + FChaosClothAssetEditorCommands::AddMeshSelectionNodeIdentifier;
+		Set(*PropertyNameString, new IMAGE_BRUSH_SVG("Icons/MeshSelect", ToolbarIconSize));
 
-	PropertyNameString = "ChaosClothAssetEditor." + FChaosClothAssetEditorCommands::AddTransferSkinWeightsNodeIdentifier;
-	Set(*PropertyNameString, new IMAGE_BRUSH_SVG("Icons/SkinTransfer", ToolbarIconSize));
+		PropertyNameString = "ChaosClothAssetEditor." + FChaosClothAssetEditorCommands::AddTransferSkinWeightsNodeIdentifier;
+		Set(*PropertyNameString, new IMAGE_BRUSH_SVG("Icons/SkinTransfer", ToolbarIconSize));
 
-	//
-	// Construction Viewport
-	// 
+		//
+		// Construction Viewport
+		// 
 	
-	PropertyNameString = "ChaosClothAssetEditor." + FChaosClothAssetEditorCommands::ToggleConstructionViewWireframeIdentifier;
-	Set(*PropertyNameString, new IMAGE_BRUSH_SVG("Icons/ClothWireframe_16", ViewportToolbarIconSize));
-
-	//
-	// Preview Viewport
-	// 
+		PropertyNameString = "ChaosClothAssetEditor." + FChaosClothAssetEditorCommands::ToggleConstructionViewWireframeIdentifier;
+		Set(*PropertyNameString, new IMAGE_BRUSH_SVG("Icons/ClothWireframe_16", ViewportToolbarIconSize));
 
-	PropertyNameString = "ChaosClothAssetEditor." + FChaosClothAssetEditorCommands::TogglePreviewWireframeIdentifier;
-	Set(*PropertyNameString, new IMAGE_BRUSH_SVG("Icons/ClothWireframe_16", ViewportToolbarIconSize));
+		//
+		// Preview Viewport
+		// 
 
-	PropertyNameString = "ChaosClothAssetEditor." + FChaosClothAssetEditorCommands::ToggleSimulationSuspendedIdentifier;
-	Set(*PropertyNameString, new IMAGE_BRUSH_SVG("Icons/ClothSimSuspend_16", ViewportToolbarIconSize));
+		PropertyNameString = "ChaosClothAssetEditor." + FChaosClothAssetEditorCommands::TogglePreviewWireframeIdentifier;
+		Set(*PropertyNameString, new IMAGE_BRUSH_SVG("Icons/ClothWireframe_16", ViewportToolbarIconSize));
 
-	PropertyNameString = "ChaosClothAssetEditor." + FChaosClothAssetEditorCommands::SoftResetSimulationIdentifier;
-	Set(*PropertyNameString, new IMAGE_BRUSH_SVG("Icons/ResetSoft_16", ViewportToolbarIconSize));
+		PropertyNameString = "ChaosClothAssetEditor." + FChaosClothAssetEditorCommands::ToggleSimulationSuspendedIdentifier;
+		Set(*PropertyNameString, new IMAGE_BRUSH_SVG("Icons/ClothSimSuspend_16", ViewportToolbarIconSize));
 
-	PropertyNameString = "ChaosClothAssetEditor." + FChaosClothAssetEditorCommands::HardResetSimulationIdentifier;
-	Set(*PropertyNameString, new IMAGE_BRUSH_SVG("Icons/ResetHard_16", ViewportToolbarIconSize));
+		PropertyNameString = "ChaosClothAssetEditor." + FChaosClothAssetEditorCommands::SoftResetSimulationIdentifier;
+		Set(*PropertyNameString, new IMAGE_BRUSH_SVG("Icons/ResetSoft_16", ViewportToolbarIconSize));
 
+		PropertyNameString = "ChaosClothAssetEditor." + FChaosClothAssetEditorCommands::HardResetSimulationIdentifier;
+		Set(*PropertyNameString, new IMAGE_BRUSH_SVG("Icons/ResetHard_16", ViewportToolbarIconSize));
+	}
 	FSlateStyleRegistry::RegisterSlateStyle(*this);
 }
 
diff --git a/Engine/Plugins/ChaosClothAssetEditor/Source/ChaosClothAssetEditor/Private/ChaosClothAsset/ClothEditorStyle.h b/Engine/Plugins/ChaosClothAssetEditor/Source/ChaosClothAssetEditor/Private/ChaosClothAsset/ClothEditorStyle.h
--- a/Engine/Plugins/ChaosClothAssetEditor/Source/ChaosClothAssetEditor/Private/ChaosClothAsset/ClothEditorStyle.h
+++ b/Engine/Plugins/ChaosClothAssetEditor/Source/ChaosClothAssetEditor/Private/ChaosClothAsset/ClothEditorStyle.h
@@ -22,8 +22,6 @@ public:
 	static UE_API FChaosClothAssetEditorStyle& Get();
 
 private:
-	static UE_API FString InContent(const FString& RelativePath, const ANSICHAR* Extension);
-
 	UE_API FChaosClothAssetEditorStyle();
 	UE_API ~FChaosClothAssetEditorStyle();
 };
