Chaos Cloth Asset:
- Fixed custom tether generation.
- Improved custom tether tooltip.
- Improved Bending constraint RestAngleType tooltip.

[RN] Chaos Cloth Asset - Fixed custom tether generation.

#Synced-CL 50102962

diff --git a/Engine/Plugins/ChaosClothAsset/Source/ChaosClothAssetEngine/Private/ChaosClothAsset/ClothEngineTools.cpp b/Engine/Plugins/ChaosClothAsset/Source/ChaosClothAssetEngine/Private/ChaosClothAsset/ClothEngineTools.cpp
--- a/Engine/Plugins/ChaosClothAsset/Source/ChaosClothAssetEngine/Private/ChaosClothAsset/ClothEngineTools.cpp
+++ b/Engine/Plugins/ChaosClothAsset/Source/ChaosClothAssetEngine/Private/ChaosClothAsset/ClothEngineTools.cpp
@@ -164,7 +164,7 @@ void FClothEngineTools::GenerateTethersFromCustomSelectionSets(const TSharedRef<
 				SelectionFacade.HasSelection(CustomDynamicEndSet) && SelectionFacade.GetSelectionGroup(CustomDynamicEndSet) == UE::Chaos::ClothAsset::ClothCollectionGroup::SimVertices3D)
 			{
 				FClothTetherData TetherData;
-				TetherData.GenerateTethers(ClothFacade.GetSimPosition3D(), TConstArrayView<uint32>(SimIndices), FixedEndSet, SelectionFacade.GetSelectionSet(CustomDynamicEndSet), SelectionFacade.GetSelectionSet(CustomFixedEndSet), bGeodesicTethers);
+				TetherData.GenerateTethers(ClothFacade.GetSimPosition3D(), TConstArrayView<uint32>(SimIndices), FixedEndSet, SelectionFacade.GetSelectionSet(CustomFixedEndSet), SelectionFacade.GetSelectionSet(CustomDynamicEndSet), bGeodesicTethers);
 				Private::AppendTetherData(ClothFacade, TetherData);
 			}
 		}
diff --git a/Engine/Plugins/ChaosClothAssetEditor/Source/ChaosClothAssetDataflowNodes/Public/ChaosClothAsset/SimulationBendingConfigNode.h b/Engine/Plugins/ChaosClothAssetEditor/Source/ChaosClothAssetDataflowNodes/Public/ChaosClothAsset/SimulationBendingConfigNode.h
--- a/Engine/Plugins/ChaosClothAssetEditor/Source/ChaosClothAssetDataflowNodes/Public/ChaosClothAsset/SimulationBendingConfigNode.h
+++ b/Engine/Plugins/ChaosClothAssetEditor/Source/ChaosClothAssetDataflowNodes/Public/ChaosClothAsset/SimulationBendingConfigNode.h
@@ -19,6 +19,20 @@ public:
 
 	/**
 	 * Method for calculating the rest angles of the constraints.
+	 * Use3DRestAngles: Use the 3D rest mesh angle. This will cause any creases and folds to be baked in the simulation.
+	 * FlatnessRatio: Calculate rest angles as a ratio between completely flat and whatever is the 3D rest angle.
+	 *   When FlatnessRatio = 0, this is equivalent to "Use3DRestAngles".
+	 *   When FlatnessRatio = 1, the rest angle will be 0 (completely flat).
+	 *   If a valid weight map is found with the given Weight Map name, then both Low and High values
+	 *   are interpolated with the per particle weight to make the final value used for the simulation.
+	 *   Otherwise all particles are considered to have a zero weight, and only the Low value is meaningful.
+	 * RestAngle: 
+	 *   Set rest angles to be the explicit value set here (in degrees).
+	 *   0 = Flat, Positive values fold away from the edge normal, Negative values fold toward the edge normal.
+	 *   When converting vertex weight values to edge values, the value with the smallest absolute value is selected.
+	 *   If a valid weight map is found with the given Weight Map name, then both Low and High values
+	 *   are interpolated with the per particle weight to make the final value used for the simulation.
+	 *   Otherwise all particles are considered to have a zero weight, and only the Low value is meaningful.
 	 */
 	UPROPERTY(EditAnywhere, Category = "Bending Types")
 	EChaosClothAssetRestAngleConstructionType RestAngleType = EChaosClothAssetRestAngleConstructionType::Use3DRestAngles;
diff --git a/Engine/Plugins/ChaosClothAssetEditor/Source/ChaosClothAssetDataflowNodes/Public/ChaosClothAsset/SimulationLongRangeAttachmentConfigNode.h b/Engine/Plugins/ChaosClothAssetEditor/Source/ChaosClothAssetDataflowNodes/Public/ChaosClothAsset/SimulationLongRangeAttachmentConfigNode.h
--- a/Engine/Plugins/ChaosClothAssetEditor/Source/ChaosClothAssetDataflowNodes/Public/ChaosClothAsset/SimulationLongRangeAttachmentConfigNode.h
+++ b/Engine/Plugins/ChaosClothAssetEditor/Source/ChaosClothAssetDataflowNodes/Public/ChaosClothAsset/SimulationLongRangeAttachmentConfigNode.h
@@ -63,19 +63,26 @@ public:
 	bool bUseGeodesicTethers = true;
 
 	/**
-	* Enable more granular control over tether generation via custom selection sets.
-	* Otherwise, all dynamic particles will be connect to the closest kinematic vertex as defined by FixedEndSet.
-	*/
+	 * Enable more granular control over tether generation via the custom selection sets CustomTetherData.
+	 * When using custom tether generation, FixedEndSet is still required to contain all kinematic vertices.
+	 * It is also not possible to provide CustomFixedEndSet indices that are not part of the FixedEndSet.
+	 * Similarly, it is not possible to provide CustomDynamicEndSet indices that are part of the FixedEndSet.
+	 * Any fixed tether end provided to a CustomFixedEndSet must be directly connected to a dynamic particles
+	 * group (island) where at least one of the CustomDynamicEndSet particles resides in order to be valid.
+	 * When enabling this option, the default tether generation is completely disabled (dynamic particles are
+	 * no longer connected to their closest kinematic vertex as defined by FixedEndSet).
+	 */
 	UPROPERTY(EditAnywhere, Category = "Long Range Attachment Properties")
 	bool bEnableCustomTetherGeneration = false;
 
-	/** The name of the vertex selection set used as fixed tether ends.
-	* When using custom tether generation, this set is still needed to contain all kinematic vertices.
-	*/
+	/**
+	 * The name of the vertex selection set used as fixed tether ends.
+	 * When using custom tether generation, this set is still needed to contain all kinematic vertices.
+	 */
 	UPROPERTY(EditAnywhere, Category = "Long Range Attachment Properties")
 	FChaosClothAssetConnectableIStringValue FixedEndSet = { TEXT("KinematicVertices3D") };
 
-	/** Pairs of vertex selections used for custom tether generation. */
+	/** Pairs of vertex selections used for custom tether generation. See EnableCustomTetherGeneration. */
 	UPROPERTY(EditAnywhere, EditFixedSize, Category = "Long Range Attachment Properties", Meta = (EditCondition = "bEnableCustomTetherGeneration", SkipInDisplayNameChain))
 	TArray<FChaosClothAssetTetherGenerationSet> CustomTetherData;
 
