Chaos Cloth Asset - Replaced lambda by existing LinearToSRGB function in the static mesh color space conversion code.

#rnx

#Synced-CL 50102962

diff --git a/Engine/Plugins/ChaosClothAssetEditor/Source/ChaosClothAssetDataflowNodes/Private/ChaosClothAsset/ClothDataflowTools.cpp b/Engine/Plugins/ChaosClothAssetEditor/Source/ChaosClothAssetDataflowNodes/Private/ChaosClothAsset/ClothDataflowTools.cpp
--- a/Engine/Plugins/ChaosClothAssetEditor/Source/ChaosClothAssetDataflowNodes/Private/ChaosClothAsset/ClothDataflowTools.cpp
+++ b/Engine/Plugins/ChaosClothAssetEditor/Source/ChaosClothAssetDataflowNodes/Private/ChaosClothAsset/ClothDataflowTools.cpp
@@ -10,6 +10,7 @@
 #include "Interfaces/ITargetPlatformManagerModule.h"
 #include "Modules/ModuleManager.h"
 #include "Rendering/SkeletalMeshLODModel.h"
+#include "Util/ColorConstants.h"
 #include "Widgets/Notifications/SNotificationList.h"
 #include "MeshDescriptionToDynamicMesh.h"
 #include "MeshUtilities.h"
@@ -586,28 +587,15 @@ namespace UE::Chaos::ClothAsset
 			SkinWeights.Set(VertexID, RootBinding);
 		}
 
-		// Convert weird static mesh inverse sRGB gamma to linear.
-		// FIXME: Remove once static mesh color space has been fixed to be linear again.
+		// Convert static mesh linear to sRGB
 		TVertexInstanceAttributesRef<FVector4f> VertexInstanceColors = SkeletalMeshAttributes.GetVertexInstanceColors();
-		auto ConvertLinearToSRGBGamma = [](float V)
-			{
-				V = FMath::Clamp(V, 0.0f, 1.0f);
-				if (V <= 0.0031308)
-				{
-					return V * 12.92f;
-				}
-				else
-				{
-					return 1.055f * FMath::Pow(V, 1.0f / 2.4f) - 0.055f;
-				}
-			};
-
-		for (FVertexInstanceID VertexInstanceID : SkeletalMeshGeometry.VertexInstances().GetElementIDs())
+		for (const FVertexInstanceID VertexInstanceID : SkeletalMeshGeometry.VertexInstances().GetElementIDs())
 		{
-			FLinearColor VertexColor = VertexInstanceColors.Get(VertexInstanceID);
-			VertexColor.R = ConvertLinearToSRGBGamma(VertexColor.R);
-			VertexColor.G = ConvertLinearToSRGBGamma(VertexColor.G);
-			VertexColor.B = ConvertLinearToSRGBGamma(VertexColor.B);
+			FVector4f VertexColor = VertexInstanceColors.Get(VertexInstanceID);
+			VertexColor.X = FMath::Clamp(VertexColor.X, 0.f, 1.f);  // The clamp is required to exactly match the behavior in
+			VertexColor.Y = FMath::Clamp(VertexColor.Y, 0.f, 1.f);  // FStaticToSkeletalMeshConverter::InitializeSkeletonFromStaticMesh()
+			VertexColor.Z = FMath::Clamp(VertexColor.Z, 0.f, 1.f);  // while allowing to reuse LinearToSRGB
+			UE::Geometry::LinearColors::LinearToSRGB(VertexColor);
 			VertexInstanceColors.Set(VertexInstanceID, VertexColor);
 		}
 
