Guard against crash and unexpected results in cloth remesh node

#Synced-CL 50173042

diff --git a/Engine/Plugins/ChaosClothAssetEditor/Source/ChaosClothAssetDataflowNodes/Private/ChaosClothAsset/RemeshNode.cpp b/Engine/Plugins/ChaosClothAssetEditor/Source/ChaosClothAssetDataflowNodes/Private/ChaosClothAsset/RemeshNode.cpp
--- a/Engine/Plugins/ChaosClothAssetEditor/Source/ChaosClothAssetDataflowNodes/Private/ChaosClothAsset/RemeshNode.cpp
+++ b/Engine/Plugins/ChaosClothAssetEditor/Source/ChaosClothAssetDataflowNodes/Private/ChaosClothAsset/RemeshNode.cpp
@@ -1259,7 +1259,8 @@ namespace UE::Chaos::ClothAsset::Private
 		double SmoothingSim,
 		const FString& DensityMapName,
 		const FVector2f& DensityLowHigh,
-		const TSharedRef<FManagedArrayCollection>& OutClothCollection)
+		const TSharedRef<FManagedArrayCollection>& OutClothCollection,
+		TOptional<FText>& OutWarningMsg)
 	{
 		using namespace UE::Geometry;
 		using namespace UE::Chaos::ClothAsset;
@@ -1387,7 +1388,48 @@ namespace UE::Chaos::ClothAsset::Private
 			Converter.Convert(ClothCollection, PatternID, EClothPatternVertexType::Sim2D, ProjectionTarget2D);
 			FDynamicMesh3 ProjectionTarget3D;
 			Converter.Convert(ClothCollection, PatternID, EClothPatternVertexType::Sim3D, ProjectionTarget3D);
-			TSharedPtr<UE::Geometry::FDynamicMeshAABBTree3> ProjectionTargetSpatial = MakeShared<UE::Geometry::FDynamicMeshAABBTree3>(&ProjectionTarget2D, true);
+
+			if (ProjectionTarget2D.TriangleCount() < ProjectionTarget3D.TriangleCount())
+			{
+				OutWarningMsg = LOCTEXT("Sim2DDegenerateTriangle", "2D Sim mesh has degenerate triangles. Remeshing may produce unexpected results");
+			}
+			else if (ProjectionTarget3D.TriangleCount() < ProjectionTarget2D.TriangleCount())
+			{
+				OutWarningMsg = LOCTEXT("Sim3DDegenerateTriangle", "3D Sim mesh has degenerate triangles. Remeshing may produce unexpected results");
+			}
+
+			TArray<int32> ValidTriangles;
+			for (int32 TriangleID = 0; TriangleID < FMath::Min(ProjectionTarget2D.MaxTriangleID(), ProjectionTarget3D.MaxTriangleID()); ++TriangleID)
+			{
+				if (ProjectionTarget2D.IsTriangle(TriangleID) && ProjectionTarget3D.IsTriangle(TriangleID))
+				{
+					ValidTriangles.Add(TriangleID);
+				}
+			}
+			TSharedPtr<UE::Geometry::FDynamicMeshAABBTree3> ProjectionTargetSpatial = MakeShared<UE::Geometry::FDynamicMeshAABBTree3>(&ProjectionTarget2D, /*bAutoBuild=*/false);
+			ProjectionTargetSpatial->Build(ValidTriangles);
+
+			const MeshIntersection::FIntersectionsQueryResult Result = ProjectionTargetSpatial->FindAllSelfIntersections(true, UE::Geometry::IMeshSpatial::FQueryOptions(), [](FIntrTriangle3Triangle3d& IntersectionInfo) -> bool
+			{
+				const FTriangle3d T03D = IntersectionInfo.GetTriangle0();
+				const FTriangle2d T02D{ FVector2d(T03D.V[0]), FVector2d(T03D.V[1]), FVector2d(T03D.V[2]) };
+				const FTriangle3d T13D = IntersectionInfo.GetTriangle1();
+				const FTriangle2d T12D{ FVector2d(T13D.V[0]), FVector2d(T13D.V[1]), FVector2d(T13D.V[2]) };
+
+				FIntrTriangle2Triangle2d Intersection2D(T02D, T12D);
+				if (Intersection2D.Find())
+				{
+					IntersectionInfo.SetReportCoplanarIntersection(true);
+					return IntersectionInfo.Find();
+				}
+				return false;
+			});
+
+			if ((Result.Points.Num() > 0) || (Result.Polygons.Num() > 0) || (Result.Segments.Num() > 0))
+			{
+				OutWarningMsg = LOCTEXT("SelfIntersecting2DPattern", "2D Sim mesh has one or more self-intersecting patterns. Remeshing may produce unexpected results");
+			}
+			
 
 			for (const int32 VID : PatternVertexIDs[PatternID])
 			{
@@ -1396,8 +1438,28 @@ namespace UE::Chaos::ClothAsset::Private
 				double Distance;
 				const int NearestTriangle = ProjectionTargetSpatial->FindNearestTriangle(SrcVert, Distance);
 
+				if (!ProjectionTarget2D.IsTriangle(NearestTriangle))
+				{
+					UE_LOG(LogChaosClothAssetDataflowNodes, Verbose, TEXT("UE::Chaos::ClothAsset::RemeshNode: Failed to find a nearest triangle on the input 2D Sim mesh. Ensure the 2D sim mesh has no empty or degenerate triangles before remeshing"));
+					continue;
+				}
+
 				const FDistPoint3Triangle3d Dist = TMeshQueries<FDynamicMesh3>::TriangleDistance(ProjectionTarget2D, NearestTriangle, SrcVert);
 				const FVector3d Bary = Dist.TriangleBaryCoords;
+
+				if (!ProjectionTarget3D.IsTriangle(NearestTriangle))
+				{
+					UE_LOG(LogChaosClothAssetDataflowNodes, Verbose, TEXT("UE::Chaos::ClothAsset::RemeshNode: A triangle ID that is valid on the 2D Sim mesh was not found on the 3D sim mesh. Check Simulation faces"));
+					continue;
+				}
+
+				const FIndex3i Triangle3D = ProjectionTarget3D.GetTriangle(NearestTriangle);
+				if (!ProjectionTarget3D.IsVertex(Triangle3D[0]) || !ProjectionTarget3D.IsVertex(Triangle3D[1]) || !ProjectionTarget3D.IsVertex(Triangle3D[2]))
+				{
+					UE_LOG(LogChaosClothAssetDataflowNodes, Verbose, TEXT("UE::Chaos::ClothAsset::RemeshNode: A triangle ID is valid on the 3D Sim mesh has vertex indices that are not valid. Check Simulation faces"));
+					continue;
+				}
+
 				const FVector3d InterpolatedPoint = ProjectionTarget3D.GetTriBaryPoint(NearestTriangle, Bary[0], Bary[1], Bary[2]);
 
 				Mesh3D.SetVertex(VID, InterpolatedPoint);
@@ -1894,7 +1956,12 @@ void FChaosClothAssetRemeshNode_v2::Evaluate(UE::Dataflow::FContext& Context, co
 				else
 				{
 					Private::EmptySimSelections(OutputClothCollection);
-					Private::RemeshSimMesh(ClothCollection, IterationsSim, SmoothingSim, SimDensityMapName, FVector2f(DensityMapSim.Low / 100.0f, DensityMapSim.High / 100.0f), OutputClothCollection);
+					TOptional<FText> WarningMsg;
+					Private::RemeshSimMesh(ClothCollection, IterationsSim, SmoothingSim, SimDensityMapName, FVector2f(DensityMapSim.Low / 100.0f, DensityMapSim.High / 100.0f), OutputClothCollection, WarningMsg);
+					if (WarningMsg)
+					{
+						Context.Warning(*WarningMsg, this, Out);
+					}
 					Private::RebuildTopologyDependentSimData(ClothCollection, OutputClothCollection);
 				}
 			}
